#summary Tidied some crash report web server code
#branch main
#proj ue4

diff --git a/Engine/Source/Programs/CrashReporter/CrashReportWebSite/Models/BuggRepository.cs b/Engine/Source/Programs/CrashReporter/CrashReportWebSite/Models/BuggRepository.cs
--- a/Engine/Source/Programs/CrashReporter/CrashReportWebSite/Models/BuggRepository.cs
+++ b/Engine/Source/Programs/CrashReporter/CrashReportWebSite/Models/BuggRepository.cs
@@ -397,11 +397,9 @@ namespace Tools.CrashReporter.CrashReportWebSite.Models
 				Results =
 				(
 					from BuggDetail in Buggs
-					from UserDetail in BuggsDataContext.Users
-					join BuggsUserDetail in BuggsDataContext.Buggs_Users on BuggDetail.Id equals BuggsUserDetail.BuggId
-					join UserGroupDetail in BuggsDataContext.UserGroups on UserDetail.UserGroupId equals UserGroupDetail.Id
-					where BuggsUserDetail.UserNameId == UserDetail.Id || BuggsUserDetail.UserName == UserDetail.UserName
-					group BuggDetail by UserGroupDetail.Name into GroupCount
+					join BuggsUserGroupDetail in BuggsDataContext.Buggs_UserGroups on BuggDetail.Id equals BuggsUserGroupDetail.UserGroupId
+					join UserGroupDetail in BuggsDataContext.UserGroups on BuggsUserGroupDetail.UserGroupId equals UserGroupDetail.Id
+					group 0 by UserGroupDetail.Name into GroupCount
 					select new { Key = GroupCount.Key, Count = GroupCount.Count() }
 				).ToDictionary( x => x.Key, y => y.Count );
 
diff --git a/Engine/Source/Programs/CrashReporter/CrashReportWebSite/Models/CrashReport.cs b/Engine/Source/Programs/CrashReporter/CrashReportWebSite/Models/CrashReport.cs
--- a/Engine/Source/Programs/CrashReporter/CrashReportWebSite/Models/CrashReport.cs
+++ b/Engine/Source/Programs/CrashReporter/CrashReportWebSite/Models/CrashReport.cs
@@ -13,8 +13,6 @@ namespace Tools.CrashReporter.CrashReportWebSite.Models
 	/// </summary>
 	public partial class Bugg
 	{
-		/// <summary></summary>
-		public EntitySet<Crash> BuggCrashes { get; set; }
 		/// <summary></summary>
 		public int CrashesInTimeFrame { get; set; }
 		/// <summary></summary>
@@ -26,12 +24,7 @@ namespace Tools.CrashReporter.CrashReportWebSite.Models
 		/// <returns></returns>
 		public EntitySet<Crash> GetCrashes()
 		{
-			if( BuggCrashes == null )
-			{
-				BuggCrashes = Crashes;
-			}
-
-			int CrashCount = BuggCrashes.Count;
+			int CrashCount = Crashes.Count;
 			if( NumberOfCrashes != CrashCount )
 			{
 				NumberOfCrashes = CrashCount;
@@ -39,12 +32,12 @@ namespace Tools.CrashReporter.CrashReportWebSite.Models
 				if( NumberOfCrashes > 0 )
 				{
 					BuggRepository LocalBuggRepository = new BuggRepository();
-					LocalBuggRepository.UpdateBuggData( this, BuggCrashes );
+					LocalBuggRepository.UpdateBuggData( this, Crashes );
 				}
 			}
 
 			// Just fill the CallStackContainers
-			foreach( Crash CurrentCrash in BuggCrashes )
+			foreach( Crash CurrentCrash in Crashes )
 			{
 				if( CurrentCrash.CallStackContainer == null )
 				{
@@ -57,7 +50,7 @@ namespace Tools.CrashReporter.CrashReportWebSite.Models
 				}
 			}
 
-			return BuggCrashes;
+			return Crashes;
 		}
 
 		/// <summary>
@@ -85,17 +78,13 @@ namespace Tools.CrashReporter.CrashReportWebSite.Models
 		{
 			get
 			{
-				if( CrashesCache == null )
+				if (CrashesCache == null)
 				{
-					CrashesCache = new EntitySet<Crash>( OnCrashesAdd, OnCrashesRemove );
-					CrashesCache.SetSource( Buggs_Crashes.Select( CrashInstance => CrashInstance.Crash ) );
+					CrashesCache = new EntitySet<Crash>(OnCrashesAdd, OnCrashesRemove);
+					CrashesCache.SetSource(Buggs_Crashes.Select(CrashInstance => CrashInstance.Crash));
 				}
 				return CrashesCache;
 			}
-			set
-			{
-				CrashesCache.Assign( value );
-			}
 		}
 
 		private void OnCrashesAdd( Crash CrashInstance )
@@ -223,32 +212,6 @@ namespace Tools.CrashReporter.CrashReportWebSite.Models
 			return LocalCrashRepository.GetCallStack( this );
 		}
 
-		/// <summary>
-		/// http://www.codeproject.com/KB/linq/linq-to-sql-many-to-many.aspx
-		/// </summary>
-		private EntitySet<Bugg> BuggsCache;
-
-		/// <summary>
-		/// 
-		/// </summary>
-		public EntitySet<Bugg> Buggs
-		{
-			get
-			{
-				if( BuggsCache == null )
-				{
-					BuggsCache = new EntitySet<Bugg>( OnBuggsAdd, OnBuggsRemove );
-					BuggsCache.SetSource( Buggs_Crashes.Select( c => c.Bugg ) );
-				}
-
-				return BuggsCache;
-			}
-			set
-			{
-				BuggsCache.Assign( value );
-			}
-		}
-
 		private void OnBuggsAdd( Bugg BuggInstance )
 		{
 			Buggs_Crashes.Add( new Buggs_Crash { Crash = this, Bugg = BuggInstance } );
