#summary Added SQL file containing stored procedure that's run to generate Buggs on the crash server
#branch main
#proj ue4

diff --git a/Engine/Source/Programs/CrashReporter/CrashReportWebSite/Models/CreateBugg-stored-procedure.sql b/Engine/Source/Programs/CrashReporter/CrashReportWebSite/Models/CreateBugg-stored-procedure.sql
new file mode 100644
--- /dev/null
+++ b/Engine/Source/Programs/CrashReporter/CrashReportWebSite/Models/CreateBugg-stored-procedure.sql
@@ -0,0 +1,172 @@
+USE [CrashReport]
+GO
+/****** Object:  StoredProcedure [dbo].[UpdateCrashesByPattern]    Script Date: 07/04/2014 16:01:04 ******/
+SET ANSI_NULLS ON
+GO
+SET QUOTED_IDENTIFIER ON
+GO
+-- Batch submitted through debugger: SQLQuery1.sql|7|0|[REDACTED_USER_HOME]\AppData\Local\Temp\~vs33DE.sql
+
+
+
+ALTER PROCEDURE [dbo].[UpdateCrashesByPattern]
+AS
+BEGIN
+	-- SET NOCOUNT ON added to prevent extra result sets from
+	-- interfering with SELECT statements.
+	SET NOCOUNT ON;
+	
+	-- Update crash user names
+	UPDATE [dbo].[Crashes]
+	SET [UserName] = CASE WHEN c.[UserName] is NULL THEN u.[UserName] ELSE c.[UserName] END
+	FROM [dbo].[Crashes] c
+	LEFT JOIN [dbo].[Users] u on (u.[id] = c.[UserNameId]);
+
+	--Create Buggs
+	MERGE Buggs AS B
+	USING 
+	(
+		SELECT [TTPID]
+			  ,[Pattern]
+			  ,[NumberOfCrashes]
+			  ,[NumberOfUsers]
+			  ,[TimeOfFirstCrash]
+			  ,[TimeOfLastCrash]
+			  ,[Status]
+			  ,[FixedChangeList]
+			  ,[Game]
+		FROM
+		(
+			SELECT COUNT(1)as NumberOfCrashes
+				  , max([TimeOfCrash]) as TimeOfLastCrash
+				  , min([TimeOfCrash]) as TimeOfFirstCrash
+				  , count(distinct [UserName]) as NumberOfUsers
+				  , Count(distinct GameName) as GameNameCount
+				  , MIN(GameName) as Game
+				  , Max([Status]) as [Status]
+				  , Max(TTPID) as TTPID
+				  , Max(FixedChangeList) as FixedChangeList
+				  , Pattern
+			  FROM [dbo].[Crashes] c 
+			  WHERE
+				Pattern is not NULL AND Pattern not like ''
+			  group by Pattern
+		  ) as s
+		 Where s.NumberOfCrashes > 1
+	) AS C 
+	ON (B.Pattern = C.Pattern)
+	WHEN NOT MATCHED BY TARGET
+		THEN INSERT (
+				[TTPID]
+			  ,[Pattern]
+			  ,[NumberOfCrashes]
+			  ,[NumberOfUsers]
+			  ,[TimeOfFirstCrash]
+			  ,[TimeOfLastCrash]
+			  ,[Status]
+			  ,[FixedChangeList]
+			  ,[Game]
+			  ) 
+		Values(
+				C.[TTPID]
+			  , C.[Pattern]
+			  , C.[NumberOfCrashes]
+			  , C.[NumberOfUsers]
+			  , C.[TimeOfFirstCrash]
+			  , C.[TimeOfLastCrash]
+			  , C.[Status]
+			  , C.[FixedChangeList]
+			  , C.[Game]
+			  ) 
+
+	WHEN MATCHED 
+		THEN UPDATE SET B.[NumberOfCrashes] = C.[NumberOfCrashes]
+						, B.[TimeOfLastCrash] = C.[TimeOfLastCrash]
+						, B.[NumberOfUsers] = C.[NumberOfUsers]
+
+	OUTPUT $action, Inserted.*, Deleted.*;
+		
+		
+		
+	/****** Join Buggs and Crashes  ******/
+	MERGE dbo.Buggs_Crashes BC 
+	USING 
+	( 
+		SELECT  b.Id as BuggId, c.[Id] as CrashId
+		FROM [dbo].[Crashes] c
+		Join [dbo].Buggs b on (b.Pattern = c.Pattern)
+		group by b.Id, c.Id
+	 ) AS C
+	 ON BC.BuggId = C.BuggId AND BC.CrashId = C.CrashId
+	 WHEN NOT MATCHED BY TARGET
+		THEN INSERT 
+			(BuggId, CrashId)
+		VALUES 
+			(C.BuggId, C.CrashId)	
+	 WHEN MATCHED
+		THEN UPDATE SET
+			BC.BuggId = C.BuggId, 
+			BC.CrashId = C.CrashId
+	OUTPUT $action, Inserted.*, Deleted.*;
+	
+	
+	/****** Join Buggs_Users and Crashes  ******/
+	MERGE dbo.Buggs_Users AS BU
+	USING
+		(
+		  SELECT  b.Id as BuggId, c.[UserName] as UserName
+		  FROM [dbo].[Crashes] c
+		  Join [dbo].Buggs b on (b.Pattern = c.Pattern)
+		  Group by b.Id, c.UserName
+		) AS C
+	ON BU.BuggId = C.BuggId AND BU.UserName = C.UserName
+	 WHEN NOT MATCHED BY TARGET
+		THEN INSERT 
+			(BuggId, UserName) 
+		VALUES 
+			(C.BuggId, C.UserName)	
+	 WHEN MATCHED
+		THEN UPDATE SET
+			BU.BuggId = C.BuggId, 
+			BU.UserName = C.UserName
+	OUTPUT $action, Inserted.*, Deleted.*;
+	
+		  
+		
+	  
+	/****** Join Buggs and UserGroups  ******/
+	MERGE dbo.Buggs_UserGroups AS BUG
+	USING
+	(
+	  Select BuggId, UserGroupId
+	  From [dbo].[Buggs_Users] bu
+	  Join dbo.Users u on (u.UserName = bu.UserName)
+	  Group by BuggId, UserGroupId
+	 ) AS BU
+	 ON  BUG.BuggId = BU.BuggId AND BUG.UserGroupId = BU.UserGroupId
+	 WHEN NOT MATCHED BY TARGET
+		THEN INSERT 
+			(BuggId, UserGroupId) 
+		VALUES 
+			(BU.BuggId, BU.UserGroupId)	
+	 WHEN MATCHED
+		THEN UPDATE SET
+			BUG.BuggId = BU.BuggId, 
+			BUG.UserGroupId = BU.UserGroupId
+	OUTPUT $action, Inserted.*, Deleted.*;
+	 
+	/*** Update Crashes to match the Bugg Status, TTPID, and Fixed Change List ****/
+	-- This handles new crashes that enter the system; it has the side effect of making the Bugg authoritative in this case but that's how they should be used.
+	update C
+	SET 
+		C.TTPID = B.TTPID,
+		C.FixedChangeList = B.FixedChangeList,
+		C.[Status] = B.[Status]
+	From 
+		dbo.Crashes as C
+		Join Buggs_Crashes as BC ON (c.Id = BC.CrashId)
+		Join Buggs as B ON (B.Id =BC.BuggId AND C.Id = BC.CrashId) 
+	 
+END
+
+
