#summary Removed old 'Buggs' name from web-site tooltip
#branch main
#proj ue4

diff --git a/Engine/Source/Programs/CrashReporter/CrashReportWebSite/Models/BuggsViewModel.cs b/Engine/Source/Programs/CrashReporter/CrashReportWebSite/Models/BuggsViewModel.cs
--- a/Engine/Source/Programs/CrashReporter/CrashReportWebSite/Models/BuggsViewModel.cs
+++ b/Engine/Source/Programs/CrashReporter/CrashReportWebSite/Models/BuggsViewModel.cs
@@ -12,6 +12,9 @@ namespace Tools.CrashReporter.CrashReportWebSite.Models
 	/// </summary>
 	public class BuggsViewModel
 	{
+		/// <summary>Tooltip for "Crash Groups" tab</summary>
+		public const string Tooltip = "Collections of crashes that have the exact same callstack.";
+
 		/// <summary>A container of sorted Buggs.</summary>
 		public IEnumerable<Bugg> Results { get; set; }
 		/// <summary>Information to paginate the list of Buggs.</summary>
diff --git a/Engine/Source/Programs/CrashReporter/CrashReportWebSite/Views/Buggs/Index.aspx b/Engine/Source/Programs/CrashReporter/CrashReportWebSite/Views/Buggs/Index.aspx
--- a/Engine/Source/Programs/CrashReporter/CrashReportWebSite/Views/Buggs/Index.aspx
+++ b/Engine/Source/Programs/CrashReporter/CrashReportWebSite/Views/Buggs/Index.aspx
@@ -81,7 +81,8 @@
 			<form action="/" method="POST">
 				<div style="background-color: #E8EEF4; margin-bottom: -5px; width: 19.7em;">
 					<span style="background-color: #C3CAD0; font-size: medium; padding: 0 1em;"><%=Html.ActionLink( "Crashes", "Index", "Crashes", new { SearchQuery = Model.Query, SortTerm = Model.Term, SortOrder = Model.Order, UserGroup = Model.UserGroup, DateFrom = Model.DateFrom, DateTo = Model.DateTo }, new { style = "color:black; text-decoration:none;" } )%></span>
-					<span style="background-color: #E8EEF4; font-size: medium; padding:0 1em;" title='A Bugg is a collection of Crashes that have the exact same callstack.'><%=Html.ActionLink( "CrashGroups", "Index", "Buggs", new { SearchQuery = Model.Query, SortTerm = Model.Term, SortOrder = Model.Order, UserGroup = Model.UserGroup, DateFrom = Model.DateFrom, DateTo = Model.DateTo }, new { style = "color:black; text-decoration:none;" } )%></span>
+					<span style="background-color: #E8EEF4; font-size: medium; padding:0 1em;"
+					      title="<%= BuggsViewModel.Tooltip %>"><%=Html.ActionLink( "CrashGroups", "Index", "Buggs", new { SearchQuery = Model.Query, SortTerm = Model.Term, SortOrder = Model.Order, UserGroup = Model.UserGroup, DateFrom = Model.DateFrom, DateTo = Model.DateTo }, new { style = "color:black; text-decoration:none;" } )%></span>
 				</div>
 				<% Html.RenderPartial("/Views/Buggs/ViewBuggs.ascx"); %>
 			</form>
diff --git a/Engine/Source/Programs/CrashReporter/CrashReportWebSite/Views/Crashes/Index.aspx b/Engine/Source/Programs/CrashReporter/CrashReportWebSite/Views/Crashes/Index.aspx
--- a/Engine/Source/Programs/CrashReporter/CrashReportWebSite/Views/Crashes/Index.aspx
+++ b/Engine/Source/Programs/CrashReporter/CrashReportWebSite/Views/Crashes/Index.aspx
@@ -150,7 +150,8 @@ Crash Reports
 						CrashType = Model.CrashType
 					}, new { style = "color:black; text-decoration:none;" }
 				)%></span>
-				<span style="background-color: #C3CAD0; font-size: medium; margin-left: 1em; padding:0 1em;" title='A Bugg is a collection of Crashes that have the exact same callstack.'><%=Html.ActionLink( "CrashGroups", "Index", "Buggs", new { SearchQuery = Model.Query, SortTerm = Model.Term, SortOrder = Model.Order, UserGroup = Model.UserGroup, DateFrom = Model.DateFrom, DateTo = Model.DateTo }, new { style = "color:black; text-decoration:none;" } )%></span>
+				<span style="background-color: #C3CAD0; font-size: medium; margin-left: 1em; padding:0 1em;"
+				      title="<%= BuggsViewModel.Tooltip %>"><%=Html.ActionLink( "CrashGroups", "Index", "Buggs", new { SearchQuery = Model.Query, SortTerm = Model.Term, SortOrder = Model.Order, UserGroup = Model.UserGroup, DateFrom = Model.DateFrom, DateTo = Model.DateTo }, new { style = "color:black; text-decoration:none;" } )%></span>
 			</div>
 			<%Html.RenderPartial( "/Views/Crashes/ViewCrash.ascx" );%>
 		</form>
