#CrashReportClient: Doesn't display assertion text  (TTP#33238,fixed)

diff --git a/Engine/Source/Programs/CrashReporter/CrashReportClient/Private/CrashReportUtil.h b/Engine/Source/Programs/CrashReporter/CrashReportClient/Private/CrashReportUtil.h
--- a/Engine/Source/Programs/CrashReporter/CrashReportClient/Private/CrashReportUtil.h
+++ b/Engine/Source/Programs/CrashReporter/CrashReportClient/Private/CrashReportUtil.h
@@ -108,13 +108,28 @@ FunctorDirectoryVisitor<Functor> MakeDirectoryVisitor(Functor&& FunctorInstance)
 /**
  * Create a multiline string to display from an exception and callstack
  * @param Exception Exception description string
+ * @param Assertion Assertion description string
  * @param Callstack List of callstack entry strings
  * @return Multiline text
  */
-inline FText FormatReportDescription(const FString& Exception, const TArray<FString>& Callstack)
+inline FText FormatReportDescription(const FString& Exception, const FString& Assertion, const TArray<FString>& Callstack)
 {
 	FString Diagnostic = Exception + "\n\n";
 
+	if( !Assertion.IsEmpty() )
+	{
+		TArray<FString> MultilineAssertion;
+		Assertion.ParseIntoArray( &MultilineAssertion, TEXT( "#" ), true );
+
+		for( const auto& AssertionLine : MultilineAssertion )
+		{
+			Diagnostic += AssertionLine;
+			Diagnostic += "\n";
+		}
+
+		Diagnostic += "\n";
+	}
+
 	for (const auto& Line: Callstack)
 	{
 		Diagnostic += Line + "\n";
diff --git a/Engine/Source/Programs/CrashReporter/CrashReportClient/Private/GenericErrorReport.cpp b/Engine/Source/Programs/CrashReporter/CrashReportClient/Private/GenericErrorReport.cpp
--- a/Engine/Source/Programs/CrashReporter/CrashReportClient/Private/GenericErrorReport.cpp
+++ b/Engine/Source/Programs/CrashReporter/CrashReportClient/Private/GenericErrorReport.cpp
@@ -165,7 +165,7 @@ bool FGenericErrorReport::TryReadDiagnosticsFile(FText& OutReportDescription)
 			break;
 		}
 	}
-	OutReportDescription = FormatReportDescription(Exception, Callstack);
+	OutReportDescription = FormatReportDescription(Exception, TEXT(""), Callstack);
 	return true;
 }
 
diff --git a/Engine/Source/Programs/CrashReporter/CrashReportClient/Private/Windows/WindowsErrorReport.cpp b/Engine/Source/Programs/CrashReporter/CrashReportClient/Private/Windows/WindowsErrorReport.cpp
--- a/Engine/Source/Programs/CrashReporter/CrashReportClient/Private/Windows/WindowsErrorReport.cpp
+++ b/Engine/Source/Programs/CrashReporter/CrashReportClient/Private/Windows/WindowsErrorReport.cpp
@@ -19,6 +19,45 @@ namespace
 	FCrashDebugHelperModule* CrashHelperModule;
 }
 
+/** Helper class used to parse specified string value based on the marker. */
+struct FWindowsReportParser
+{
+	static FString Find( const FString& ReportDirectory, const TCHAR* Marker )
+	{
+		FString Result;
+
+		TArray<uint8> FileData;
+		FFileHelper::LoadFileToArray( FileData, *(ReportDirectory / TEXT( "Report.wer" )) );
+		FileData.Add( 0 );
+		FileData.Add( 0 );
+
+		const FString FileAsString = reinterpret_cast<TCHAR*>(FileData.GetData());
+
+		TArray<FString> String;
+		FileAsString.ParseIntoArray( &String, TEXT( "\r\n" ), true );
+
+		for( const auto& StringLine : String )
+		{
+			if( StringLine.Contains( Marker ) )
+			{
+				TArray<FString> SeparatedParameters;
+				StringLine.ParseIntoArray( &SeparatedParameters, Marker, true );
+
+				FString MatchedValue;
+				const bool bFound = FParse::Value( *StringLine, Marker, MatchedValue );
+
+				if( bFound )
+				{
+					Result = MatchedValue;
+					break;
+				}
+			}
+		}
+
+		return Result;
+	}
+};
+
 FWindowsErrorReport::FWindowsErrorReport(const FString& Directory)
 	: FGenericErrorReport(Directory)
 {
@@ -58,45 +97,15 @@ FText FWindowsErrorReport::DiagnoseReport() const
 	CrashDebugHelper->CrashInfo.GenerateReport(ReportDirectory / GDiagnosticsFilename);
 
 	const auto& Exception = CrashDebugHelper->CrashInfo.Exception;
-	return FormatReportDescription(Exception.ExceptionString, Exception.CallStackString);
+	const FString Assertion = FWindowsReportParser::Find( ReportDirectory, TEXT( "AssertLog=" ) );
+
+	return FormatReportDescription( Exception.ExceptionString, Assertion, Exception.CallStackString );
 }
 
 FString FWindowsErrorReport::FindCrashedAppName() const
 {
-	TArray<uint8> FileData;
-	if(!FFileHelper::LoadFileToArray(FileData, *(ReportDirectory / TEXT("Report.wer"))))
-	{
-		return "";
-	}
-
-	// Look backwards for AppPath= line
-	auto Data = reinterpret_cast<TCHAR*>(FileData.GetData());
-	TCHAR* Line = nullptr;
-	for (int Index = FileData.Num() / sizeof(TCHAR); Index != 0; --Index)
-	{
-		auto& Char = Data[Index - 1];
-		if (Char != '\n' && Char != '\r')
-		{
-			Line = &Char;
-			continue;
-		}
-		
-		// Zero terminate line for conversion to string below
-		Char = 0;
-		if (!Line)
-		{
-			continue;
-		}
-
-		static const TCHAR AppPathLineStart[] = TEXT("AppPath=");
-		static const int AppPathIdLength = ARRAY_COUNT(AppPathLineStart) - 1;
-		if (0 == FCString::Strncmp(Line, AppPathLineStart, AppPathIdLength))
-		{
-			return FPaths::GetCleanFilename(Line + AppPathIdLength);
-		}
-		Line = nullptr;
-	}
-	return "";
+	const FString CrashedAppName = FWindowsReportParser::Find( ReportDirectory, TEXT( "AppPath=" ) );
+	return CrashedAppName;
 }
 
 FString FWindowsErrorReport::FindMostRecentErrorReport()
