fixed missing module name

diff --git a/Engine/Source/Programs/CrashReporter/CrashReportClient/Private/CrashReportClientApp.cpp b/Engine/Source/Programs/CrashReporter/CrashReportClient/Private/CrashReportClientApp.cpp
--- a/Engine/Source/Programs/CrashReporter/CrashReportClient/Private/CrashReportClientApp.cpp
+++ b/Engine/Source/Programs/CrashReporter/CrashReportClient/Private/CrashReportClientApp.cpp
@@ -123,7 +123,7 @@ void RunCrashReportClient(const TCHAR* CommandLine)
 				SNew(SWindow)
 				.ClientSize(FVector2D(800, 600))
 				[
-					FModuleManager::LoadModuleChecked<ISlateReflectorModule>().GetWidgetReflector()
+					FModuleManager::LoadModuleChecked<ISlateReflectorModule>("SlateReflector").GetWidgetReflector()
 				]);
 		}
 
