UE4: Crash reports are now always evaluated client side, even in Rocket. If absolutely no symbols get resolved the "Please send the crash report" message appears, otherwise the callstack appears with the user's ID at the top line.

diff --git a/Engine/Source/Programs/CrashReporter/CrashReportClient/Private/CrashReportClient.cpp b/Engine/Source/Programs/CrashReporter/CrashReportClient/Private/CrashReportClient.cpp
--- a/Engine/Source/Programs/CrashReporter/CrashReportClient/Private/CrashReportClient.cpp
+++ b/Engine/Source/Programs/CrashReporter/CrashReportClient/Private/CrashReportClient.cpp
@@ -31,18 +31,14 @@ FCrashReportClient::FCrashReportClient(const FString& Directory)
 		// The Epic ID can be looked up from this ID
 		auto DeviceId = FPlatformMisc::GetUniqueDeviceId();
 		GCrashUserId = FString(TEXT("!Id:")) + DeviceId;
-
-		// Rocket builds will not try to analyze the report locally. Reports will be sent to the crash server and analyzed there.
-		DiagnosticText = FText::Format(
-			LOCTEXT("NoDiagnoseReportRocket", "We apologize for the inconvenience.\nPlease send this crash report to help improve our software.\n\n{0}"),
-			FText::FromString(DeviceId)
-		);
-		return;
+	}
+	else
+	{
+		// Remove periods from internal user names to match AutoReporter user names
+		// The name prefix is read by CrashRepository.AddNewCrash in the website code
+		GCrashUserId = FString(TEXT("!Name:")) + FString(FPlatformProcess::UserName()).Replace(TEXT("."), TEXT(""));
 	}
 
-	// Remove periods from internal user names to match AutoReporter user names
-	// The name prefix is read by CrashRepository.AddNewCrash in the website code
-	GCrashUserId = FString(TEXT("!Name:")) + FString(FPlatformProcess::UserName()).Replace(TEXT("."), TEXT(""));
 	if (!ErrorReportFiles.TryReadDiagnosticsFile(DiagnosticText) && !FParse::Param(FCommandLine::Get(), TEXT("no-local-diagnosis")))
 	{
 		auto& Worker = DiagnoseReportTask.GetTask();
@@ -202,7 +198,7 @@ bool FCrashReportClient::UIWillCloseTick(float UnusedDeltaTime)
 
 void FDiagnoseReportWorker::DoWork()
 {
-	*DiagnosticText = ErrorReportFiles->DiagnoseReport();
+	*DiagnosticText = FText::Format( LOCTEXT("CrashReportClientCallstackPattern", "{0}\n\n{1}"), FText::FromString(GCrashUserId), ErrorReportFiles->DiagnoseReport() );
 }
 
 const TCHAR* FDiagnoseReportWorker::Name()
diff --git a/Engine/Source/Programs/CrashReporter/CrashReportClient/Private/Windows/WindowsErrorReport.cpp b/Engine/Source/Programs/CrashReporter/CrashReportClient/Private/Windows/WindowsErrorReport.cpp
--- a/Engine/Source/Programs/CrashReporter/CrashReportClient/Private/Windows/WindowsErrorReport.cpp
+++ b/Engine/Source/Programs/CrashReporter/CrashReportClient/Private/Windows/WindowsErrorReport.cpp
@@ -76,7 +76,7 @@ FText FWindowsErrorReport::DiagnoseReport() const
 	if (!CrashDebugHelper)
 	{
 		// Not localized: should never be seen
-		return FText::FromString(TEXT("Failed to load CrashDebugHelper"));
+		return FText::FromString(TEXT("Failed to load CrashDebugHelper."));
 	}
 
 	FString DumpFilename;
@@ -84,13 +84,20 @@ FText FWindowsErrorReport::DiagnoseReport() const
 	{
 		if (!FindFirstReportFileWithExtension(DumpFilename, TEXT(".mdmp")))
 		{
-			return LOCTEXT("MinidumpNotFound", "No minidump found for this crash");
+			return LOCTEXT("MinidumpNotFound", "No minidump found for this crash.");
 		}
 	}
 
 	if (!CrashDebugHelper->CreateMinidumpDiagnosticReport(ReportDirectory / DumpFilename))
 	{
-		return LOCTEXT("NoDebuggingSymbols", "You do not have any debugging symbols required to display the callstack for this crash");
+		if ( FRocketSupport::IsRocket() )
+		{
+			return LOCTEXT("NoDebuggingSymbolsRocket", "We apologize for the inconvenience.\nPlease send this crash report to help improve our software.");
+		}
+		else
+		{
+			return LOCTEXT("NoDebuggingSymbols", "You do not have any debugging symbols required to display the callstack for this crash.");
+		}
 	}
 
 	// There's a callstack, so write it out to save the server trying to do it
