UE4: Rocket builds no longer write a local Diagnostics.txt when crashing. The callstack is evaluated and displayed to users, but the callstack will be re-resolved on the server.

diff --git a/Engine/Source/Programs/CrashReporter/CrashReportClient/Private/Windows/WindowsErrorReport.cpp b/Engine/Source/Programs/CrashReporter/CrashReportClient/Private/Windows/WindowsErrorReport.cpp
--- a/Engine/Source/Programs/CrashReporter/CrashReportClient/Private/Windows/WindowsErrorReport.cpp
+++ b/Engine/Source/Programs/CrashReporter/CrashReportClient/Private/Windows/WindowsErrorReport.cpp
@@ -100,8 +100,12 @@ FText FWindowsErrorReport::DiagnoseReport() const
 		}
 	}
 
-	// There's a callstack, so write it out to save the server trying to do it
-	CrashDebugHelper->CrashInfo.GenerateReport(ReportDirectory / GDiagnosticsFilename);
+	// Don't write a Diagnostics.txt to disk in rocket. It will be displayed in the UI but not sent to the server.
+	if ( !FRocketSupport::IsRocket() )
+	{
+		// There's a callstack, so write it out to save the server trying to do it
+		CrashDebugHelper->CrashInfo.GenerateReport(ReportDirectory / GDiagnosticsFilename);
+	}
 
 	const auto& Exception = CrashDebugHelper->CrashInfo.Exception;
 	const FString Assertion = FWindowsReportParser::Find( ReportDirectory, TEXT( "AssertLog=" ) );
