Fixed crash in CrashReportClient introduced by thread API changes

diff --git a/Engine/Source/Programs/CrashReporter/CrashReportClient/Private/CrashReportUtil.h b/Engine/Source/Programs/CrashReporter/CrashReportClient/Private/CrashReportUtil.h
--- a/Engine/Source/Programs/CrashReporter/CrashReportClient/Private/CrashReportUtil.h
+++ b/Engine/Source/Programs/CrashReporter/CrashReportClient/Private/CrashReportUtil.h
@@ -55,17 +55,8 @@ private:
 		return 0;
 	}
 
-	virtual void Exit() OVERRIDE
-	{
-		if (Thread)
-		{
-			Thread->Kill();
-			delete Thread;
-		}
-	}
-
 	TWorker Worker;
-	FRunnableThread * Thread;
+	TScopedPointer<FRunnableThread> Thread;
 	volatile int32 bIsDone;	
 };
 
