Added assertion messages to the crash report website
Relates to TTP 333619

diff --git a/Engine/Source/Programs/CrashReporter/CrashReportWebSite/Models/CrashRepository.cs b/Engine/Source/Programs/CrashReporter/CrashReportWebSite/Models/CrashRepository.cs
--- a/Engine/Source/Programs/CrashReporter/CrashReportWebSite/Models/CrashRepository.cs
+++ b/Engine/Source/Programs/CrashReporter/CrashReportWebSite/Models/CrashRepository.cs
@@ -880,7 +880,13 @@ namespace Tools.CrashReporter.CrashReportWebSite.Models
 			NewCrash.GameName = NewCrashInfo.GameName;
 			NewCrash.LanguageExt = NewCrashInfo.Language;		// Convert
 			NewCrash.PlatformName = NewCrashInfo.Platform;
-			if( NewCrashInfo.CallStack != null )
+
+			if (NewCrashInfo.AssertionMessage != null)
+			{
+				NewCrash.Summary = string.Join( "\n", NewCrashInfo.AssertionMessage );
+			}
+			
+			if (NewCrashInfo.CallStack != null)
 			{
 				NewCrash.RawCallStack = string.Join( "\n", NewCrashInfo.CallStack );
 			}
diff --git a/Engine/Source/Programs/CrashReporter/CrashReportWebSite/Views/Crashes/ViewCrash.ascx b/Engine/Source/Programs/CrashReporter/CrashReportWebSite/Views/Crashes/ViewCrash.ascx
--- a/Engine/Source/Programs/CrashReporter/CrashReportWebSite/Views/Crashes/ViewCrash.ascx
+++ b/Engine/Source/Programs/CrashReporter/CrashReportWebSite/Views/Crashes/ViewCrash.ascx
@@ -62,6 +62,7 @@
 		<th style='width: 9em;'><%=Url.TableHeader( "TTP", "TTPID", Model) %></th>
 		<th style='width: 9em;'><%=Url.TableHeader( "Branch", "Branch", Model) %></th>
 		<th>Description</th>
+		<th style='width: 15em'><%=Url.TableHeader( "Message", "Summary", Model )%></th>
 		<th style='width: 9em;'><%=Url.TableHeader( "CL#", "ChangeListVersion", Model )%></th>
 		<th style='width: 12em;'><%=Url.TableHeader( "Computer", "ComputerName", Model )%></th>
 		<th style='width: 14em;'><%=Url.TableHeader( "Platform", "PlatformName", Model )%></th>
@@ -153,6 +154,7 @@
 				<td class="Ttp"><%=CurrentCrash.TTPID%>&nbsp;</td>
 				<td class="Branch"><%=CurrentCrash.Branch%>&nbsp;</td>
 				<td class="Description"><span class="TableData"><%=CurrentCrash.Description%>&nbsp;</span></td>
+				<td class="Summary"><%=CurrentCrash.Summary%></td>
 				<td class="ChangeListVersion"><%=CurrentCrash.ChangeListVersion%></td>
 				<td class="Computer"><%=CurrentCrash.ComputerName%></td>
 				<td class="Platform"><%=CurrentCrash.PlatformName%></td>
