Fixed crash reporter website project (location of common C# library)

diff --git a/Engine/Source/Programs/CrashReporter/CrashReportWebSite/CrashReportWebSite.csproj b/Engine/Source/Programs/CrashReporter/CrashReportWebSite/CrashReportWebSite.csproj
--- a/Engine/Source/Programs/CrashReporter/CrashReportWebSite/CrashReportWebSite.csproj
+++ b/Engine/Source/Programs/CrashReporter/CrashReportWebSite/CrashReportWebSite.csproj
@@ -246,7 +246,7 @@
       <Project>{5d7d66e8-8c76-4af9-b3ec-2ef03421d730}</Project>
       <Name>DotNETUtilities</Name>
     </ProjectReference>
-    <ProjectReference Include="..\..\NoRedist\CrashReportServer\CrashReportCommon\CrashReportCommon.csproj">
+    <ProjectReference Include="..\..\NotForLicensees\CrashReportServer\CrashReportCommon\CrashReportCommon.csproj">
       <Project>{b83de8c9-9429-4765-86bd-8eeea89b138f}</Project>
       <Name>CrashReportCommon</Name>
     </ProjectReference>
