CrashReporter - Moved DO_LOCAL_TESTING to a better place

diff --git a/Engine/Source/Developer/CrashDebugHelper/Private/Windows/WindowsPlatformStackWalkExt.cpp b/Engine/Source/Developer/CrashDebugHelper/Private/Windows/WindowsPlatformStackWalkExt.cpp
--- a/Engine/Source/Developer/CrashDebugHelper/Private/Windows/WindowsPlatformStackWalkExt.cpp
+++ b/Engine/Source/Developer/CrashDebugHelper/Private/Windows/WindowsPlatformStackWalkExt.cpp
@@ -210,7 +210,6 @@ void FWindowsPlatformStackWalkExt::SetSymbolPathsFromModules( FCrashInfo* CrashI
 		const bool bUseNewSetPath = true;
 		if( bUseNewSetPath )
 		{
-//#define DO_LOCAL_TESTING 1
 #if	DO_LOCAL_TESTING
 			const FString RootDir = TEXT( "U:/P4EPIC/UE4-Releases/4.1/" );
 #else
diff --git a/Engine/Source/Developer/CrashDebugHelper/Public/CrashDebugHelper.h b/Engine/Source/Developer/CrashDebugHelper/Public/CrashDebugHelper.h
--- a/Engine/Source/Developer/CrashDebugHelper/Public/CrashDebugHelper.h
+++ b/Engine/Source/Developer/CrashDebugHelper/Public/CrashDebugHelper.h
@@ -2,6 +2,8 @@
 
 #pragma once
 
+//#define DO_LOCAL_TESTING 1
+
 struct FPDBCacheEntry;
 typedef TSharedRef<FPDBCacheEntry> FPDBCacheEntryRef;
 typedef TSharedPtr<FPDBCacheEntry> FPDBCacheEntryPtr;
diff --git a/Engine/Source/Programs/CrashReporter/CrashReportClient/Private/CrashReportClient.cpp b/Engine/Source/Programs/CrashReporter/CrashReportClient/Private/CrashReportClient.cpp
--- a/Engine/Source/Programs/CrashReporter/CrashReportClient/Private/CrashReportClient.cpp
+++ b/Engine/Source/Programs/CrashReporter/CrashReportClient/Private/CrashReportClient.cpp
@@ -7,7 +7,12 @@
 
 #define LOCTEXT_NAMESPACE "CrashReportClient"
 
-const TCHAR* GServerIP = TEXT("http://crashreporter.epicgames.com:57005");
+#if	DO_LOCAL_TESTING
+	const TCHAR* GServerIP = TEXT( "http://localhost:57005" );
+#else
+	const TCHAR* GServerIP = TEXT( "http://crashreporter.epicgames.com:57005" );
+#endif // DO_LOCAL_TESTING
+
 // Must match filename specified in RunMinidumpDiagnostics
 const TCHAR* GDiagnosticsFilename = TEXT("Diagnostics.txt");
 FString GCrashUserId;
diff --git a/Engine/Source/Programs/CrashReporter/CrashReportClient/Public/CrashReportClientApp.h b/Engine/Source/Programs/CrashReporter/CrashReportClient/Public/CrashReportClientApp.h
--- a/Engine/Source/Programs/CrashReporter/CrashReportClient/Public/CrashReportClientApp.h
+++ b/Engine/Source/Programs/CrashReporter/CrashReportClient/Public/CrashReportClientApp.h
@@ -2,6 +2,8 @@
 
 #pragma once
 
+//#define DO_LOCAL_TESTING 1
+
 #define CRASH_REPORT_UNATTENDED_ONLY			PLATFORM_LINUX
 
 // Pre-compiled header includes
