Crash web server: fixed date range changing seemingly at random.

Timestamps coming in from Javascript were being treated as local times, but they're actually UTC (although all times and dates in Javascripts are local time, millisecond timestamps are UTC).

diff --git a/Engine/Source/Programs/CrashReporter/CrashReportWebSite/Models/FormHelper.cs b/Engine/Source/Programs/CrashReporter/CrashReportWebSite/Models/FormHelper.cs
--- a/Engine/Source/Programs/CrashReporter/CrashReportWebSite/Models/FormHelper.cs
+++ b/Engine/Source/Programs/CrashReporter/CrashReportWebSite/Models/FormHelper.cs
@@ -139,13 +139,11 @@ namespace Tools.CrashReporter.CrashReportWebSite.Models
 
 			GetFormParameter( Request, Form, "GameName", GameName, out GameName );
 
-			DateFrom = DateTime.Today.AddDays(-7);
+			DateFrom = DateTime.Today.AddDays(-7).ToUniversalTime();
 			TryParseDate(Request, Form, "DateFrom", ref DateFrom);
-			DateFrom = DateFrom.ToUniversalTime();
 
-			DateTo = DateTime.Today;
+			DateTo = DateTime.Today.ToUniversalTime();
 			TryParseDate(Request, Form, "DateTo", ref DateTo);
-			DateTo = DateTo.ToUniversalTime();
 
 			// Set the sort order 
 			if( PreviousOrder == "Descending" && PreviousTerm == SortTerm )
