CrashReporter - CrashReportClient should be built and distributed in Shipping, not Development (both Launcher and Editor versions)

diff --git a/Engine/Source/Programs/CrashReporter/CrashReportClient/CrashReportClient.Target.cs b/Engine/Source/Programs/CrashReporter/CrashReportClient/CrashReportClient.Target.cs
--- a/Engine/Source/Programs/CrashReporter/CrashReportClient/CrashReportClient.Target.cs
+++ b/Engine/Source/Programs/CrashReporter/CrashReportClient/CrashReportClient.Target.cs
@@ -15,8 +15,8 @@ public class CrashReportClientTarget : TargetRules
 	//
 	public override bool GetSupportedPlatforms(ref List<UnrealTargetPlatform> OutPlatforms)
 	{
-        OutPlatforms.Add(UnrealTargetPlatform.Win32);
-        OutPlatforms.Add(UnrealTargetPlatform.Win64);
+		OutPlatforms.Add(UnrealTargetPlatform.Win32);
+		OutPlatforms.Add(UnrealTargetPlatform.Win64);
 		OutPlatforms.Add(UnrealTargetPlatform.Mac);
 		OutPlatforms.Add(UnrealTargetPlatform.Linux);
 		return true;
@@ -39,10 +39,10 @@ public class CrashReportClientTarget : TargetRules
 		}
 	}
 
-    public override bool ShouldCompileMonolithic(UnrealTargetPlatform InPlatform, UnrealTargetConfiguration InConfiguration)
-    {
-        return true;
-    }
+	public override bool ShouldCompileMonolithic(UnrealTargetPlatform InPlatform, UnrealTargetConfiguration InConfiguration)
+	{
+		return true;
+	}
 
 	public override void SetupGlobalEnvironment(
 		TargetInfo Target,
@@ -58,6 +58,7 @@ public class CrashReportClientTarget : TargetRules
 		// CrashReportClient doesn't ever compile with the engine linked in
 		UEBuildConfiguration.bCompileAgainstEngine = false;
 		UEBuildConfiguration.bCompileAgainstCoreUObject = false;
+		UEBuildConfiguration.bUseLoggingInShipping = true;
 
 		UEBuildConfiguration.bIncludeADO = (Target.Platform == UnrealTargetPlatform.Win64 || Target.Platform == UnrealTargetPlatform.Win32);
 
@@ -67,19 +68,26 @@ public class CrashReportClientTarget : TargetRules
 
 		// Do NOT produce additional console app exe
 		OutLinkEnvironmentConfiguration.bBuildAdditionalConsoleApplication = false;
+
+		OutCPPEnvironmentConfiguration.Definitions.Add( "USE_CHECKS_IN_SHIPPING=1" );
+	}
+	public override bool GUBP_AlwaysBuildWithTools(UnrealTargetPlatform InHostPlatform, out bool bInternalToolOnly, out bool SeparateNode)
+	{
+		bInternalToolOnly = false;
+		SeparateNode = false;
+		return true;
+	}
+	public override List<UnrealTargetPlatform> GUBP_ToolPlatforms(UnrealTargetPlatform InHostPlatform)
+	{
+		if (InHostPlatform == UnrealTargetPlatform.Win64)
+		{
+			return new List<UnrealTargetPlatform> { UnrealTargetPlatform.Win64, UnrealTargetPlatform.Win32 };
+		}
+		return base.GUBP_ToolPlatforms(InHostPlatform);
+	}
+
+	public override List<UnrealTargetConfiguration> GUBP_ToolConfigs( UnrealTargetPlatform InHostPlatform )
+	{
+		return new List<UnrealTargetConfiguration> { UnrealTargetConfiguration.Development, UnrealTargetConfiguration.Shipping };
 	}
-    public override bool GUBP_AlwaysBuildWithTools(UnrealTargetPlatform InHostPlatform, out bool bInternalToolOnly, out bool SeparateNode)
-    {
-        bInternalToolOnly = false;
-        SeparateNode = false;
-        return true;
-    }
-    public override List<UnrealTargetPlatform> GUBP_ToolPlatforms(UnrealTargetPlatform InHostPlatform)
-    {
-        if (InHostPlatform == UnrealTargetPlatform.Win64)
-        {
-            return new List<UnrealTargetPlatform> { UnrealTargetPlatform.Win64, UnrealTargetPlatform.Win32 };
-        }
-        return base.GUBP_ToolPlatforms(InHostPlatform);
-    }   
 }
