CrashReporter - XML WER is always send as UTF16 regardless the internal format which may be (UTF8 or UTF16)

diff --git a/Engine/Source/Programs/CrashReporter/CrashReportClient/Private/CrashUpload.cpp b/Engine/Source/Programs/CrashReporter/CrashReportClient/Private/CrashUpload.cpp
--- a/Engine/Source/Programs/CrashReporter/CrashReportClient/Private/CrashUpload.cpp
+++ b/Engine/Source/Programs/CrashReporter/CrashReportClient/Private/CrashUpload.cpp
@@ -114,6 +114,8 @@ void FCrashUpload::Cancel()
 
 bool FCrashUpload::SendCheckReportRequest()
 {
+	FString XMLString;
+
 	UE_LOG(CrashReportClientLog, Log, TEXT("Sending HTTP request (checking report)"));
 	auto Request = CreateHttpRequest();
 	if (State == EUploadState::CheckingReport)
@@ -125,11 +127,18 @@ bool FCrashUpload::SendCheckReportRequest()
 	else
 	{
 		// This part is Windows-specific on the server
-		ErrorReport.LoadWindowsReportXmlFile(PostData);
+		ErrorReport.LoadWindowsReportXmlFile( XMLString );
+		// Copy the string data into a TArray<uint8>.
+		const int32 XMLStringSize = XMLString.Len() * sizeof(TCHAR);
+		PostData.Reset( XMLStringSize );
+		PostData.AddUninitialized( XMLStringSize );
+		FMemory::Memcpy( PostData.GetData(), *XMLString, XMLStringSize );
+
 		Request->SetURL(UrlPrefix / TEXT("CheckReportDetail"));
 		Request->SetHeader(TEXT("Content-Type"), TEXT("text/plain; charset=utf-16"));
 	}
 
+	UE_LOG( CrashReportClientLog, Log, TEXT( "PostData Num: %i" ), PostData.Num() );
 	Request->SetVerb(TEXT("POST"));
 	Request->SetContent(PostData);
 
diff --git a/Engine/Source/Programs/CrashReporter/CrashReportClient/Private/GenericErrorReport.cpp b/Engine/Source/Programs/CrashReporter/CrashReportClient/Private/GenericErrorReport.cpp
--- a/Engine/Source/Programs/CrashReporter/CrashReportClient/Private/GenericErrorReport.cpp
+++ b/Engine/Source/Programs/CrashReporter/CrashReportClient/Private/GenericErrorReport.cpp
@@ -87,7 +87,7 @@ TArray<FString> FGenericErrorReport::GetFilesToUpload() const
 	return FilesToUpload;
 }
 
-bool FGenericErrorReport::LoadWindowsReportXmlFile(TArray<uint8>& OutBuffer) const
+bool FGenericErrorReport::LoadWindowsReportXmlFile( FString& OutString ) const
 {
 	// Find .xml file
 	FString XmlFilename;
@@ -96,7 +96,7 @@ bool FGenericErrorReport::LoadWindowsReportXmlFile(TArray<uint8>& OutBuffer) con
 		return false;
 	}
 
-	return FFileHelper::LoadFileToArray(OutBuffer, *(ReportDirectory / XmlFilename));
+	return FFileHelper::LoadFileToString( OutString, *(ReportDirectory / XmlFilename) );
 }
 
 bool FGenericErrorReport::TryReadDiagnosticsFile(FText& OutReportDescription)
diff --git a/Engine/Source/Programs/CrashReporter/CrashReportClient/Private/GenericErrorReport.h b/Engine/Source/Programs/CrashReporter/CrashReportClient/Private/GenericErrorReport.h
--- a/Engine/Source/Programs/CrashReporter/CrashReportClient/Private/GenericErrorReport.h
+++ b/Engine/Source/Programs/CrashReporter/CrashReportClient/Private/GenericErrorReport.h
@@ -54,10 +54,10 @@ public:
 	 * Load the WER XML file for this report
 	 * @note This is Windows specific and so shouldn't really be part of the public interface, but currently the server
 	 * is Windows-specific in its checking of reports, so this is needed.
-	 * @param OutBuffer Buffer to load the file into
+	 * @param OutString String to load the file into
 	 * @return Whether finding and loading the file succeeded
 	 */
-	bool LoadWindowsReportXmlFile(TArray<uint8>& OutBuffer) const;
+	bool LoadWindowsReportXmlFile( FString& OutString ) const;
 
 	/**
 	 * @param Description (exception and callstack) to fill in if succesful
