CrashReporter - XML WER is always send as UTF8 regardless the internal format (changed from UTF16 to UTF8)

diff --git a/Engine/Source/Programs/CrashReporter/CrashReportClient/Private/CrashUpload.cpp b/Engine/Source/Programs/CrashReporter/CrashReportClient/Private/CrashUpload.cpp
--- a/Engine/Source/Programs/CrashReporter/CrashReportClient/Private/CrashUpload.cpp
+++ b/Engine/Source/Programs/CrashReporter/CrashReportClient/Private/CrashUpload.cpp
@@ -128,14 +128,16 @@ bool FCrashUpload::SendCheckReportRequest()
 	{
 		// This part is Windows-specific on the server
 		ErrorReport.LoadWindowsReportXmlFile( XMLString );
-		// Copy the string data into a TArray<uint8>.
-		const int32 XMLStringSize = XMLString.Len() * sizeof(TCHAR);
-		PostData.Reset( XMLStringSize );
-		PostData.AddUninitialized( XMLStringSize );
-		FMemory::Memcpy( PostData.GetData(), *XMLString, XMLStringSize );
+
+		// Convert the XMLString into the UTF-8.
+		FTCHARToUTF8 Converter( (const TCHAR*)*XMLString, XMLString.Len() );
+		const int32 Length = Converter.Length();
+		PostData.Reset( Length );
+		PostData.AddUninitialized( Length );
+		CopyAssignItems( (ANSICHAR*)PostData.GetTypedData(), Converter.Get(), Length );
 
 		Request->SetURL(UrlPrefix / TEXT("CheckReportDetail"));
-		Request->SetHeader(TEXT("Content-Type"), TEXT("text/plain; charset=utf-16"));
+		Request->SetHeader(TEXT("Content-Type"), TEXT("text/plain; charset=utf-8"));
 	}
 
 	UE_LOG( CrashReportClientLog, Log, TEXT( "PostData Num: %i" ), PostData.Num() );
