Fixed a few compile errors (mostly non-unity)

diff --git a/Engine/Source/Editor/MainFrame/Private/Menus/MainMenu.cpp b/Engine/Source/Editor/MainFrame/Private/Menus/MainMenu.cpp
--- a/Engine/Source/Editor/MainFrame/Private/Menus/MainMenu.cpp
+++ b/Engine/Source/Editor/MainFrame/Private/Menus/MainMenu.cpp
@@ -8,6 +8,7 @@
 #include "Runtime/Core/Public/Features/IModularFeatures.h"
 #include "SourceCodeNavigation.h"
 #include "Toolkits/AssetEditorToolkit.h"
+#include "TranslationEditorMenu.h"
 
 
 #define LOCTEXT_NAMESPACE "MainFileMenu"
diff --git a/Engine/Source/Programs/CrashReporter/CrashReportClient/Private/Mac/CrashReportClientMainMac.cpp b/Engine/Source/Programs/CrashReporter/CrashReportClient/Private/Mac/CrashReportClientMainMac.cpp
--- a/Engine/Source/Programs/CrashReporter/CrashReportClient/Private/Mac/CrashReportClientMainMac.cpp
+++ b/Engine/Source/Programs/CrashReporter/CrashReportClient/Private/Mac/CrashReportClientMainMac.cpp
@@ -1,6 +1,7 @@
 // Copyright 1998-2014 Epic Games, Inc. All Rights Reserved.
 
 #include "CrashReportClientApp.h"
+#include "ExceptionHandling.h"
 
 /**
  * Because crash reporters can crash, too
diff --git a/Engine/Source/Programs/CrashReporter/CrashReportClient/Private/Mac/MacErrorReport.cpp b/Engine/Source/Programs/CrashReporter/CrashReportClient/Private/Mac/MacErrorReport.cpp
--- a/Engine/Source/Programs/CrashReporter/CrashReportClient/Private/Mac/MacErrorReport.cpp
+++ b/Engine/Source/Programs/CrashReporter/CrashReportClient/Private/Mac/MacErrorReport.cpp
@@ -1,6 +1,7 @@
 // Copyright 1998-2014 Epic Games, Inc. All Rights Reserved.
 
 #include "CrashReportClientApp.h"
+#include "../CrashReportUtil.h"
 
 #include "MacErrorReport.h"
 
diff --git a/Engine/Source/Programs/UnrealFrontend/Private/Mac/MacUnrealFrontendMain.cpp b/Engine/Source/Programs/UnrealFrontend/Private/Mac/MacUnrealFrontendMain.cpp
--- a/Engine/Source/Programs/UnrealFrontend/Private/Mac/MacUnrealFrontendMain.cpp
+++ b/Engine/Source/Programs/UnrealFrontend/Private/Mac/MacUnrealFrontendMain.cpp
@@ -1,6 +1,8 @@
 // Copyright 1998-2014 Epic Games, Inc. All Rights Reserved.
 
 #include "UnrealFrontendMain.h"
+#include "ExceptionHandling.h"
+#include "RequiredProgramMainCPPInclude.h"
 
 
 static FString GSavedCommandLine;
diff --git a/Engine/Source/Programs/UnrealSync/Private/P4Env.cpp b/Engine/Source/Programs/UnrealSync/Private/P4Env.cpp
--- a/Engine/Source/Programs/UnrealSync/Private/P4Env.cpp
+++ b/Engine/Source/Programs/UnrealSync/Private/P4Env.cpp
@@ -4,6 +4,7 @@
 
 #include "P4Env.h"
 #include "ProcessHelper.h"
+#include "Regex.h"
 
 FString FP4Env::GetParam(const TCHAR* CommandLine, const FString& ParamName)
 {
diff --git a/Engine/Source/Runtime/Launch/Public/LaunchEngineLoop.h b/Engine/Source/Runtime/Launch/Public/LaunchEngineLoop.h
--- a/Engine/Source/Runtime/Launch/Public/LaunchEngineLoop.h
+++ b/Engine/Source/Runtime/Launch/Public/LaunchEngineLoop.h
@@ -6,6 +6,9 @@
 
 #pragma once
 
+#if WITH_ENGINE
+#include "ISessionService.h"
+#endif
 
 /**
  * Implements the main engine loop.	
