CrashReportClient - bCompileSteamOSS=false

diff --git a/Engine/Source/Programs/CrashReporter/CrashReportClient/CrashReportClient.Target.cs b/Engine/Source/Programs/CrashReporter/CrashReportClient/CrashReportClient.Target.cs
--- a/Engine/Source/Programs/CrashReporter/CrashReportClient/CrashReportClient.Target.cs
+++ b/Engine/Source/Programs/CrashReporter/CrashReportClient/CrashReportClient.Target.cs
@@ -27,6 +27,7 @@ public class CrashReportClientTarget : TargetRules
 		if( base.GetSupportedConfigurations( ref OutConfigurations, bIncludeTestAndShippingConfigs ) )
 		{
 			OutConfigurations.Add( UnrealTargetConfiguration.Shipping );
+			OutConfigurations.Add( UnrealTargetConfiguration.Debug );
 			return true;
 		}
 		else
@@ -77,6 +78,7 @@ public class CrashReportClientTarget : TargetRules
 		UEBuildConfiguration.bCompileAgainstEngine = false;
 		UEBuildConfiguration.bCompileAgainstCoreUObject = false;
 		UEBuildConfiguration.bUseLoggingInShipping = true;
+		UEBuildConfiguration.bCompileSteamOSS = false;
 
 		UEBuildConfiguration.bIncludeADO = (Target.Platform == UnrealTargetPlatform.Win64 || Target.Platform == UnrealTargetPlatform.Win32);
 
