CrashReportWebsite - Updated project files

diff --git a/Engine/Source/Programs/CrashReporter/CrashReportWebSite/CrashReportWebSite.csproj b/Engine/Source/Programs/CrashReporter/CrashReportWebSite/CrashReportWebSite.csproj
--- a/Engine/Source/Programs/CrashReporter/CrashReportWebSite/CrashReportWebSite.csproj
+++ b/Engine/Source/Programs/CrashReporter/CrashReportWebSite/CrashReportWebSite.csproj
@@ -23,7 +23,7 @@
     <SccProvider>
     </SccProvider>
     <TargetFrameworkProfile />
-    <UseIISExpress>false</UseIISExpress>
+    <UseIISExpress>true</UseIISExpress>
     <FileUpgradeFlags>
     </FileUpgradeFlags>
     <OldToolsVersion>4.0</OldToolsVersion>
@@ -32,6 +32,7 @@
     <IISExpressAnonymousAuthentication />
     <IISExpressWindowsAuthentication />
     <IISExpressUseClassicPipelineMode />
+    <MvcProjectUpgradeChecked>true</MvcProjectUpgradeChecked>
   </PropertyGroup>
   <PropertyGroup Condition=" '$(Configuration)|$(Platform)' == 'Debug|AnyCPU' ">
     <DebugSymbols>true</DebugSymbols>
@@ -266,6 +267,10 @@
     <DocumentationFile>bin\CrashReportWebSite.XML</DocumentationFile>
     <RunCodeAnalysis>true</RunCodeAnalysis>
   </PropertyGroup>
+  <PropertyGroup Condition="'$(Configuration)|$(Platform)' == 'Release|AnyCPU'">
+    <OutputPath>bin\</OutputPath>
+    <DocumentationFile>bin\CrashReportWebSite.XML</DocumentationFile>
+  </PropertyGroup>
   <Import Project="$(MSBuildBinPath)\Microsoft.CSharp.targets" />
   <Import Project="$(VSToolsPath)\WebApplications\Microsoft.WebApplication.targets" Condition="'$(VSToolsPath)' != ''" />
   <Import Project="$(MSBuildExtensionsPath32)\Microsoft\VisualStudio\v10.0\WebApplications\Microsoft.WebApplication.targets" Condition="false" />
@@ -280,12 +285,11 @@
     <VisualStudio>
       <FlavorProperties GUID="{349c5851-65df-11da-9384-00065b846f21}">
         <WebProjectProperties>
-          <UseIIS>False</UseIIS>
+          <UseIIS>True</UseIIS>
           <AutoAssignPort>True</AutoAssignPort>
           <DevelopmentServerPort>61478</DevelopmentServerPort>
           <DevelopmentServerVPath>/</DevelopmentServerVPath>
-          <IISUrl>
-          </IISUrl>
+          <IISUrl>http://localhost:80/</IISUrl>
           <NTLMAuthentication>False</NTLMAuthentication>
           <UseCustomServer>False</UseCustomServer>
           <CustomServerUrl>
diff --git a/Engine/Source/Programs/CrashReporter/CrashReportWebSite/Web.config b/Engine/Source/Programs/CrashReporter/CrashReportWebSite/Web.config
--- a/Engine/Source/Programs/CrashReporter/CrashReportWebSite/Web.config
+++ b/Engine/Source/Programs/CrashReporter/CrashReportWebSite/Web.config
@@ -67,7 +67,7 @@
   <system.webServer>
     <modules runAllManagedModulesForAllRequests="false">
       <add name="UriCacheModule"/>
-      <add name="FileCacheModule"/>
+      <!--<add name="FileCacheModule"/>-->
       <add name="TokenCacheModule"/>
       <add name="RequestMonitorModule"/>
     </modules>
