CrashReport
- Crashes from the UE4-QA are removed from the disk and not ignored as they are from CL0
- Removed obsolete method for extracting the CL, replaced with the latest one
- Command line should be properly displayed now (without !AssertLog)

diff --git a/Engine/Source/Programs/CrashReporter/CrashReportWebSite/Models/CrashRepository.cs b/Engine/Source/Programs/CrashReporter/CrashReportWebSite/Models/CrashRepository.cs
--- a/Engine/Source/Programs/CrashReporter/CrashReportWebSite/Models/CrashRepository.cs
+++ b/Engine/Source/Programs/CrashReporter/CrashReportWebSite/Models/CrashRepository.cs
@@ -825,24 +825,6 @@ namespace Tools.CrashReporter.CrashReportWebSite.Models
 			return UserGroupNameId;
 		}
 
-		/// <summary>
-		/// Extract the changelist from the 4 digit version number.
-		/// </summary>
-		/// <param name="NewCrash">The crash to extract the changelist for.</param>
-		/// <returns>The changelist number the build was made from.</returns>
-		private int DeriveChangelist( Crash NewCrash )
-		{
-			int Changelist = 0;
-			Version ApplicationVersion = null;
-
-			if( Version.TryParse( NewCrash.BuildVersion, out ApplicationVersion ) )
-			{
-				Changelist = ApplicationVersion.Build * 65536 + ApplicationVersion.Revision;
-			}
-
-			return Changelist;
-		}
-
 		// This string is written to the report in FCrashReportClient's constructor
 		const string UserNamePrefix = "!Name:";
 
@@ -859,8 +841,8 @@ namespace Tools.CrashReporter.CrashReportWebSite.Models
 			NewCrash.Branch = NewCrashInfo.BranchName;
 			NewCrash.BaseDir = NewCrashInfo.BaseDir;
 			NewCrash.BuildVersion = NewCrashInfo.BuildVersion;
-			NewCrash.ChangeListVersion = DeriveChangelist( NewCrash ).ToString();
-			NewCrash.CommandLine = NewCrashInfo.CommandLine.Trim( "!".ToCharArray() );
+			NewCrash.ChangeListVersion = NewCrashInfo.BuiltFromCL.ToString();
+			NewCrash.CommandLine = NewCrashInfo.CommandLine;
 			NewCrash.EngineMode = NewCrashInfo.EngineMode;
 
 			NewCrash.ComputerName = NewCrashInfo.MachineGuid;
