Fixing CIS (MinidumpDiagnostics)

diff --git a/Engine/Source/Programs/CrashReporter/MinidumpDiagnostics/MinidumpDiagnostics.Build.cs b/Engine/Source/Programs/CrashReporter/MinidumpDiagnostics/MinidumpDiagnostics.Build.cs
--- a/Engine/Source/Programs/CrashReporter/MinidumpDiagnostics/MinidumpDiagnostics.Build.cs
+++ b/Engine/Source/Programs/CrashReporter/MinidumpDiagnostics/MinidumpDiagnostics.Build.cs
@@ -13,6 +13,7 @@ public class MinidumpDiagnostics : ModuleRules
 			new string[] 
 			{ 
 				"Core",
+				"CoreUObject",
 				"CrashDebugHelper",
 				"PerforceSourceControl",
 				"Projects"
diff --git a/Engine/Source/Runtime/Launch/Private/LaunchEngineLoop.cpp b/Engine/Source/Runtime/Launch/Private/LaunchEngineLoop.cpp
--- a/Engine/Source/Runtime/Launch/Private/LaunchEngineLoop.cpp
+++ b/Engine/Source/Runtime/Launch/Private/LaunchEngineLoop.cpp
@@ -86,6 +86,10 @@
 	#include "HideWindowsPlatformTypes.h"
 #endif
 
+#if WITH_COREUOBJECT
+	#include "CoreUObject.h"
+#endif
+
 // Pipe output to std output
 // This enables UBT to collect the output for it's own use
 class FOutputDeviceStdOutput : public FOutputDevice
