Fixing wrong shared pch warning when building MiniDumpDiagnostics

diff --git a/Engine/Source/Programs/CrashReporter/MinidumpDiagnostics/Public/MinidumpDiagnosticsApp.h b/Engine/Source/Programs/CrashReporter/MinidumpDiagnostics/Public/MinidumpDiagnosticsApp.h
--- a/Engine/Source/Programs/CrashReporter/MinidumpDiagnostics/Public/MinidumpDiagnosticsApp.h
+++ b/Engine/Source/Programs/CrashReporter/MinidumpDiagnostics/Public/MinidumpDiagnosticsApp.h
@@ -2,7 +2,7 @@
 
 #pragma once
 
-#include "Core.h"
+#include "CoreUObject.h"
 
 // Extract a debugging report from a minidump
 int32 RunMinidumpDiagnostics(int32 ArgC, TCHAR* Argv[]);
