Removing direct depency on CoreUObject from MinidumpDiagnostics, no longer needed.

diff --git a/Engine/Source/Programs/CrashReporter/MinidumpDiagnostics/MinidumpDiagnostics.Build.cs b/Engine/Source/Programs/CrashReporter/MinidumpDiagnostics/MinidumpDiagnostics.Build.cs
--- a/Engine/Source/Programs/CrashReporter/MinidumpDiagnostics/MinidumpDiagnostics.Build.cs
+++ b/Engine/Source/Programs/CrashReporter/MinidumpDiagnostics/MinidumpDiagnostics.Build.cs
@@ -13,7 +13,6 @@ public class MinidumpDiagnostics : ModuleRules
 			new string[] 
 			{ 
 				"Core",
-				"CoreUObject",
 				"CrashDebugHelper",
 				"PerforceSourceControl",
 				"Projects"
diff --git a/Engine/Source/Programs/CrashReporter/MinidumpDiagnostics/Public/MinidumpDiagnosticsApp.h b/Engine/Source/Programs/CrashReporter/MinidumpDiagnostics/Public/MinidumpDiagnosticsApp.h
--- a/Engine/Source/Programs/CrashReporter/MinidumpDiagnostics/Public/MinidumpDiagnosticsApp.h
+++ b/Engine/Source/Programs/CrashReporter/MinidumpDiagnostics/Public/MinidumpDiagnosticsApp.h
@@ -2,7 +2,7 @@
 
 #pragma once
 
-#include "CoreUObject.h"
+#include "Core.h"
 
 // Extract a debugging report from a minidump
 int32 RunMinidumpDiagnostics(int32 ArgC, TCHAR* Argv[]);
