CrashReportWWW - Minor code formatting

diff --git a/Engine/Source/Programs/CrashReporter/CrashReportWebSite/Controllers/BuggsController.cs b/Engine/Source/Programs/CrashReporter/CrashReportWebSite/Controllers/BuggsController.cs
--- a/Engine/Source/Programs/CrashReporter/CrashReportWebSite/Controllers/BuggsController.cs
+++ b/Engine/Source/Programs/CrashReporter/CrashReportWebSite/Controllers/BuggsController.cs
@@ -117,6 +117,8 @@ namespace Tools.CrashReporter.CrashReportWebSite.Controllers
 				{
 					Bugg.Description = BuggsForm["Description"];
 				}
+
+				// <STATUS>
 			}
 
 			// Set up the view model with the crash data
diff --git a/Engine/Source/Programs/CrashReporter/CrashReportWebSite/Controllers/CrashesController.cs b/Engine/Source/Programs/CrashReporter/CrashReportWebSite/Controllers/CrashesController.cs
--- a/Engine/Source/Programs/CrashReporter/CrashReportWebSite/Controllers/CrashesController.cs
+++ b/Engine/Source/Programs/CrashReporter/CrashReportWebSite/Controllers/CrashesController.cs
@@ -63,6 +63,8 @@ namespace Tools.CrashReporter.CrashReportWebSite.Controllers
 				}
 			}
 
+			// <STATUS>
+
 			// Parse the contents of the query string, and populate the form
 			FormHelper FormData = new FormHelper( Request, CrashesForm, "TimeOfCrash" );
 			CrashesViewModel Result = LocalCrashRepository.GetResults( FormData );
