CrashReportWWW - Fixed filtering by date

diff --git a/Engine/Source/Programs/CrashReporter/CrashReportWebSite/Models/BuggRepository.cs b/Engine/Source/Programs/CrashReporter/CrashReportWebSite/Models/BuggRepository.cs
--- a/Engine/Source/Programs/CrashReporter/CrashReportWebSite/Models/BuggRepository.cs
+++ b/Engine/Source/Programs/CrashReporter/CrashReportWebSite/Models/BuggRepository.cs
@@ -441,15 +441,12 @@ namespace Tools.CrashReporter.CrashReportWebSite.Models
 		/// <returns>The set of Buggs between the earliest and latest date.</returns>
 		public IQueryable<Bugg> FilterByDate( IQueryable<Bugg> Results, DateTime DateFrom, DateTime DateTo )
 		{
-			var LocalDateFrom = DateFrom.ToLocalTime();
-			var LocalDateTo = DateTo.ToLocalTime();
-
-			IQueryable<Bugg> BuggsInTimeFrame = Results.Where( Bugg => Bugg.TimeOfLastCrash >= LocalDateFrom && Bugg.TimeOfLastCrash <= LocalDateTo.AddDays( 1 ) );
+			IQueryable<Bugg> BuggsInTimeFrame = Results.Where( Bugg => Bugg.TimeOfLastCrash >= DateFrom && Bugg.TimeOfLastCrash <= AddOneDayToDate( DateTo ) );
 
 #if DEBUG
 			foreach( var MyBugg in BuggsInTimeFrame )
 			{
-				Debug.WriteLine( "MyBugg=" + MyBugg.ToString() );
+				Debug.WriteLine( "FilterByDate=" + MyBugg.ToString() );
 			}
 #endif
 
@@ -479,7 +476,7 @@ namespace Tools.CrashReporter.CrashReportWebSite.Models
 #if DEBUG
 				foreach( var MyBugg in BuggsForBuildVersion )
 				{
-					Debug.WriteLine( "MyBugg=" + MyBugg.ToString() );
+					Debug.WriteLine( "FilterByBuildVersion=" + MyBugg.ToString() );
 				}			
 #endif
 			}
@@ -507,6 +504,13 @@ namespace Tools.CrashReporter.CrashReportWebSite.Models
 					where UserGroupDetail.Name.Contains(UserGroup)
 					select BuggDetail
 				);
+
+#if DEBUG
+				foreach( var MyBugg in NewSetOfBuggs )
+				{
+					Debug.WriteLine( "FilterByUserGroup=" + MyBugg.ToString() );
+				}
+#endif
 			}
 			catch( Exception Ex )
 			{
@@ -516,6 +520,11 @@ namespace Tools.CrashReporter.CrashReportWebSite.Models
 			return NewSetOfBuggs;
 		}
 
+		private DateTime AddOneDayToDate( DateTime Date )
+		{
+			return Date.AddDays( 1 );
+		}
+
 		/// <summary>
 		/// Sort the container of Buggs by the requested criteria.
 		/// </summary>
@@ -532,7 +541,7 @@ namespace Tools.CrashReporter.CrashReportWebSite.Models
 				var IntermediateResults =
 				(
 					from BuggCrashDetail in BuggsDataContext.Buggs_Crashes
-					where BuggCrashDetail.Crash.TimeOfCrash >= DateFrom && BuggCrashDetail.Crash.TimeOfCrash <= DateTo
+					where BuggCrashDetail.Crash.TimeOfCrash >= DateFrom && BuggCrashDetail.Crash.TimeOfCrash <= AddOneDayToDate( DateTo )
 					group BuggCrashDetail by BuggCrashDetail.BuggId into CrashesGrouped
 					join BuggDetail in Results on CrashesGrouped.Key equals BuggDetail.Id
 					select new { Bugg = BuggDetail, Count = CrashesGrouped.Count() }
