CrashReport - Added MachineID/EpicID to the crashes (crash website)

diff --git a/Engine/Source/Programs/CrashReporter/CrashReportWebSite/Models/CrashReport.dbml b/Engine/Source/Programs/CrashReporter/CrashReportWebSite/Models/CrashReport.dbml
--- a/Engine/Source/Programs/CrashReporter/CrashReportWebSite/Models/CrashReport.dbml
+++ b/Engine/Source/Programs/CrashReporter/CrashReportWebSite/Models/CrashReport.dbml
@@ -95,6 +95,8 @@
       <Column Name="CrashType" Type="System.Int16" DbType="SmallInt" CanBeNull="true" />
       <Column Name="HasMetaData" Type="System.Boolean" DbType="Bit" CanBeNull="true" />
       <Column Name="SourceContext" Type="System.String" DbType="VarChar(MAX)" CanBeNull="true" />
+      <Column Name="EpicAccountId" Type="System.String" DbType="varchar(64)" CanBeNull="true" />
+      <Column Name="EngineVersion" Type="System.String" DbType="varchar(64)" CanBeNull="true" />
       <Association Name="Crash_Buggs_Crash" Member="Buggs_Crashes" ThisKey="Id" OtherKey="CrashId" Type="Buggs_Crash" />
       <Association Name="Crash_Crash_FunctionCall" Member="Crash_FunctionCalls" ThisKey="Id" OtherKey="CrashId" Type="Crash_FunctionCall" />
       <Association Name="User_Crash" Member="UserById" Storage="_User" ThisKey="UserNameId" OtherKey="Id" Type="User" IsForeignKey="true" />
diff --git a/Engine/Source/Programs/CrashReporter/CrashReportWebSite/Models/CrashReport.dbml.layout b/Engine/Source/Programs/CrashReporter/CrashReportWebSite/Models/CrashReport.dbml.layout
--- a/Engine/Source/Programs/CrashReporter/CrashReportWebSite/Models/CrashReport.dbml.layout
+++ b/Engine/Source/Programs/CrashReporter/CrashReportWebSite/Models/CrashReport.dbml.layout
@@ -1,5 +1,5 @@
 ﻿<?xml version="1.0" encoding="utf-8"?>
-<ordesignerObjectsDiagram dslVersion="1.0.0.0" absoluteBounds="0, 0, 13.75, 8.5" name="CrashReport">
+<ordesignerObjectsDiagram dslVersion="1.0.0.0" absoluteBounds="0, 0, 13.75, 9.25" name="CrashReport">
   <DataContextMoniker Name="/CrashReportDataContext" />
   <nestedChildShapes>
     <classShape Id="a37f1cdb-93d7-48ef-8da7-2b7ca09282b0" absoluteBounds="11.5, 1.25, 2, 4.0785139973958326">
@@ -32,10 +32,10 @@
         <elementListCompartment Id="77c436bb-7563-449c-a08a-8c417bcde3fa" absoluteBounds="3.3900000000000006, 6.3350000000000009, 1.9700000000000002, 0.8262939453125" name="DataPropertiesCompartment" titleTextColor="Black" itemTextColor="Black" />
       </nestedChildShapes>
     </classShape>
-    <classShape Id="f1204ef8-6660-48ce-8c1d-131c49e2dd25" absoluteBounds="6.125, 0.625, 2, 7.5399397786458309">
+    <classShape Id="f1204ef8-6660-48ce-8c1d-131c49e2dd25" absoluteBounds="6.125, 1.125, 2, 7.9245426432291648">
       <DataClassMoniker Name="/CrashReportDataContext/Crash" />
       <nestedChildShapes>
-        <elementListCompartment Id="53c11757-f0be-47da-bcae-7d20a571002e" absoluteBounds="6.1400000000000006, 1.0850000000000009, 1.9700000000000002, 6.979939778645833" name="DataPropertiesCompartment" titleTextColor="Black" itemTextColor="Black" />
+        <elementListCompartment Id="53c11757-f0be-47da-bcae-7d20a571002e" absoluteBounds="6.1400000000000006, 1.5850000000000009, 1.9700000000000002, 7.3645426432291661" name="DataPropertiesCompartment" titleTextColor="Black" itemTextColor="Black" />
       </nestedChildShapes>
     </classShape>
     <classShape Id="a188623d-4f2c-4a7c-b729-d354a2df0d5a" absoluteBounds="9, 6.25, 2, 1.1939925130208344">
@@ -63,7 +63,7 @@
         <classShapeMoniker Id="b4a7205e-4223-43f1-837f-ac4a1390bf89" />
       </nodes>
     </associationConnector>
-    <associationConnector edgePoints="[(5.375 : 4.56814697265625); (6.125 : 4.56814697265625)]" fixedFrom="NotFixed" fixedTo="NotFixed">
+    <associationConnector edgePoints="[(5.375 : 4.82381790161133); (6.125 : 4.82381790161133)]" fixedFrom="NotFixed" fixedTo="NotFixed">
       <AssociationMoniker Name="/CrashReportDataContext/User/User_Crash" />
       <nodes>
         <classShapeMoniker Id="01dddf4f-058f-43b4-b6c0-b90e7acfe0b9" />
@@ -77,7 +77,7 @@
         <classShapeMoniker Id="8a60fc4a-a797-417f-87ed-25b22bd5e420" />
       </nodes>
     </associationConnector>
-    <associationConnector edgePoints="[(11.5 : 2.47199625651042); (10.875 : 2.47199625651042)]" fixedFrom="Algorithm" fixedTo="Algorithm">
+    <associationConnector edgePoints="[(11.5 : 2.47199625651042); (10.875 : 2.47199625651042)]" fixedFrom="NotFixed" fixedTo="NotFixed">
       <AssociationMoniker Name="/CrashReportDataContext/Bugg/Bugg_Buggs_Crash" />
       <nodes>
         <classShapeMoniker Id="a37f1cdb-93d7-48ef-8da7-2b7ca09282b0" />
@@ -112,7 +112,7 @@
         <classShapeMoniker Id="a188623d-4f2c-4a7c-b729-d354a2df0d5a" />
       </nodes>
     </associationConnector>
-    <associationConnector edgePoints="[(5.375 : 4.15907348632812); (6.125 : 4.15907348632812)]" fixedFrom="NotFixed" fixedTo="NotFixed">
+    <associationConnector edgePoints="[(5.375 : 4.28690895080566); (6.125 : 4.28690895080566)]" fixedFrom="NotFixed" fixedTo="NotFixed">
       <AssociationMoniker Name="/CrashReportDataContext/User/User_Crash1" />
       <nodes>
         <classShapeMoniker Id="01dddf4f-058f-43b4-b6c0-b90e7acfe0b9" />
diff --git a/Engine/Source/Programs/CrashReporter/CrashReportWebSite/Models/CrashRepository.cs b/Engine/Source/Programs/CrashReporter/CrashReportWebSite/Models/CrashRepository.cs
--- a/Engine/Source/Programs/CrashReporter/CrashReportWebSite/Models/CrashRepository.cs
+++ b/Engine/Source/Programs/CrashReporter/CrashReportWebSite/Models/CrashRepository.cs
@@ -825,9 +825,15 @@ namespace Tools.CrashReporter.CrashReportWebSite.Models
 			return UserGroupNameId;
 		}
 
-		// This string is written to the report in FCrashReportClient's constructor
+		/// <summary> 
+		/// This string is written to the report in FCrashReportClient's constructor.
+		/// This is obsolete, but left here for the backward compatibility.
+		/// </summary>
 		const string UserNamePrefix = "!Name:";
 
+		/// <summary> Special user name, currently used to mark crashes from UE4 releases. </summary>
+		const string UserNameAnonymous = "Anonymous";
+
 		/// <summary>
 		/// Add a new crash to the db based on a CrashDescription sent from the client.
 		/// </summary>
@@ -845,13 +851,29 @@ namespace Tools.CrashReporter.CrashReportWebSite.Models
 			NewCrash.CommandLine = NewCrashInfo.CommandLine;
 			NewCrash.EngineMode = NewCrashInfo.EngineMode;
 
-			NewCrash.ComputerName = NewCrashInfo.MachineGuid;
-
-			// Check for internal crashes with user name in the form "!Name:<name>"
-			NewCrash.UserNameId = CrashRepositoryDataContext.FindOrAddUser(
-				NewCrashInfo.MachineGuid.StartsWith(UserNamePrefix) ?
-					NewCrashInfo.MachineGuid.Substring(UserNamePrefix.Count()) : "Anonymous");
-
+			// Valid MachineID and UserName, updated crash from non-UE4 release
+			if( !string.IsNullOrEmpty(NewCrashInfo.MachineGuid) && !string.IsNullOrEmpty(NewCrashInfo.UserName) )
+			{
+				NewCrash.ComputerName = NewCrashInfo.MachineGuid;
+				NewCrash.UserNameId = CrashRepositoryDataContext.FindOrAddUser( NewCrashInfo.UserName );
+			}
+			// Valid MachineID and EpicAccountId, updated crash from UE4 release
+			else if( !string.IsNullOrEmpty( NewCrashInfo.MachineGuid ) && !string.IsNullOrEmpty( NewCrashInfo.EpicAccountId ) )
+			{
+				NewCrash.EpicAccountId = NewCrashInfo.EpicAccountId;
+				NewCrash.UserNameId = CrashRepositoryDataContext.FindOrAddUser( UserNameAnonymous );
+			}
+			// Crash from an older version.
+			else
+			{
+				// MachineGuid for older crashes is obsolete, so ignore it.
+				//NewCrash.ComputerName = NewCrashInfo.MachineGuid;
+				NewCrash.UserNameId = CrashRepositoryDataContext.FindOrAddUser
+				(
+					!string.IsNullOrEmpty( NewCrashInfo.UserName ) ? NewCrashInfo.UserName : UserNameAnonymous
+				);
+			}
+			
 			NewCrash.Description = "";
 			if( NewCrashInfo.UserDescription != null )
 			{
@@ -860,7 +882,7 @@ namespace Tools.CrashReporter.CrashReportWebSite.Models
 
 			NewCrash.EngineMode = NewCrashInfo.EngineMode;
 			NewCrash.GameName = NewCrashInfo.GameName;
-			NewCrash.LanguageExt = NewCrashInfo.Language;		// Convert
+			NewCrash.LanguageExt = NewCrashInfo.Language; // Converted by the crash process.
 			NewCrash.PlatformName = NewCrashInfo.Platform;
 
 			if (NewCrashInfo.AssertionMessage != null)
@@ -906,6 +928,28 @@ namespace Tools.CrashReporter.CrashReportWebSite.Models
 			// As we're adding it, the status is always new
 			NewCrash.Status = "New";
 
+			/*
+				Unused Crashes' fields.
+	 
+				Title				nchar(20)		
+				Selected			bit				
+				Version				int				
+				AutoReporterID		int				
+				Processed			bit				
+				HasDiagnosticsFile	bit				
+				HasNewLogFile		bit				
+				HasMetaData			bit	
+			*/
+			// Set the unused fields to the default values.
+			NewCrash.Title = "";
+			NewCrash.Selected = false;
+			NewCrash.Version = 4;
+			NewCrash.AutoReporterID = 0;
+			NewCrash.Processed = true;
+			//NewCrash.HasDiagnosticsFile = true;
+			NewCrash.HasNewLogFile = false;
+			//NewCrash.HasMetaData = true;
+
 			// Add the crash to the database
 			CrashRepositoryDataContext.Crashes.InsertOnSubmit( NewCrash );
 			CrashRepositoryDataContext.SubmitChanges();
diff --git a/Engine/Source/Programs/CrashReporter/CrashReportWebSite/Views/Crashes/Show.aspx b/Engine/Source/Programs/CrashReporter/CrashReportWebSite/Views/Crashes/Show.aspx
--- a/Engine/Source/Programs/CrashReporter/CrashReportWebSite/Views/Crashes/Show.aspx
+++ b/Engine/Source/Programs/CrashReporter/CrashReportWebSite/Views/Crashes/Show.aspx
@@ -182,9 +182,12 @@
 			<dt>Platform</dt>
 				<dd ><%=Html.DisplayFor(m => Model.Crash.PlatformName) %></dd>
 
-			<dt>Computer</dt> 
+			<dt>Machine Id</dt> 
 				<dd class='even'><%=Html.DisplayFor(m => Model.Crash.ComputerName) %></dd>
 
+			<dt>Epic Account Id</dt> 
+				<dd class='even'><%=Html.DisplayFor(m => Model.Crash.EpicAccountId) %></dd>
+
 			<dt>Build Version</dt>
 				<dd ><%=Html.DisplayFor(m => Model.Crash.BuildVersion) %></dd>
 
