CrashReport - Added MachineID/EpicID to the crashes (crash report client)

diff --git a/Engine/Source/Programs/CrashReporter/CrashReportClient/Private/CrashReportClient.cpp b/Engine/Source/Programs/CrashReporter/CrashReportClient/Private/CrashReportClient.cpp
--- a/Engine/Source/Programs/CrashReporter/CrashReportClient/Private/CrashReportClient.cpp
+++ b/Engine/Source/Programs/CrashReporter/CrashReportClient/Private/CrashReportClient.cpp
@@ -4,6 +4,7 @@
 
 #include "CrashReportClient.h"
 #include "UniquePtr.h"
+#include "DesktopPlatformModule.h"
 
 #define LOCTEXT_NAMESPACE "CrashReportClient"
 
@@ -19,36 +20,44 @@ FString GCrashUserId;
 
 #if !CRASH_REPORT_UNATTENDED_ONLY
 
-FCrashReportClient::FCrashReportClient(const FPlatformErrorReport& ErrorReport, const FString& AppName)
+FCrashReportClient::FCrashReportClient(const FPlatformErrorReport& InErrorReport, const FString& AppName)
 	: AppState(EApplicationState::Ready)
 	, SubmittedCountdown(-1)
 	, bDiagnosticFileSent(false)
-	, ErrorReportFiles(ErrorReport)
+	, ErrorReport( InErrorReport )
 	, Uploader(GServerIP)
 	, CrashedAppName(AppName)
 	, CancelButtonText(LOCTEXT("Cancel", "Don't Send"))
 {
+	IDesktopPlatform* DesktopPlatform = FDesktopPlatformModule::Get();
+	const FString MachineId = DesktopPlatform->GetMachineId().ToString( EGuidFormats::Digits );
+
+	// The Epic ID can be looked up from this ID.
+	const FString EpicAccountId = DesktopPlatform->GetEpicAccountId();
+
+	// Remove periods from internal user names to match AutoReporter user names
+	// The name prefix is read by CrashRepository.AddNewCrash in the website code
+	const FString UserNameNoDot = FString( FPlatformProcess::UserName() ).Replace( TEXT( "." ), TEXT( "" ) );
+
 	// Set global user name ID: will be added to the report
 	if (FRocketSupport::IsRocket())
 	{
-		// The Epic ID can be looked up from this ID
-		auto DeviceId = FPlatformMisc::GetUniqueDeviceId();
-		GCrashUserId = FString(TEXT("!Id:")) + DeviceId;
+		GCrashUserId = FString::Printf( TEXT( "!MachineId:%s!EpicAccountId:%s" ), *MachineId, *EpicAccountId );
 	}
 	else
 	{
-		// Remove periods from internal user names to match AutoReporter user names
-		// The name prefix is read by CrashRepository.AddNewCrash in the website code
-		GCrashUserId = FString(TEXT("!Name:")) + FString(FPlatformProcess::UserName()).Replace(TEXT("."), TEXT(""));
+		GCrashUserId = FString::Printf( TEXT( "!MachineId:%s!Name:%s" ), *MachineId, *UserNameNoDot );
 	}
 
-	if (!ErrorReportFiles.TryReadDiagnosticsFile(DiagnosticText) && !FParse::Param(FCommandLine::Get(), TEXT("no-local-diagnosis")))
+	if (!ErrorReport.TryReadDiagnosticsFile(DiagnosticText) && !FParse::Param(FCommandLine::Get(), TEXT("no-local-diagnosis")))
 	{
-		auto& Worker = DiagnoseReportTask.GetTask();
-		Worker.DiagnosticText = &DiagnosticText;
-		Worker.ErrorReportFiles = &ErrorReportFiles;
+		FDiagnoseReportWorker& Worker = DiagnoseReportTask.GetTask( &DiagnosticText, MachineId, EpicAccountId, UserNameNoDot, &ErrorReport );
 		DiagnoseReportTask.StartBackgroundTask();
 	}
+	else if( !DiagnosticText.IsEmpty() )
+	{
+		DiagnosticText = FCrashReportClient::FormatDiagnosticText( DiagnosticText, MachineId, EpicAccountId, UserNameNoDot );
+	}
 }
 
 FReply FCrashReportClient::Submit()
@@ -95,7 +104,7 @@ FText FCrashReportClient::GetCancelButtonText() const
 FText FCrashReportClient::GetDiagnosticText() const
 {
 	static const FText ProcessingReportText = LOCTEXT("ProcessingReport", "Processing crash report ...");
-	return DiagnoseReportTask.IsWorkDone() ? DiagnosticText : ProcessingReportText;
+	return DiagnoseReportTask.IsDone() ? DiagnosticText : ProcessingReportText;
 }
 
 EVisibility FCrashReportClient::SubmitButtonVisibility() const
@@ -139,8 +148,8 @@ void FCrashReportClient::StartUIWillCloseTicker()
 void FCrashReportClient::StoreCommentAndUpload()
 {
 	// Call upload even if the report is empty: pending reports will be sent if any
-	ErrorReportFiles.SetUserComment(UserComment.IsEmpty() ? LOCTEXT("NoComment", "No comment provided") : UserComment);
-	Uploader.BeginUpload(ErrorReportFiles);
+	ErrorReport.SetUserComment(UserComment.IsEmpty() ? LOCTEXT("NoComment", "No comment provided") : UserComment);
+	Uploader.BeginUpload(ErrorReport);
 
 	SubmittedCountdown = 5;
 	AppState = EApplicationState::CountingDown;
@@ -160,9 +169,9 @@ bool FCrashReportClient::UIWillCloseTick(float UnusedDeltaTime)
 
 	static const FText CountdownTextFormat = LOCTEXT("CloseApplication", "Close ({0})");
 
-	if (!bDiagnosticFileSent && DiagnoseReportTask.IsWorkDone())
+	if( !bDiagnosticFileSent && DiagnoseReportTask.IsDone() )
 	{
-		auto DiagnosticsFilePath = ErrorReportFiles.GetReportDirectory() / GDiagnosticsFilename;
+		auto DiagnosticsFilePath = ErrorReport.GetReportDirectory() / GDiagnosticsFilename;
 
 		Uploader.LocalDiagnosisComplete(FPaths::FileExists(DiagnosticsFilePath) ? DiagnosticsFilePath : TEXT(""));
 		bDiagnosticFileSent = true;	
@@ -184,7 +193,7 @@ bool FCrashReportClient::UIWillCloseTick(float UnusedDeltaTime)
 
 	// IsWorkDone will always return true here (since uploader can't finish until the diagnosis has been sent), but it
 	//  has the side effect of joining the worker thread.
-	if (!Uploader.IsFinished() || !DiagnoseReportTask.IsDone())
+	if( !Uploader.IsFinished() || !DiagnoseReportTask.IsDone() )
 	{
 		// More ticks, please
 		return true;
@@ -195,16 +204,25 @@ bool FCrashReportClient::UIWillCloseTick(float UnusedDeltaTime)
 	return false;
 }
 
-#endif // !CRASH_REPORT_UNATTENDED_ONLY
-
-void FDiagnoseReportWorker::DoWork()
+FText FCrashReportClient::FormatDiagnosticText( const FText& DiagnosticText, const FString MachineId, const FString EpicAccountId, const FString UserNameNoDot )
 {
-	*DiagnosticText = FText::Format( LOCTEXT("CrashReportClientCallstackPattern", "{0}\n\n{1}"), FText::FromString(GCrashUserId), ErrorReportFiles->DiagnoseReport() );
+	if( FRocketSupport::IsRocket() )
+	{
+		return FText::Format( LOCTEXT( "CrashReportClientCallstackPattern", "MachineId:{0}\nEpicAccountId:{1}\n\n{2}" ), FText::FromString( MachineId ), FText::FromString( EpicAccountId ), DiagnosticText );
+	}
+	else
+	{
+		return FText::Format( LOCTEXT( "CrashReportClientCallstackPattern", "MachineId:{0}\nUserName:{1}\n\n{2}" ), FText::FromString( MachineId ), FText::FromString( UserNameNoDot ), DiagnosticText );
+	}
+
 }
 
-const TCHAR* FDiagnoseReportWorker::Name()
+#endif // !CRASH_REPORT_UNATTENDED_ONLY
+
+void FDiagnoseReportWorker::DoWork()
 {
-	return TEXT("FDiagnoseCrashWorker");
+	const FText ReportText = ErrorReport.DiagnoseReport();
+	DiagnosticText = FCrashReportClient::FormatDiagnosticText( ReportText, MachineId, EpicAccountId, UserNameNoDot );
 }
 
 #undef LOCTEXT_NAMESPACE
diff --git a/Engine/Source/Programs/CrashReporter/CrashReportClient/Private/CrashReportClient.h b/Engine/Source/Programs/CrashReporter/CrashReportClient/Private/CrashReportClient.h
--- a/Engine/Source/Programs/CrashReporter/CrashReportClient/Private/CrashReportClient.h
+++ b/Engine/Source/Programs/CrashReporter/CrashReportClient/Private/CrashReportClient.h
@@ -9,31 +9,51 @@
 /**
  * Helper task class to process a crash report in the background
  */
-class FDiagnoseReportWorker
+class FDiagnoseReportWorker 
 {
 public:
+	/** Diagnostics from crash report text to fill in. */
+	FText& DiagnosticText;
+
+	/** Error report files to use. */
+	const FPlatformErrorReport& ErrorReport;
+
+	/** Machine ID. */
+	const FString MachineId;
+
+	/** Epic Account ID. */
+	const FString EpicAccountId;
+
+	/** User name without dot. */
+	const FString UserNameNoDot;
+
+	/** Initialization constructor. */
+	FDiagnoseReportWorker( FText* InDiagnosticText, const FString InMachineId, const FString InEpicAccountId, const FString InUserNameNoDot, const FPlatformErrorReport* InErrorReport )
+		: DiagnosticText( *InDiagnosticText )
+		, MachineId( InMachineId )
+		, EpicAccountId( InEpicAccountId )
+		, UserNameNoDot( InUserNameNoDot )
+		, ErrorReport( *InErrorReport )
+	{}
+
 	/**
-	 * Do platform-specific work to get information about the crash
+	 * Do platform-specific work to get information about the crash.
 	 */
 	void DoWork();
 
 	/** 
-	 * Give the name for external event viewers
 	 * @return The name to display in external event viewers
 	 */
-	static const TCHAR* Name();
-
-	/** Diagnostics from crash report text to fill in */
-	FText* DiagnosticText;
-
-	/** Error report files to use */
-	const FPlatformErrorReport* ErrorReportFiles;
+	static const TCHAR* Name()
+	{
+		return TEXT( "FDiagnoseCrashWorker" );
+	}
 };
 
 #if !CRASH_REPORT_UNATTENDED_ONLY
 
 /**
- * Main implementation of the crash report client applcation
+ * Main implementation of the crash report client application
  */
 class FCrashReportClient : public TSharedFromThis<FCrashReportClient>
 {
@@ -42,7 +62,7 @@ public:
 	 * Constructor: sets up background diagnosis
 	 * @param ErrorReport Error report to upload
 	 */
-	explicit FCrashReportClient(const FPlatformErrorReport& ErrorReport, const FString& AppName);
+	FCrashReportClient( const FPlatformErrorReport& InErrorReport, const FString& AppName );
 
 	/**
 	 * Respond to the user pressing Submit
@@ -111,6 +131,9 @@ public:
 	 */
 	void RequestCloseWindow(const TSharedRef<SWindow>& Window);
 
+	/** Formats processed diagnostic text by adding additional information about machine and user. */
+	static FText FormatDiagnosticText( const FText& DiagnosticText, const FString MachineId, const FString EpicAccountId, const FString UserNameNoDot );
+
 private:
 	/**
 	 * Write the user's comment to the report and begin uploading the entire report 
@@ -162,7 +185,7 @@ private:
 	TOneShotTaskUsingDedicatedThread<FDiagnoseReportWorker> DiagnoseReportTask;
 
 	/** Platform code for accessing the report */
-	FPlatformErrorReport ErrorReportFiles;
+	FPlatformErrorReport ErrorReport;
 
 	/** Object that uploads report files to the server */
 	FCrashUpload Uploader;
diff --git a/Engine/Source/Programs/CrashReporter/CrashReportClient/Private/CrashReportUtil.h b/Engine/Source/Programs/CrashReporter/CrashReportClient/Private/CrashReportUtil.h
--- a/Engine/Source/Programs/CrashReporter/CrashReportClient/Private/CrashReportUtil.h
+++ b/Engine/Source/Programs/CrashReporter/CrashReportClient/Private/CrashReportUtil.h
@@ -22,27 +22,24 @@ public:
 		CRASHREPORTCLIENT_CHECK(!!bIsDone);
 	}
 
-	TWorker& GetTask()
+
+	template<typename T1, typename T2, typename T3, typename T4, typename T5>
+	TWorker& GetTask( T1 Arg1, T2 Arg2, T3 Arg3, T4 Arg4, T5 Arg5 )
 	{
-		return Worker;
+		Worker = new TWorker( Arg1, Arg2, Arg3, Arg4, Arg5 );
+		return *Worker;
 	}
 
 	void StartBackgroundTask()
 	{
 		bIsDone = 0;
-
-		Thread = FRunnableThread::Create(this, Worker.Name());
+		Thread = FRunnableThread::Create( this, Worker->Name() );
 	}
 
 	/**
 	 * No need to join the thread, since this can't be reused
 	 */
-	bool IsDone()
-	{
-		return bIsDone == 1;
-	}
-
-	bool IsWorkDone() const
+	bool IsDone() const
 	{
 		return bIsDone == 1;
 	}
@@ -50,12 +47,12 @@ public:
 private:
 	virtual uint32 Run() override
 	{
-		Worker.DoWork();
+		Worker->DoWork();
 		FPlatformAtomics::InterlockedIncrement(&bIsDone);
 		return 0;
 	}
 
-	TWorker Worker;
+	TAutoPtr<TWorker> Worker;
 	TScopedPointer<FRunnableThread> Thread;
 	volatile int32 bIsDone;	
 };
@@ -104,7 +101,7 @@ FunctorDirectoryVisitor<Functor> MakeDirectoryVisitor(Functor&& FunctorInstance)
 }
 
 /**
- * Create a multiline string to display from an exception and callstack
+ * Create a multi line string to display from an exception and callstack
  * @param Exception Exception description string
  * @param Assertion Assertion description string
  * @param Callstack List of callstack entry strings
