CrashReportWWW - Minor code formatting

diff --git a/Engine/Source/Programs/CrashReporter/CrashReportWebSite/Views/Crashes/Index.aspx b/Engine/Source/Programs/CrashReporter/CrashReportWebSite/Views/Crashes/Index.aspx
--- a/Engine/Source/Programs/CrashReporter/CrashReportWebSite/Views/Crashes/Index.aspx
+++ b/Engine/Source/Programs/CrashReporter/CrashReportWebSite/Views/Crashes/Index.aspx
@@ -149,7 +149,8 @@ Crash Reports
 				<div style="background-color: #E8EEF4; margin-bottom: 10px; width: 19.7em; padding-bottom:0px;">
 				<span style="background-color: #E8EEF4; font-size: medium; padding-left: 1em;"><%=Html.ActionLink(
 					"Crashes", "Index", "Crashes",
-					new {
+					new 
+					{
 						SearchQuery = Model.SearchQuery, 
 						SortTerm = Model.SortTerm, 
 						SortOrder = Model.SortOrder, 
@@ -157,10 +158,24 @@ Crash Reports
 						DateFrom = Model.DateFrom, 
 						DateTo = Model.DateTo, 
 						CrashType = Model.CrashType
-					}, new { style = "color:black; text-decoration:none;" }
+					}
+					, 
+					new { style = "color:black; text-decoration:none;" }
 				)%></span>
 				<span style="background-color: #C3CAD0; font-size: medium; margin-left: 1em; padding:0 1em;"
-					  title="<%= BuggsViewModel.Tooltip %>"><%=Html.ActionLink( "CrashGroups", "Index", "Buggs", new { SearchQuery = Model.SearchQuery, SortTerm = Model.SortTerm, SortOrder = Model.SortOrder, UserGroup = Model.UserGroup, DateFrom = Model.DateFrom, DateTo = Model.DateTo }, new { style = "color:black; text-decoration:none;" } )%></span>
+					  title="<%= BuggsViewModel.Tooltip %>"><%=Html.ActionLink( "CrashGroups", "Index", "Buggs", 
+					new 
+					{ 
+						SearchQuery = Model.SearchQuery, 
+						SortTerm = Model.SortTerm, 
+						SortOrder = Model.SortOrder, 
+						UserGroup = Model.UserGroup, 
+						DateFrom = Model.DateFrom, 
+						DateTo = Model.DateTo, 
+						CrashType = Model.CrashType
+					}
+					, 
+					new { style = "color:black; text-decoration:none;" } )%></span>
 			</div>
 			<%Html.RenderPartial( "/Views/Crashes/ViewCrash.ascx" );%>
 		</form>
@@ -168,7 +183,19 @@ Crash Reports
 </div>
 
 <div class="PaginationBox">
-	<%=Html.PageLinks( Model.PagingInfo, i => Url.Action( "", new { page = i, SearchQuery = Model.SearchQuery, SortTerm = Model.SortTerm, SortOrder = Model.SortOrder, UserGroup = Model.UserGroup, DateFrom = Model.DateFrom, DateTo = Model.DateTo, CrashType = Model.CrashType } ) )%>
+	<%=Html.PageLinks( Model.PagingInfo, i => Url.Action( "", 
+		new 
+		{ 
+			page = i, 
+			SearchQuery = Model.SearchQuery, 
+			SortTerm = Model.SortTerm, 
+			SortOrder = Model.SortOrder, 
+			UserGroup = Model.UserGroup, 
+			DateFrom = Model.DateFrom, 
+			DateTo = Model.DateTo, 
+			CrashType = Model.CrashType 
+		} 
+	) )%>
 	<div id="clear"></div>
 </div>
 
