CrashReport - Added filtering by crash type to the buggs (crash website)

diff --git a/Engine/Source/Programs/CrashReporter/CrashReportWebSite/Models/BuggRepository.cs b/Engine/Source/Programs/CrashReporter/CrashReportWebSite/Models/BuggRepository.cs
--- a/Engine/Source/Programs/CrashReporter/CrashReportWebSite/Models/BuggRepository.cs
+++ b/Engine/Source/Programs/CrashReporter/CrashReportWebSite/Models/BuggRepository.cs
@@ -367,6 +367,26 @@ namespace Tools.CrashReporter.CrashReportWebSite.Models
 				Results = Search( Results, FormData.SearchQuery );
 			}
 
+			// Filter by Crash Type
+			if( FormData.CrashType != "All" )
+			{
+				switch( FormData.CrashType )
+				{
+					case "Crashes":
+						Results = Results.Where( BuggInstance => BuggInstance.CrashType == 1 );
+						break;
+					case "Assert":
+						Results = Results.Where( BuggInstance => BuggInstance.CrashType == 2 );
+						break;
+					case "Ensure":
+						Results = Results.Where( BuggInstance => BuggInstance.CrashType == 3 );
+						break;
+					case "CrashesAsserts":
+						Results = Results.Where( BuggInstance => BuggInstance.CrashType == 1 || BuggInstance.CrashType == 2 );
+						break;
+				}
+			}
+
 			// Get UserGroup ResultCounts
 			Dictionary<string, int> GroupCounts = GetCountsByGroup( Results );
 
@@ -386,6 +406,7 @@ namespace Tools.CrashReporter.CrashReportWebSite.Models
 				SortTerm = FormData.SortTerm,
 				SortOrder = FormData.SortOrder,
 				UserGroup = FormData.UserGroup,
+				CrashType = FormData.CrashType,
 				SearchQuery = FormData.SearchQuery,
 				DateFrom = (long)(FormData.DateFrom - CrashesViewModel.Epoch).TotalMilliseconds,
 				DateTo = (long)(FormData.DateTo - CrashesViewModel.Epoch).TotalMilliseconds,
diff --git a/Engine/Source/Programs/CrashReporter/CrashReportWebSite/Models/BuggsViewModel.cs b/Engine/Source/Programs/CrashReporter/CrashReportWebSite/Models/BuggsViewModel.cs
--- a/Engine/Source/Programs/CrashReporter/CrashReportWebSite/Models/BuggsViewModel.cs
+++ b/Engine/Source/Programs/CrashReporter/CrashReportWebSite/Models/BuggsViewModel.cs
@@ -10,6 +10,7 @@ namespace Tools.CrashReporter.CrashReportWebSite.Models
 	/// <summary>
 	/// The view model for the Bugg index page.
 	/// </summary>
+	/// // @TODO yrx 2014-08-14 Add base class
 	public class BuggsViewModel
 	{
 		/// <summary>Tooltip for "Crash Groups" tab</summary>
@@ -29,6 +30,9 @@ namespace Tools.CrashReporter.CrashReportWebSite.Models
 		
 		/// <summary>The current user group name.</summary>
 		public string UserGroup { get; set; }
+
+		/// <summary>The type of report to filter by. e.g. Asserts.</summary>
+		public string CrashType { get; set; }
 		
 		/// <summary>The query that filtered the results.</summary>
 		public string SearchQuery { get; set; }
diff --git a/Engine/Source/Programs/CrashReporter/CrashReportWebSite/Models/CrashReport.cs b/Engine/Source/Programs/CrashReporter/CrashReportWebSite/Models/CrashReport.cs
--- a/Engine/Source/Programs/CrashReporter/CrashReportWebSite/Models/CrashReport.cs
+++ b/Engine/Source/Programs/CrashReporter/CrashReportWebSite/Models/CrashReport.cs
@@ -59,6 +59,27 @@ namespace Tools.CrashReporter.CrashReportWebSite.Models
 			return Crashes;
 		}
 
+		/// <summary></summary>
+		public string CrashTypeAsString
+		{
+			get
+			{
+				if( CrashType == 1 )
+				{
+					return "Crash";
+				}
+				else if( CrashType == 2 )
+				{
+					return "Assert";
+				}
+				else if( CrashType == 3 )
+				{
+					return "Ensure";
+				}
+				return "Unknown";
+			}
+		}
+
 		/// <summary>
 		/// Return the top lines of a callstack.
 		/// </summary>
@@ -212,6 +233,27 @@ namespace Tools.CrashReporter.CrashReportWebSite.Models
 			return Results;
 		}
 
+		/// <summary></summary>
+		public string CrashTypeAsString
+		{
+			get
+			{
+				if( CrashType == 0 )
+				{
+					return "Crash";
+				}
+				else if( CrashType == 1 )
+				{
+					return "Assert";
+				}
+				else if( CrashType == 2 )
+				{
+					return "Ensure";
+				}
+				return "Unknown";
+			}
+		}
+
 		/// <summary>
 		/// Gets a formatted callstack for the crash.
 		/// </summary>
diff --git a/Engine/Source/Programs/CrashReporter/CrashReportWebSite/Models/CrashReport.dbml b/Engine/Source/Programs/CrashReporter/CrashReportWebSite/Models/CrashReport.dbml
--- a/Engine/Source/Programs/CrashReporter/CrashReportWebSite/Models/CrashReport.dbml
+++ b/Engine/Source/Programs/CrashReporter/CrashReportWebSite/Models/CrashReport.dbml
@@ -8,7 +8,6 @@
       <Column Name="Summary" Type="System.String" DbType="VarChar(250)" CanBeNull="true" />
       <Column Name="Priority" Type="System.Int32" DbType="Int" CanBeNull="true" />
       <Column Name="Pattern" Type="System.String" DbType="VarChar(800) NOT NULL" CanBeNull="false" />
-      <Column Name="Type" Type="System.String" DbType="VarChar(50)" CanBeNull="true" />
       <Column Name="NumberOfCrashes" Type="System.Int32" DbType="Int" CanBeNull="true" />
       <Column Name="NumberOfUsers" Type="System.Int32" DbType="Int" CanBeNull="true" />
       <Column Name="TimeOfFirstCrash" Type="System.DateTime" DbType="DateTime" CanBeNull="true" />
@@ -19,6 +18,7 @@
       <Column Name="ReproSteps" Type="System.String" DbType="Text" CanBeNull="true" UpdateCheck="Never" />
       <Column Name="Game" Type="System.String" DbType="VarChar(120)" CanBeNull="true" />
       <Column Name="BuildVersion" Type="System.String" DbType="VarChar(128)" CanBeNull="true" />
+      <Column Name="CrashType" Type="System.Int16" DbType="SmallInt" CanBeNull="true" />
       <Association Name="Bugg_Buggs_Crash" Member="Buggs_Crashes" ThisKey="Id" OtherKey="BuggId" Type="Buggs_Crash" />
     </Type>
   </Table>
diff --git a/Engine/Source/Programs/CrashReporter/CrashReportWebSite/Models/CrashReport.dbml.layout b/Engine/Source/Programs/CrashReporter/CrashReportWebSite/Models/CrashReport.dbml.layout
--- a/Engine/Source/Programs/CrashReporter/CrashReportWebSite/Models/CrashReport.dbml.layout
+++ b/Engine/Source/Programs/CrashReporter/CrashReportWebSite/Models/CrashReport.dbml.layout
@@ -77,7 +77,7 @@
         <classShapeMoniker Id="8a60fc4a-a797-417f-87ed-25b22bd5e420" />
       </nodes>
     </associationConnector>
-    <associationConnector edgePoints="[(11.5 : 2.47199625651042); (10.875 : 2.47199625651042)]" fixedFrom="NotFixed" fixedTo="NotFixed">
+    <associationConnector edgePoints="[(11.5 : 2.47199625651042); (10.875 : 2.47199625651042)]" fixedFrom="Algorithm" fixedTo="Algorithm">
       <AssociationMoniker Name="/CrashReportDataContext/Bugg/Bugg_Buggs_Crash" />
       <nodes>
         <classShapeMoniker Id="a37f1cdb-93d7-48ef-8da7-2b7ca09282b0" />
diff --git a/Engine/Source/Programs/CrashReporter/CrashReportWebSite/Views/Buggs/Index.aspx b/Engine/Source/Programs/CrashReporter/CrashReportWebSite/Views/Buggs/Index.aspx
--- a/Engine/Source/Programs/CrashReporter/CrashReportWebSite/Views/Buggs/Index.aspx
+++ b/Engine/Source/Programs/CrashReporter/CrashReportWebSite/Views/Buggs/Index.aspx
@@ -23,7 +23,6 @@
 			.datepicker({ maxDate: '+0D' })
 			.datepicker('setDate', new Date(parseInt($('#dateTo').val())));
 
-		//$("#BuildVersionVisible").val($('#BuildVersion').val());
 	});
 
 	$.datepicker.setDefaults({
@@ -31,7 +30,6 @@
 		{
 			$("#dateFrom").val($("#dateFromVisible").datepicker('getDate').getTime());
 			$("#dateTo").val($("#dateToVisible").datepicker('getDate').getTime());
-			//$('#BuildVersion').val($("#BuildVersionVisible").val());
 			$("#FilterBuggsForm").submit();
 		}
 	});
@@ -44,6 +42,11 @@
 		//Zebrastripes
 		$("#CrashesTable tr:nth-child(even)").css("background-color", "#C3CAD0");
 		$("#CrashesTable tr:nth-child(odd)").css("background-color", "#eeeeee");
+
+		$(".CrashType").click(function ()
+		{
+			$("#FilterBuggsForm").submit();
+		});
 	});
 </script>
 </asp:Content>
@@ -56,6 +59,14 @@
 	<%=Html.Hidden( "SortTerm", Model.SortTerm )%>
 	<%=Html.Hidden( "SortOrder", Model.SortOrder )%>
 
+	<span style="float:left;">
+			<input type="radio" name="CrashType" class="CrashType" value="CrashesAsserts" <%=( Model.CrashType == "CrashesAsserts" ) ? "checked='checked'" : "" %> /> <span title='All Crashes Except Ensures'>Crashes+Asserts</span>
+			<input type="radio" name="CrashType" class="CrashType" value="Ensure" <%=( Model.CrashType == "Ensure" ) ? "checked='checked'" : "" %>/> <span title='Only Ensures'>Ensures</span>
+			<input type="radio" name="CrashType" class="CrashType" value="Assert" <%=( Model.CrashType == "Assert" ) ? "checked='checked'" : "" %>/> <span title='Only Asserts'>Asserts</span>
+			<input type="radio" name="CrashType" class="CrashType" value="Crashes" <%=( Model.CrashType == "Crashes" ) ? "checked='checked'" : "" %>/> <span title='Crashes Except Ensures and Asserts'>Crashes</span>
+			<input type="radio" name="CrashType" class="CrashType" value="All" <%=( Model.CrashType == "All" ) ? "checked='checked'" : "" %>/> <span title='All Crashes'>All</span>
+	</span>
+
 	<div id="SearchBox"><%=Html.TextBox( "SearchQuery", Model.SearchQuery, new { width = "1000" })%><input type="submit" value="Search" class='SearchButton' /></div>
 
 	<script>$.datepicker.setDefaults($.datepicker.regional['']);</script>
@@ -114,7 +125,8 @@
 								UserGroup = Model.UserGroup, 
 								DateFrom = Model.DateFrom, 
 								DateTo = Model.DateTo, 
-								BuildVersion = Model.BuildVersion 
+								BuildVersion = Model.BuildVersion,
+								CrashType = Model.CrashType
 							}
 							, 
 							new { style = "color:black; text-decoration:none;" } )%></span>
@@ -129,7 +141,8 @@
 								UserGroup = Model.UserGroup, 
 								DateFrom = Model.DateFrom, 
 								DateTo = Model.DateTo, 
-								BuildVersion = Model.BuildVersion 
+								BuildVersion = Model.BuildVersion,
+								CrashType = Model.CrashType
 							}
 							, 
 							new { style = "color:black; text-decoration:none;" } )%></span>
@@ -150,7 +163,8 @@
 				UserGroup = Model.UserGroup, 
 				DateFrom = Model.DateFrom, 
 				DateTo = Model.DateTo, 
-				BuildVersion = Model.BuildVersion 
+				BuildVersion = Model.BuildVersion,
+				CrashType = Model.CrashType
 			} 
 		) )%>
 		<div id="clear"></div>
diff --git a/Engine/Source/Programs/CrashReporter/CrashReportWebSite/Views/Buggs/Show.aspx b/Engine/Source/Programs/CrashReporter/CrashReportWebSite/Views/Buggs/Show.aspx
--- a/Engine/Source/Programs/CrashReporter/CrashReportWebSite/Views/Buggs/Show.aspx
+++ b/Engine/Source/Programs/CrashReporter/CrashReportWebSite/Views/Buggs/Show.aspx
@@ -106,6 +106,9 @@
 		<dt>Time of First Crash</dt>
 			<dd style='width:8em'><%=Model.Bugg.TimeOfFirstCrash%></dd>
 
+		<dt>Crash type</dt> 
+			<dd class='even' style='width:8em'><%=Model.Bugg.CrashTypeAsString%></dd>
+
 		<dt>Number of Users</dt>
 			<dd class='even'><%=Html.DisplayFor( m => Model.Bugg.NumberOfUsers )%></dd>
 
diff --git a/Engine/Source/Programs/CrashReporter/CrashReportWebSite/Views/Buggs/ViewBuggs.ascx b/Engine/Source/Programs/CrashReporter/CrashReportWebSite/Views/Buggs/ViewBuggs.ascx
--- a/Engine/Source/Programs/CrashReporter/CrashReportWebSite/Views/Buggs/ViewBuggs.ascx
+++ b/Engine/Source/Programs/CrashReporter/CrashReportWebSite/Views/Buggs/ViewBuggs.ascx
@@ -21,6 +21,7 @@
 					DateFrom = Model.DateFrom, 
 					DateTo = Model.DateTo,
 					BuildVersion = Model.BuildVersion,
+					CrashType = Model.CrashType
 				} 
 			) )%>
 			</div>
diff --git a/Engine/Source/Programs/CrashReporter/CrashReportWebSite/Views/Crashes/Show.aspx b/Engine/Source/Programs/CrashReporter/CrashReportWebSite/Views/Crashes/Show.aspx
--- a/Engine/Source/Programs/CrashReporter/CrashReportWebSite/Views/Crashes/Show.aspx
+++ b/Engine/Source/Programs/CrashReporter/CrashReportWebSite/Views/Crashes/Show.aspx
@@ -164,6 +164,9 @@
 				<%=Model.Crash.GetTimeOfCrash()[0] %><br />
 				<%=Model.Crash.GetTimeOfCrash()[1] %>
 
+			<dt>Crash type</dt> 
+				<dd class='even' style='width:8em'><%=Model.Crash.CrashTypeAsString%></dd>
+
 			<dt>User</dt>
 				<dd><%=Html.DisplayFor(m => Model.Crash.User.UserName) %></dd>
 
diff --git a/Engine/Source/Programs/CrashReporter/CrashReportWebSite/Views/Helpers/UrlHelperExtension.cs b/Engine/Source/Programs/CrashReporter/CrashReportWebSite/Views/Helpers/UrlHelperExtension.cs
--- a/Engine/Source/Programs/CrashReporter/CrashReportWebSite/Views/Helpers/UrlHelperExtension.cs
+++ b/Engine/Source/Programs/CrashReporter/CrashReportWebSite/Views/Helpers/UrlHelperExtension.cs
@@ -151,7 +151,8 @@ namespace Tools.CrashReporter.CrashReportWebSite.Views.Helpers
 														SearchQuery = Model.SearchQuery,
 														DateFrom = Model.DateFrom,
 														DateTo = Model.DateTo,
-														BuildVersion = Model.BuildVersion
+														BuildVersion = Model.BuildVersion,
+														CrashType = Model.CrashType,
 													} );
 
 			Tag.MergeAttribute( "href", URL );
@@ -198,7 +199,8 @@ namespace Tools.CrashReporter.CrashReportWebSite.Views.Helpers
 														SearchQuery = CallStack,
 														DateFrom = Model.DateFrom,
 														DateTo = Model.DateTo,
-														BuildVersion = Model.BuildVersion
+														BuildVersion = Model.BuildVersion,
+														CrashType = Model.CrashType,
 													} );
 
 			Tag.MergeAttribute( "href", URL );
@@ -228,7 +230,8 @@ namespace Tools.CrashReporter.CrashReportWebSite.Views.Helpers
 														SearchQuery = Model.SearchQuery,
 														DateFrom = Model.DateFrom,
 														DateTo = Model.DateTo,
-														BuildVersion = Model.BuildVersion
+														BuildVersion = Model.BuildVersion,
+														CrashType = Model.CrashType,
 													} );
 
 			Tag.MergeAttribute( "href", URL );
