CrashReportWWW - Fixed incorrectly displaying crash type

diff --git a/Engine/Source/Programs/CrashReporter/CrashReportWebSite/Models/CrashReport.cs b/Engine/Source/Programs/CrashReporter/CrashReportWebSite/Models/CrashReport.cs
--- a/Engine/Source/Programs/CrashReporter/CrashReportWebSite/Models/CrashReport.cs
+++ b/Engine/Source/Programs/CrashReporter/CrashReportWebSite/Models/CrashReport.cs
@@ -238,15 +238,15 @@ namespace Tools.CrashReporter.CrashReportWebSite.Models
 		{
 			get
 			{
-				if( CrashType == 0 )
+				if( CrashType == 1 )
 				{
 					return "Crash";
 				}
-				else if( CrashType == 1 )
+				else if( CrashType == 2 )
 				{
 					return "Assert";
 				}
-				else if( CrashType == 2 )
+				else if( CrashType == 3 )
 				{
 					return "Ensure";
 				}
