CrashReport - Minor code clean-up

diff --git a/Engine/Source/Programs/CrashReporter/CrashReportClient/Private/GenericErrorReport.cpp b/Engine/Source/Programs/CrashReporter/CrashReportClient/Private/GenericErrorReport.cpp
--- a/Engine/Source/Programs/CrashReporter/CrashReportClient/Private/GenericErrorReport.cpp
+++ b/Engine/Source/Programs/CrashReporter/CrashReportClient/Private/GenericErrorReport.cpp
@@ -78,7 +78,7 @@ TArray<FString> FGenericErrorReport::GetFilesToUpload() const
 	{
 		if (FRocketSupport::IsRocket() && Filename.EndsWith(TEXT(".log")))
 		{
-			// Temporarily side-step privacy concerns by not uploading Rocket logs
+			// Temporarily side-step privacy concerns by not uploading UE4 release builds' logs
 			continue;
 		}
 
diff --git a/Engine/Source/Programs/CrashReporter/CrashReportClient/Private/Windows/WindowsErrorReport.cpp b/Engine/Source/Programs/CrashReporter/CrashReportClient/Private/Windows/WindowsErrorReport.cpp
--- a/Engine/Source/Programs/CrashReporter/CrashReportClient/Private/Windows/WindowsErrorReport.cpp
+++ b/Engine/Source/Programs/CrashReporter/CrashReportClient/Private/Windows/WindowsErrorReport.cpp
@@ -104,7 +104,7 @@ FText FWindowsErrorReport::DiagnoseReport() const
 		}
 	}
 
-	// Don't write a Diagnostics.txt to disk in rocket. It will be displayed in the UI but not sent to the server.
+	// Don't write a Diagnostics.txt to disk in UE4 release build. It will be displayed in the UI but not sent to the server.
 	if ( !FRocketSupport::IsRocket() )
 	{
 		// There's a callstack, so write it out to save the server trying to do it
