CrashReportWWW - Fixed broken search

diff --git a/Engine/Source/Programs/CrashReporter/CrashReportWebSite/Models/CrashReport.cs b/Engine/Source/Programs/CrashReporter/CrashReportWebSite/Models/CrashReport.cs
--- a/Engine/Source/Programs/CrashReporter/CrashReportWebSite/Models/CrashReport.cs
+++ b/Engine/Source/Programs/CrashReporter/CrashReportWebSite/Models/CrashReport.cs
@@ -60,24 +60,21 @@ namespace Tools.CrashReporter.CrashReportWebSite.Models
 		}
 
 		/// <summary></summary>
-		public string CrashTypeAsString
+		public string GetCrashTypeAsString()
 		{
-			get
+			if( CrashType == 1 )
 			{
-				if( CrashType == 1 )
-				{
-					return "Crash";
-				}
-				else if( CrashType == 2 )
-				{
-					return "Assert";
-				}
-				else if( CrashType == 3 )
-				{
-					return "Ensure";
-				}
-				return "Unknown";
+				return "Crash";
+			}
+			else if( CrashType == 2 )
+			{
+				return "Assert";
+			}
+			else if( CrashType == 3 )
+			{
+				return "Ensure";
 			}
+			return "Unknown";
 		}
 
 		/// <summary>
@@ -219,7 +216,7 @@ namespace Tools.CrashReporter.CrashReportWebSite.Models
 		/// <summary>
 		/// Return a display friendly version of the time of crash.
 		/// </summary>
-		/// <returns>A pair of strings represening the date and time of the crash.</returns>
+		/// <returns>A pair of strings representing the date and time of the crash.</returns>
 		public string[] GetTimeOfCrash()
 		{
 			string[] Results = new string[2];
@@ -234,24 +231,21 @@ namespace Tools.CrashReporter.CrashReportWebSite.Models
 		}
 
 		/// <summary></summary>
-		public string CrashTypeAsString
+		public string GetCrashTypeAsString()
 		{
-			get
+			if( CrashType == 1 )
 			{
-				if( CrashType == 1 )
-				{
-					return "Crash";
-				}
-				else if( CrashType == 2 )
-				{
-					return "Assert";
-				}
-				else if( CrashType == 3 )
-				{
-					return "Ensure";
-				}
-				return "Unknown";
+				return "Crash";
+			}
+			else if( CrashType == 2 )
+			{
+				return "Assert";
+			}
+			else if( CrashType == 3 )
+			{
+				return "Ensure";
 			}
+			return "Unknown";
 		}
 
 		/// <summary>
diff --git a/Engine/Source/Programs/CrashReporter/CrashReportWebSite/Views/Buggs/Show.aspx b/Engine/Source/Programs/CrashReporter/CrashReportWebSite/Views/Buggs/Show.aspx
--- a/Engine/Source/Programs/CrashReporter/CrashReportWebSite/Views/Buggs/Show.aspx
+++ b/Engine/Source/Programs/CrashReporter/CrashReportWebSite/Views/Buggs/Show.aspx
@@ -107,7 +107,7 @@
 			<dd style='width:8em'><%=Model.Bugg.TimeOfFirstCrash%></dd>
 
 		<dt>Crash type</dt> 
-			<dd class='even' style='width:8em'><%=Model.Bugg.CrashTypeAsString%></dd>
+			<dd class='even' style='width:8em'><%=Model.Bugg.GetCrashTypeAsString()%></dd>
 
 		<dt>Number of Users</dt>
 			<dd class='even'><%=Html.DisplayFor( m => Model.Bugg.NumberOfUsers )%></dd>
diff --git a/Engine/Source/Programs/CrashReporter/CrashReportWebSite/Views/Crashes/Show.aspx b/Engine/Source/Programs/CrashReporter/CrashReportWebSite/Views/Crashes/Show.aspx
--- a/Engine/Source/Programs/CrashReporter/CrashReportWebSite/Views/Crashes/Show.aspx
+++ b/Engine/Source/Programs/CrashReporter/CrashReportWebSite/Views/Crashes/Show.aspx
@@ -165,7 +165,7 @@
 				<%=Model.Crash.GetTimeOfCrash()[1] %>
 
 			<dt>Crash type</dt> 
-				<dd class='even' style='width:8em'><%=Model.Crash.CrashTypeAsString%></dd>
+				<dd class='even' style='width:8em'><%=Model.Crash.GetCrashTypeAsString()%></dd>
 
 			<dt>User</dt>
 				<dd><%=Html.DisplayFor(m => Model.Crash.User.UserName) %></dd>
