Make the CrashReportClient separate main & game thread aware so that it works.

diff --git a/Engine/Source/Programs/CrashReporter/CrashReportClient/Private/Mac/CrashReportClientMainMac.cpp b/Engine/Source/Programs/CrashReporter/CrashReportClient/Private/Mac/CrashReportClientMainMac.cpp
--- a/Engine/Source/Programs/CrashReporter/CrashReportClient/Private/Mac/CrashReportClientMainMac.cpp
+++ b/Engine/Source/Programs/CrashReporter/CrashReportClient/Private/Mac/CrashReportClientMainMac.cpp
@@ -2,6 +2,7 @@
 
 #include "CrashReportClientApp.h"
 #include "ExceptionHandling.h"
+#include "CocoaThread.h"
 
 /**
  * Because crash reporters can crash, too
@@ -41,26 +42,43 @@ static FString GSavedCommandLine;
 //handler for the quit apple event used by the Dock menu
 - (void)handleQuitEvent:(NSAppleEventDescriptor*)Event withReplyEvent:(NSAppleEventDescriptor*)ReplyEvent
 {
-	[NSApp terminate:self];
+    [self requestQuit:self];
 }
 
-- (void)applicationDidFinishLaunching:(NSNotification *)Notification
+- (IBAction)requestQuit:(id)Sender
+{
+	GIsRequestingExit = true;
+}
+
+- (void) runGameThread:(id)Arg
 {
-	//install the custom quit event handler
-    NSAppleEventManager* appleEventManager = [NSAppleEventManager sharedAppleEventManager];
-    [appleEventManager setEventHandler:self andSelector:@selector(handleQuitEvent:withReplyEvent:) forEventClass:kCoreEventClass andEventID:kAEQuitApplication];
-	
 	FPlatformMisc::SetGracefulTerminationHandler();
-	FPlatformMisc::SetCrashHandler(CrashReporterCrashHandler);
-	
-	// Force ourselves to the front
-	[[NSApplication sharedApplication] activateIgnoringOtherApps:YES];
+	FPlatformMisc::SetCrashHandler(&CrashReporterCrashHandler);
 	
 	RunCrashReportClient(*GSavedCommandLine);
 	
 	[NSApp terminate: self];
 }
 
+- (NSApplicationTerminateReply)applicationShouldTerminate:(NSApplication *)Sender;
+{
+	if(!GIsRequestingExit)
+	{
+		[self requestQuit:self];
+		return NSTerminateCancel;
+	}
+	return NSTerminateNow;
+}
+
+- (void)applicationDidFinishLaunching:(NSNotification *)Notification
+{
+	//install the custom quit event handler
+    NSAppleEventManager* appleEventManager = [NSAppleEventManager sharedAppleEventManager];
+    [appleEventManager setEventHandler:self andSelector:@selector(handleQuitEvent:withReplyEvent:) forEventClass:kCoreEventClass andEventID:kAEQuitApplication];
+	
+	RunGameThread(self, @selector(runGameThread:));
+}
+
 @end
 
 int main(int argc, char *argv[])
