CrashReportWWW - Fixed an issue when summary was null

diff --git a/Engine/Source/Programs/CrashReporter/CrashReportWebSite/Views/Buggs/Show.aspx b/Engine/Source/Programs/CrashReporter/CrashReportWebSite/Views/Buggs/Show.aspx
--- a/Engine/Source/Programs/CrashReporter/CrashReportWebSite/Views/Buggs/Show.aspx
+++ b/Engine/Source/Programs/CrashReporter/CrashReportWebSite/Views/Buggs/Show.aspx
@@ -128,7 +128,7 @@
 
 	<div id="CallStackContainer" >
 
-		<% if( Model.Bugg.Summary.Length > 0 ) 
+		<% if( !string.IsNullOrEmpty( Model.Bugg.Summary ) ) 
 			{ %>
 				<div id='ShowErrorMessage'>
 					<br />
diff --git a/Engine/Source/Programs/CrashReporter/CrashReportWebSite/Views/Crashes/Show.aspx b/Engine/Source/Programs/CrashReporter/CrashReportWebSite/Views/Crashes/Show.aspx
--- a/Engine/Source/Programs/CrashReporter/CrashReportWebSite/Views/Crashes/Show.aspx
+++ b/Engine/Source/Programs/CrashReporter/CrashReportWebSite/Views/Crashes/Show.aspx
@@ -212,7 +212,7 @@
 	</div>
 
 	<div id="CallStackContainer" >
-		<% if( Model.Crash.Summary.Length > 0 ) 
+		<% if( !string.IsNullOrEmpty( Model.Crash.Summary ) ) 
 			{ %>
 				<div id='ShowErrorMessage'>
 					<br />
