CrashReport - Added analytics to the crash report client

diff --git a/Engine/Source/Programs/CrashReporter/CrashReportClient/CrashReportClient.Build.cs b/Engine/Source/Programs/CrashReporter/CrashReportClient/CrashReportClient.Build.cs
--- a/Engine/Source/Programs/CrashReporter/CrashReportClient/CrashReportClient.Build.cs
+++ b/Engine/Source/Programs/CrashReporter/CrashReportClient/CrashReportClient.Build.cs
@@ -6,7 +6,14 @@ public class CrashReportClient : ModuleRules
 {
 	public CrashReportClient(TargetInfo Target)
 	{
-		PublicIncludePaths.Add("Runtime/Launch/Public");
+		PublicIncludePaths.AddRange
+		(
+			new string[] 
+			{ 
+				"Runtime/Launch/Public",
+				"Programs/CrashReporter/CrashReportClient/Private",
+			}
+		);
 
 		PrivateDependencyModuleNames.AddRange(
 			new string[] 
@@ -16,6 +23,8 @@ public class CrashReportClient : ModuleRules
 				"HTTP",
 				"Projects",
 				"XmlParser",
+				"Analytics",
+				"AnalyticsET",
 			}
 			);
 
diff --git a/Engine/Source/Programs/CrashReporter/CrashReportClient/Private/CrashDescription.cpp b/Engine/Source/Programs/CrashReporter/CrashReportClient/Private/CrashDescription.cpp
new file mode 100644
--- /dev/null
+++ b/Engine/Source/Programs/CrashReporter/CrashReportClient/Private/CrashDescription.cpp
@@ -0,0 +1,235 @@
+// Copyright 1998-2014 Epic Games, Inc. All Rights Reserved.
+
+#include "CrashReportClientApp.h"
+
+#include "EngineVersion.h"
+#include "XmlFile.h"
+#include "CrashDescription.h"
+#include "CrashReportAnalytics.h"
+#include "Runtime/Analytics/Analytics/Public/Analytics.h"
+#include "Runtime/Analytics/Analytics/Public/Interfaces/IAnalyticsProvider.h"
+#include "DesktopPlatformModule.h"
+
+
+
+FCrashDescription::FCrashDescription() :
+	BuiltFromCL( -1 ),
+	TimeOfCrash( FDateTime::UtcNow() ),
+	bHasMiniDumpFile( false ),
+	bHasLogFile( false ),
+	bHasVideoFile( false ),
+	bHasCompleteData( false )
+{
+	InitializeIDs();
+}
+
+
+FCrashDescription::FCrashDescription( FString WERXMLFilepath ) :
+	BuiltFromCL( -1 ),
+	TimeOfCrash( FDateTime::UtcNow() ),
+	bHasMiniDumpFile( false ),
+	bHasLogFile( false ),
+	bHasVideoFile( false ),
+	bHasCompleteData( false )
+{
+	InitializeIDs();
+
+	// This is for the current system that uses files from Windows Error Reporting.
+	// Will be replaced with the unified version soon.
+	FXmlFile XmlFile( WERXMLFilepath );
+
+	ReportName = FPaths::GetPath( WERXMLFilepath );
+	FPaths::NormalizeDirectoryName( ReportName );
+	// Grab the last component...
+	ReportName = FPaths::GetCleanFilename( ReportName );
+
+	const bool bIsValid = XmlFile.IsValid();
+	if( bIsValid )
+	{
+		const FXmlNode* OSVersionInformationNode = XmlFile.GetRootNode()->FindChildNode( TEXT( "OSVersionInformation" ) );
+		const FXmlNode* ProblemSignaturesNode = XmlFile.GetRootNode()->FindChildNode( TEXT( "ProblemSignatures" ) );
+		const FXmlNode* DynamicSignaturesNode = XmlFile.GetRootNode()->FindChildNode( TEXT( "DynamicSignatures" ) );
+
+		FString Product;
+		int EngineVersionComponents = 0;
+
+		if( OSVersionInformationNode )
+		{
+			const FXmlNode* ProductNode = OSVersionInformationNode->FindChildNode( TEXT( "Product" ) );
+			if( ProductNode )
+			{
+				Product = ProductNode->GetContent();
+			}
+		}
+
+		if( ProblemSignaturesNode )
+		{
+			const FXmlNode* GameNameNode = ProblemSignaturesNode->FindChildNode( TEXT( "Parameter0" ) );
+			if( GameNameNode )
+			{
+				GameName = GameNameNode->GetContent();
+			}
+
+			const FXmlNode* BuildVersionNode = ProblemSignaturesNode->FindChildNode( TEXT( "Parameter1" ) );
+			if( BuildVersionNode )
+			{
+				BuildVersion = BuildVersionNode->GetContent();
+				EngineVersionComponents++;
+			}
+
+			const FXmlNode* Parameter9Node = ProblemSignaturesNode->FindChildNode( TEXT( "Parameter9" ) );
+			if( Parameter9Node )
+			{
+				const FString Parameter9Value = Parameter9Node->GetContent();
+
+				TArray<FString> ParsedParameters9;
+				Parameter9Value.ParseIntoArray( &ParsedParameters9, TEXT( "!" ), false );
+
+				if( ParsedParameters9.Num() > 0 )
+				{
+					BranchName = ParsedParameters9[0].Replace( TEXT( "+" ), TEXT( "/" ) );
+
+					const FString DepotRoot = TEXT( "//depot/" );
+					if( BranchName.StartsWith( DepotRoot ) )
+					{
+						BranchName = BranchName.Mid( DepotRoot.Len() );
+					}
+					EngineVersionComponents++;
+				}
+
+				if( ParsedParameters9.Num() > 1 )
+				{
+					const FString BaseDir = ParsedParameters9[1];
+
+					TArray<FString> SubDirs;
+					BaseDir.ParseIntoArray( &SubDirs, TEXT( "/" ), true );
+					const int SubDirsNum = SubDirs.Num();
+					const FString PlatformName = SubDirsNum > 0 ? SubDirs[SubDirsNum - 1] : TEXT("");
+					if( Product.Len() > 0 )
+					{
+						Platform = FString::Printf( TEXT( "%s [%s]" ), *PlatformName, *Product );
+					}
+					else
+					{
+						Platform = PlatformName;
+					}
+				}
+
+				if( ParsedParameters9.Num() > 2 )
+				{
+					EngineMode = ParsedParameters9[2];
+				}
+
+				if( ParsedParameters9.Num() > 3 )
+				{
+					TTypeFromString<uint32>::FromString( BuiltFromCL, *ParsedParameters9[3] );
+					EngineVersionComponents++;
+				}
+			}
+		}
+
+		if( DynamicSignaturesNode )
+		{
+			const FXmlNode* LanguageNode = DynamicSignaturesNode->FindChildNode( TEXT( "Parameter2" ) );
+			if( LanguageNode )
+			{
+				LanguageLCID = LanguageNode->GetContent();
+			}
+		}
+
+		// We have all three components of the engine version, so initialize it.
+		if( EngineVersionComponents == 3 )
+		{
+			InitializeEngineVersion();
+		}
+
+		bHasCompleteData = true;
+	}
+	else
+	{
+		bHasCompleteData = false;
+	}
+}
+
+void FCrashDescription::InitializeEngineVersion()
+{
+	uint16 Major = 0;
+	uint16 Minor = 0;
+	uint16 Patch = 0;
+
+	TArray<FString> ParsedBuildVersion;
+	BuildVersion.ParseIntoArray( &ParsedBuildVersion, TEXT( "." ), false );
+
+	if( ParsedBuildVersion.Num() >= 3 )
+	{
+		TTypeFromString<uint16>::FromString( Patch, *ParsedBuildVersion[2] );
+	}
+
+	if( ParsedBuildVersion.Num() >= 2 )
+	{
+		TTypeFromString<uint16>::FromString( Minor, *ParsedBuildVersion[1] );
+	}
+
+	if( ParsedBuildVersion.Num() >= 1 )
+	{
+		TTypeFromString<uint16>::FromString( Major, *ParsedBuildVersion[0] );
+	}
+
+	EngineVersion = FEngineVersion( Major, Minor, Patch, BuiltFromCL, BranchName );
+}
+
+void FCrashDescription::InitializeIDs()
+{
+	IDesktopPlatform* DesktopPlatform = FDesktopPlatformModule::Get();
+
+	MachineId = DesktopPlatform->GetMachineId().ToString( EGuidFormats::Digits );
+
+	// The Epic ID can be looked up from this ID.
+	EpicAccountId = DesktopPlatform->GetEpicAccountId();
+
+	// Get the user name only for non-UE4 releases.
+	if( !FRocketSupport::IsRocket() )
+	{
+		// Remove periods from internal user names to match AutoReporter user names
+		// The name prefix is read by CrashRepository.AddNewCrash in the website code
+		UserName = FString( FPlatformProcess::UserName() ).Replace( TEXT( "." ), TEXT( "" ) );
+	}
+}
+
+void FCrashDescription::SendAnalytics()
+{
+	// Connect the crash report client analytics provider.
+	FCrashReportAnalytics::Initialize();
+
+	IAnalyticsProvider& Analytics = FCrashReportAnalytics::GetProvider();
+	//Analytics.SetUserID( MachineId );
+
+	TArray<FAnalyticsEventAttribute> CrashAttributes;
+
+	CrashAttributes.Add( FAnalyticsEventAttribute( TEXT( "bHasCompleteData" ), bHasCompleteData ) );
+	CrashAttributes.Add( FAnalyticsEventAttribute( TEXT( "CrashDescriptionVersion" ), (int32)CrashDescriptionVersion ) );
+	CrashAttributes.Add( FAnalyticsEventAttribute( TEXT( "ReportName" ), ReportName ) );
+
+	//	AppID = GameName
+	CrashAttributes.Add( FAnalyticsEventAttribute( TEXT( "GameName" ), GameName ) );
+
+	//	AppVersion = EngineVersion
+	CrashAttributes.Add( FAnalyticsEventAttribute( TEXT( "EngineVersion" ), EngineVersion.ToString() ) );
+	CrashAttributes.Add( FAnalyticsEventAttribute( TEXT( "BuildVersion" ), BuildVersion ) );
+	CrashAttributes.Add( FAnalyticsEventAttribute( TEXT( "BuiltFromCL" ), BuiltFromCL ) );
+	CrashAttributes.Add( FAnalyticsEventAttribute( TEXT( "BranchName" ), BranchName ) );
+
+	CrashAttributes.Add( FAnalyticsEventAttribute( TEXT( "MachineID" ), MachineId ) );
+	CrashAttributes.Add( FAnalyticsEventAttribute( TEXT( "UserName" ), UserName ) );
+	CrashAttributes.Add( FAnalyticsEventAttribute( TEXT( "EpicAccountId" ), EpicAccountId ) );
+
+	CrashAttributes.Add( FAnalyticsEventAttribute( TEXT( "Platform" ), Platform ) );
+	CrashAttributes.Add( FAnalyticsEventAttribute( TEXT( "TimeOfCrash" ), TimeOfCrash.Ticks ) );
+	CrashAttributes.Add( FAnalyticsEventAttribute( TEXT( "EngineMode" ), EngineMode ) );
+	CrashAttributes.Add( FAnalyticsEventAttribute( TEXT( "LanguageLCID" ), LanguageLCID ) );
+
+	Analytics.RecordEvent( TEXT( "CrashReportClient.ReportCrash" ), CrashAttributes );
+
+	// Shutdown analytics.
+	FCrashReportAnalytics::Shutdown();
+}
diff --git a/Engine/Source/Programs/CrashReporter/CrashReportClient/Private/CrashDescription.h b/Engine/Source/Programs/CrashReporter/CrashReportClient/Private/CrashDescription.h
new file mode 100644
--- /dev/null
+++ b/Engine/Source/Programs/CrashReporter/CrashReportClient/Private/CrashDescription.h
@@ -0,0 +1,178 @@
+// Copyright 1998-2014 Epic Games, Inc. All Rights Reserved.
+
+#pragma once
+
+/** Enumerates crash description versions. */
+enum class ECrashDescVersions
+{
+	/** Introduces a new crash description format. */
+	VER_1_NewCrashFormat,
+};
+
+/**
+ *	Describes a unified crash, should be used by all platforms.
+ *	Still lacks some of the properties, they will be added later.
+ *	Must contain the same properties as ...\CrashReportServer\CrashReportCommon\CrashDescription.cs.
+ *	Contains all usable information about the crash. 
+ *	
+ */
+
+// @TODO yrx 2014-08-27 add samples data
+struct FCrashDescription
+{
+	/** 
+	 * Version of this crash description. 
+	 */
+	static const ECrashDescVersions CrashDescriptionVersion = ECrashDescVersions::VER_1_NewCrashFormat;
+
+	/** An unique report name that this crash belongs to. */
+	FString ReportName;
+
+	/**
+	 * The name of the game that crashed. (AppID)
+	 * @GameName	varchar(64)
+	 */
+	FString GameName;
+
+	/**
+	 * The mode the game was in e.g. editor.
+	 * @EngineMode	varchar(64)	
+	 */
+	FString EngineMode;
+
+	/**
+	 * The platform that crashed e.g. Win64.
+	 * @PlatformName	varchar(64)	
+	 */
+	FString Platform;
+
+	/** 
+	 * Encoded engine version. (AppVersion)
+	 * E.g. 4.3.0.0-2215663+UE4-Releases+4.3
+	 * BuildVersion-BuiltFromCL-BranchName
+	 * @EngineVersion	varchar(64)	
+	 */
+	FEngineVersion EngineVersion;
+
+	/**
+	 * The three component version of the app e.g. 4.4.1
+	 * @BuildVersion	varchar(128)	
+	 */
+	FString BuildVersion;
+
+	/**
+	 * Built from changelist.
+	 * @ChangeListVersion	varchar(64)	
+	 */
+	uint32 BuiltFromCL;
+
+	/**
+	 * The name of the branch this game was built out of.
+	 * @Branch varchar(32) 
+	 */
+	FString BranchName;
+
+	/**
+	 * The command line of the application that crashed.
+	 * @CommandLine varchar(512) 
+	 */
+	FString CommandLine;
+
+	/**
+	 * The base directory where the app was running.
+	 * @BaseDir varchar(512) 
+	 */
+	FString BaseDir;
+
+	/**
+	 * The language ID the application that crashed.
+	 * @LanguageExt varchar(64) 
+	 */
+	FString LanguageLCID;
+
+	/** 
+	 * The name of the user that caused this crash.
+	 * @UserName varchar(64)  
+	 */
+	FString UserName;
+
+	/**
+	 * The unique ID used to identify the machine the crash occurred on.
+	 * @ComputerName varchar(64) 
+	 */
+	FString MachineId;
+
+	/** 
+	 * The Epic account ID for the user who last used the Launcher.
+	 * @EpicAccountId	varchar(64)	
+	 */
+	FString EpicAccountId;
+
+	/**
+	 * An array of FStrings representing the callstack of the crash.
+	 * @RawCallStack	varchar(MAX)	
+	 */
+	TArray<FString> CallStack;
+
+	/**
+	 * An array of FStrings showing the source code around the crash.
+	 * @SourceContext varchar(max) 
+	 */
+	TArray<FString> SourceContext;
+
+	/**
+	 * An array of FStrings representing the user description of the crash.
+	 * @Description	varchar(512)	
+	 */
+	TArray<FString> UserDescription;
+
+	/**
+	 * The error message, can be assertion message, ensure message or message from the fatal error.
+	 * @Summary	varchar(512)	
+	 */
+	TArray<FString> ErrorMessage;
+
+	/**
+	 * The UTC time the crash occurred.
+	 * @TimeOfCrash	datetime	
+	 */
+	FDateTime TimeOfCrash;
+
+	/**
+	 * Whether this crash has a minidump file.
+	 * @HasMiniDumpFile bit 
+	 */
+	bool bHasMiniDumpFile;
+
+	/**
+	 * Whether this crash has a log file.
+	 * @HasLogFile bit 
+	 */
+	bool bHasLogFile;
+
+	/**
+	 * Whether this crash has a video file.
+	 * @HasVideoFile bit 
+	 */
+	bool bHasVideoFile;
+
+	/** Whether this crash contains complete data, otherwise only IDs are available. */
+	bool bHasCompleteData;
+
+	/**
+	 * A default constructor.
+	 */
+	FCrashDescription();
+
+	/** Initializes instance based on specified WER XML string. */
+	explicit FCrashDescription( FString WERXMLFilepath );
+
+	/** Initializes engine version from the separate components. */
+	void InitializeEngineVersion();
+
+	/** Initializes following properties: UserNAme, MachineID and EpicAccountID. */
+	void InitializeIDs();
+
+	/** Sends this crash for analytics. */
+	void SendAnalytics();
+};
\ No newline at end of file
diff --git a/Engine/Source/Programs/CrashReporter/CrashReportClient/Private/CrashReportAnalytics.cpp b/Engine/Source/Programs/CrashReporter/CrashReportClient/Private/CrashReportAnalytics.cpp
new file mode 100644
--- /dev/null
+++ b/Engine/Source/Programs/CrashReporter/CrashReportClient/Private/CrashReportAnalytics.cpp
@@ -0,0 +1,119 @@
+// Copyright 1998-2014 Epic Games, Inc. All Rights Reserved.
+
+#include "CrashReportClientApp.h"
+
+#include "CrashReportAnalytics.h"
+#include "EngineBuildSettings.h"
+#include "Runtime/Analytics/Analytics/Public/Analytics.h"
+#include "Runtime/Analytics/Analytics/Public/Interfaces/IAnalyticsProvider.h"
+#include "DesktopPlatformModule.h"
+
+// WARNING! Copied from EngineAnalytics.cpp
+
+bool FCrashReportAnalytics::bIsInitialized;
+TSharedPtr<IAnalyticsProvider> FCrashReportAnalytics::Analytics;
+
+/**
+ * Engine analytics config log to initialize the engine analytics provider.
+ * External code should bind this delegate if engine analytics are desired,
+ * preferably in private code that won't be redistributed.
+ */
+FAnalytics::FProviderConfigurationDelegate& GetCrashReportAnalyticsOverrideConfigDelegate()
+{
+	static FAnalytics::FProviderConfigurationDelegate Delegate;
+	return Delegate;
+}
+
+
+/**
+ * On-demand construction of the singleton. 
+ */
+IAnalyticsProvider& FCrashReportAnalytics::GetProvider()
+{
+	checkf(bIsInitialized && Analytics.IsValid(), TEXT("FCrashReportAnalytics::GetProvider called outside of Initialize/Shutdown."));
+	return *Analytics.Get();
+}
+ 
+void FCrashReportAnalytics::Initialize()
+{
+	checkf(!bIsInitialized, TEXT("FCrashReportAnalytics::Initialize called more than once."));
+
+	// Never use analytics when running a commandlet tool
+	const bool bShouldInitAnalytics = !IsRunningCommandlet();
+	if( bShouldInitAnalytics )
+	{
+		// Setup some default engine analytics if there is nothing custom bound
+		FAnalytics::FProviderConfigurationDelegate DefaultEngineAnalyticsConfig;
+		DefaultEngineAnalyticsConfig.BindStatic(
+			[]( const FString& KeyName, bool bIsValueRequired ) -> FString
+		{
+			static TMap<FString, FString> ConfigMap;
+			if( ConfigMap.Num() == 0 )
+			{
+				ConfigMap.Add( TEXT( "ProviderModuleName" ), TEXT( "AnalyticsET" ) );
+				ConfigMap.Add( TEXT( "APIServerET" ), TEXT( "http://etsource.epicgames.com/ET2/" ) );
+
+				// We always use the "Release" analytics account unless we're running in analytics test mode (usually with
+				// a command-line parameter), or we're an internal Epic build
+				bool bUseReleaseAccount =
+					(FAnalytics::Get().GetBuildType() == FAnalytics::Development ||
+					FAnalytics::Get().GetBuildType() == FAnalytics::Release) &&
+					!FEngineBuildSettings::IsInternalBuild();	// Internal Epic build
+
+				if( FRocketSupport::IsRocket() )
+				{
+					const TCHAR* DevelopmentAccountAPIKeyET = TEXT( "Rocket.Dev" );
+					const TCHAR* ReleaseAccountAPIKeyET = TEXT( "Rocket.Release" );
+					ConfigMap.Add( TEXT( "APIKeyET" ), bUseReleaseAccount ? ReleaseAccountAPIKeyET : DevelopmentAccountAPIKeyET );
+				}
+				else if( FEngineBuildSettings::IsPerforceBuild() )
+				{
+					const TCHAR* DevelopmentAccountAPIKeyET = TEXT( "Perforce.Dev" );
+					const TCHAR* ReleaseAccountAPIKeyET = TEXT( "Perforce.Release" );
+					ConfigMap.Add( TEXT( "APIKeyET" ), bUseReleaseAccount ? ReleaseAccountAPIKeyET : DevelopmentAccountAPIKeyET );
+				}
+				else
+				{
+					const TCHAR* DevelopmentAccountAPIKeyET = TEXT( "UnrealEngine.Dev" );
+					const TCHAR* ReleaseAccountAPIKeyET = TEXT( "UnrealEngine.Release" );
+					ConfigMap.Add( TEXT( "APIKeyET" ), bUseReleaseAccount ? ReleaseAccountAPIKeyET : DevelopmentAccountAPIKeyET );
+				}
+			}
+
+			// Check for overrides
+			if( GetCrashReportAnalyticsOverrideConfigDelegate().IsBound() )
+			{
+				const FString OverrideValue = GetCrashReportAnalyticsOverrideConfigDelegate().Execute( KeyName, bIsValueRequired );
+				if( !OverrideValue.IsEmpty() )
+				{
+					return OverrideValue;
+				}
+			}
+
+			FString* ConfigValue = ConfigMap.Find( KeyName );
+			return ConfigValue != NULL ? *ConfigValue : TEXT( "" );
+		} );
+
+		// Connect the engine analytics provider (if there is a configuration delegate installed)
+		Analytics = FAnalytics::Get().CreateAnalyticsProvider(
+			FName( *DefaultEngineAnalyticsConfig.Execute( TEXT( "ProviderModuleName" ), true ) ),
+			DefaultEngineAnalyticsConfig );
+		if( Analytics.IsValid() )
+		{
+			//IDesktopPlatform* DesktopPlatform = FDesktopPlatformModule::Get();
+			//const FString MachineId = DesktopPlatform->GetMachineId().ToString( EGuidFormats::Digits );
+			//Analytics->SetUserID( MachineId );
+			Analytics->SetUserID(FPlatformMisc::GetUniqueDeviceId());
+			Analytics->StartSession();
+		}
+	}
+	bIsInitialized = true;
+}
+
+
+void FCrashReportAnalytics::Shutdown()
+{
+	checkf(bIsInitialized, TEXT("FCrashReportAnalytics::Shutdown called outside of Initialize."));
+	Analytics.Reset();
+	bIsInitialized = false;
+}
diff --git a/Engine/Source/Programs/CrashReporter/CrashReportClient/Private/CrashReportAnalytics.h b/Engine/Source/Programs/CrashReporter/CrashReportClient/Private/CrashReportAnalytics.h
new file mode 100644
--- /dev/null
+++ b/Engine/Source/Programs/CrashReporter/CrashReportClient/Private/CrashReportAnalytics.h
@@ -0,0 +1,38 @@
+// Copyright 1998-2014 Epic Games, Inc. All Rights Reserved.
+
+#pragma once
+
+// WARNING! Copied from EngineAnalytics.h
+
+class IAnalyticsProvider;
+
+/**
+ * The public interface for the engine's analytics provider singleton.
+ * For Epic builds, this will point to epic's internal analytics provider.
+ * For licensee builds, it will be NULL by default unless they provide their own
+ * configuration.
+ * 
+ */
+class FCrashReportAnalytics : FNoncopyable
+{
+public:
+	/**
+	 * Return the provider instance. Not valid outside of Initialize/Shutdown calls.
+	 * Note: must check IsAvailable() first else this code will assert if the provider is not valid.
+	 */
+	static IAnalyticsProvider& GetProvider();
+	/** Helper function to determine if the provider is valid. */
+	static bool IsAvailable() 
+	{ 
+		return Analytics.IsValid(); 
+	}
+	/** Called to initialize the singleton. */
+	static void Initialize();
+	/** Called to shut down the singleton */
+	static void Shutdown();
+
+private:
+	static bool bIsInitialized;
+	static TSharedPtr<IAnalyticsProvider> Analytics;
+};
+
diff --git a/Engine/Source/Programs/CrashReporter/CrashReportClient/Private/CrashReportClient.cpp b/Engine/Source/Programs/CrashReporter/CrashReportClient/Private/CrashReportClient.cpp
--- a/Engine/Source/Programs/CrashReporter/CrashReportClient/Private/CrashReportClient.cpp
+++ b/Engine/Source/Programs/CrashReporter/CrashReportClient/Private/CrashReportClient.cpp
@@ -1,6 +1,9 @@
 // Copyright 1998-2014 Epic Games, Inc. All Rights Reserved.
 
 #include "CrashReportClientApp.h"
+#include "CrashReportUtil.h"
+
+#include "CrashDescription.h"
 
 #include "CrashReportClient.h"
 #include "UniquePtr.h"
@@ -16,47 +19,32 @@
 
 // Must match filename specified in RunMinidumpDiagnostics
 const TCHAR* GDiagnosticsFilename = TEXT("Diagnostics.txt");
-FString GCrashUserId;
+
+
+FCrashDescription& GetCrashDescription()
+{
+	static FCrashDescription Singleton;
+	return Singleton;
+}
 
 #if !CRASH_REPORT_UNATTENDED_ONLY
 
-FCrashReportClient::FCrashReportClient(const FPlatformErrorReport& InErrorReport, const FString& AppName)
+FCrashReportClient::FCrashReportClient(const FPlatformErrorReport& InErrorReport)
 	: AppState(EApplicationState::Ready)
 	, SubmittedCountdown(-1)
 	, bDiagnosticFileSent(false)
 	, ErrorReport( InErrorReport )
 	, Uploader(GServerIP)
-	, CrashedAppName(AppName)
 	, CancelButtonText(LOCTEXT("Cancel", "Don't Send"))
 {
-	IDesktopPlatform* DesktopPlatform = FDesktopPlatformModule::Get();
-	const FString MachineId = DesktopPlatform->GetMachineId().ToString( EGuidFormats::Digits );
-
-	// The Epic ID can be looked up from this ID.
-	const FString EpicAccountId = DesktopPlatform->GetEpicAccountId();
-
-	// Remove periods from internal user names to match AutoReporter user names
-	// The name prefix is read by CrashRepository.AddNewCrash in the website code
-	const FString UserNameNoDot = FString( FPlatformProcess::UserName() ).Replace( TEXT( "." ), TEXT( "" ) );
-
-	// Set global user name ID: will be added to the report
-	if (FRocketSupport::IsRocket())
-	{
-		GCrashUserId = FString::Printf( TEXT( "!MachineId:%s!EpicAccountId:%s" ), *MachineId, *EpicAccountId );
-	}
-	else
-	{
-		GCrashUserId = FString::Printf( TEXT( "!MachineId:%s!Name:%s" ), *MachineId, *UserNameNoDot );
-	}
-
 	if (!ErrorReport.TryReadDiagnosticsFile(DiagnosticText) && !FParse::Param(FCommandLine::Get(), TEXT("no-local-diagnosis")))
 	{
-		FDiagnoseReportWorker& Worker = DiagnoseReportTask.GetTask( &DiagnosticText, MachineId, EpicAccountId, UserNameNoDot, &ErrorReport );
+		FDiagnoseReportWorker& Worker = DiagnoseReportTask.GetTask( &DiagnosticText, GetCrashDescription().MachineId, GetCrashDescription().EpicAccountId, GetCrashDescription().UserName, &ErrorReport );
 		DiagnoseReportTask.StartBackgroundTask();
 	}
 	else if( !DiagnosticText.IsEmpty() )
 	{
-		DiagnosticText = FCrashReportClient::FormatDiagnosticText( DiagnosticText, MachineId, EpicAccountId, UserNameNoDot );
+		DiagnosticText = FCrashReportUtil::FormatDiagnosticText( DiagnosticText, GetCrashDescription().MachineId, GetCrashDescription().EpicAccountId, GetCrashDescription().UserName );
 	}
 }
 
@@ -112,11 +100,6 @@ EVisibility FCrashReportClient::SubmitButtonVisibility() const
 	return AppState == EApplicationState::Ready ? EVisibility::Visible : EVisibility::Hidden;
 }
 
-FString FCrashReportClient::GetCrashedAppName() const
-{
-	return CrashedAppName;
-}
-
 void FCrashReportClient::UserCommentChanged(const FText& Comment, ETextCommit::Type CommitType)
 {
 	UserComment = Comment;
@@ -148,7 +131,8 @@ void FCrashReportClient::StartUIWillCloseTicker()
 void FCrashReportClient::StoreCommentAndUpload()
 {
 	// Call upload even if the report is empty: pending reports will be sent if any
-	ErrorReport.SetUserComment(UserComment.IsEmpty() ? LOCTEXT("NoComment", "No comment provided") : UserComment);
+	ErrorReport.SetUserComment(UserComment);
+
 	Uploader.BeginUpload(ErrorReport);
 
 	SubmittedCountdown = 5;
@@ -204,7 +188,20 @@ bool FCrashReportClient::UIWillCloseTick(float UnusedDeltaTime)
 	return false;
 }
 
-FText FCrashReportClient::FormatDiagnosticText( const FText& DiagnosticText, const FString MachineId, const FString EpicAccountId, const FString UserNameNoDot )
+FString FCrashReportClient::GetCrashedAppName() const
+{
+	return GetCrashDescription().GameName;
+}
+
+#endif // !CRASH_REPORT_UNATTENDED_ONLY
+
+void FDiagnoseReportWorker::DoWork()
+{
+	const FText ReportText = ErrorReport.DiagnoseReport();
+	DiagnosticText = FCrashReportUtil::FormatDiagnosticText( ReportText, MachineId, EpicAccountId, UserNameNoDot );
+}
+
+FText FCrashReportUtil::FormatDiagnosticText( const FText& DiagnosticText, const FString MachineId, const FString EpicAccountId, const FString UserNameNoDot )
 {
 	if( FRocketSupport::IsRocket() )
 	{
@@ -217,12 +214,4 @@ FText FCrashReportClient::FormatDiagnosticText( const FText& DiagnosticText, con
 
 }
 
-#endif // !CRASH_REPORT_UNATTENDED_ONLY
-
-void FDiagnoseReportWorker::DoWork()
-{
-	const FText ReportText = ErrorReport.DiagnoseReport();
-	DiagnosticText = FCrashReportClient::FormatDiagnosticText( ReportText, MachineId, EpicAccountId, UserNameNoDot );
-}
-
 #undef LOCTEXT_NAMESPACE
diff --git a/Engine/Source/Programs/CrashReporter/CrashReportClient/Private/CrashReportClient.h b/Engine/Source/Programs/CrashReporter/CrashReportClient/Private/CrashReportClient.h
--- a/Engine/Source/Programs/CrashReporter/CrashReportClient/Private/CrashReportClient.h
+++ b/Engine/Source/Programs/CrashReporter/CrashReportClient/Private/CrashReportClient.h
@@ -62,7 +62,7 @@ public:
 	 * Constructor: sets up background diagnosis
 	 * @param ErrorReport Error report to upload
 	 */
-	FCrashReportClient( const FPlatformErrorReport& InErrorReport, const FString& AppName );
+	FCrashReportClient( const FPlatformErrorReport& InErrorReport );
 
 	/**
 	 * Respond to the user pressing Submit
@@ -131,9 +131,6 @@ public:
 	 */
 	void RequestCloseWindow(const TSharedRef<SWindow>& Window);
 
-	/** Formats processed diagnostic text by adding additional information about machine and user. */
-	static FText FormatDiagnosticText( const FText& DiagnosticText, const FString MachineId, const FString EpicAccountId, const FString UserNameNoDot );
-
 private:
 	/**
 	 * Write the user's comment to the report and begin uploading the entire report 
@@ -163,9 +160,6 @@ private:
 		};
 	};
 
-	/** Complete path of report directory */
-	FString ReportDirectory;
-
 	/** What the app is currently doing */
 	EApplicationState::Type AppState;
 
@@ -190,9 +184,6 @@ private:
 	/** Object that uploads report files to the server */
 	FCrashUpload Uploader;
 
-	/** Name of the executable that crashed */
-	FString CrashedAppName;
-
 	/** Text for Cancel/Close button, including count-down value */
 	FText CancelButtonText;
 };
diff --git a/Engine/Source/Programs/CrashReporter/CrashReportClient/Private/CrashReportClientApp.cpp b/Engine/Source/Programs/CrashReporter/CrashReportClient/Private/CrashReportClientApp.cpp
--- a/Engine/Source/Programs/CrashReporter/CrashReportClient/Private/CrashReportClientApp.cpp
+++ b/Engine/Source/Programs/CrashReporter/CrashReportClient/Private/CrashReportClientApp.cpp
@@ -1,18 +1,25 @@
 // Copyright 1998-2014 Epic Games, Inc. All Rights Reserved.
 
 #include "CrashReportClientApp.h"
+#include "CrashDescription.h"
+#include "CrashReportAnalytics.h"
+
 #if !CRASH_REPORT_UNATTENDED_ONLY
 	#include "SCrashReportClient.h"
 	#include "CrashReportClient.h"
 	#include "CrashReportClientStyle.h"
 	#include "SlateReflector.h"
 #endif // !CRASH_REPORT_UNATTENDED_ONLY
+
 #include "CrashReportClientUnattended.h"
 #include "RequiredProgramMainCPPInclude.h"
 #include "AsyncWork.h"
 
 #include "MainLoopTiming.h"
 
+#include "PlatformErrorReport.h"
+#include "XmlFile.h"
+
 /** Default main window size */
 const FVector2D InitialWindowDimensions(640, 560);
 
@@ -23,23 +30,20 @@ const float IdealTickRate = 30.f;
 const bool RunWidgetReflector = false;
 
 IMPLEMENT_APPLICATION(CrashReportClient, "CrashReportClient");
-
 DEFINE_LOG_CATEGORY(CrashReportClientLog);
 
-namespace
-{
-	/** Directory containing the report */
-	FString ReportDirectoryAbsolutePath;
+/** Directory containing the report */
+static FString ReportDirectoryAbsolutePath;
+
+/** Name of the game passed via the command line. */
+static FString GameNameFromCmd;
 
-	/** Name of crashed app to display */
-	FString CrashedAppName;
-}
 /**
  * Upload the crash report with no user interaction
  */
 void RunCrashReportClientUnattended(FMainLoopTiming& MainLoop, const FPlatformErrorReport& ErrorReport)
 {
-	FCrashReportClientUnattended CrashReportClient(ErrorReport);
+	FCrashReportClientUnattended CrashReportClient( ErrorReport );
 
 	// loop until the app is ready to quit
 	while (!GIsRequestingExit)
@@ -62,7 +66,7 @@ void ParseCommandLine(const TCHAR* CommandLine)
 		{
 			ReportDirectoryAbsolutePath = FParse::Token(CommandLineAfterExe, true /* handle escaped quotes */);
 		}
-		FParse::Value(CommandLineAfterExe, TEXT("AppName="), CrashedAppName);
+		FParse::Value( CommandLineAfterExe, TEXT( "AppName=" ), GameNameFromCmd );
 	}
 
 	if (ReportDirectoryAbsolutePath.IsEmpty())
@@ -83,16 +87,16 @@ FPlatformErrorReport LoadErrorReport()
 	}
 
 	FPlatformErrorReport ErrorReport(ReportDirectoryAbsolutePath);
+	
+	FString XMLWerFilename;
+	ErrorReport.FindFirstReportFileWithExtension( XMLWerFilename, TEXT( ".xml" ) );
+
+	GetCrashDescription() = FCrashDescription( ReportDirectoryAbsolutePath / XMLWerFilename );
+
 #if CRASH_REPORT_UNATTENDED_ONLY
 	return ErrorReport;
 #else
-	auto AppNameInReport = ErrorReport.FindCrashedAppName();
-	if (CrashedAppName.IsEmpty())
-	{
-		// If no app name was provided on the command line, assume the report is the correct one
-		CrashedAppName = AppNameInReport;
-	}
-	else if (CrashedAppName != AppNameInReport)
+	if( !GameNameFromCmd.IsEmpty() && GameNameFromCmd != GetCrashDescription().GameName )
 	{
 		// Don't display or upload anything if it's not the report we expected
 		ErrorReport = FPlatformErrorReport();
@@ -106,7 +110,7 @@ void RunCrashReportClient(const TCHAR* CommandLine)
 	// Set up the main loop
 	GEngineLoop.PreInit(CommandLine);
 
-	bool bUnattended = 
+	const bool bUnattended = 
 #if CRASH_REPORT_UNATTENDED_ONLY
 		true;
 #else
@@ -122,8 +126,15 @@ void RunCrashReportClient(const TCHAR* CommandLine)
 	FPlatformErrorReport::Init();
 	auto ErrorReport = LoadErrorReport();
 	
+	if( ErrorReport.HasFilesToUpload() )
+	{
+		// Send analytics.
+		GetCrashDescription().SendAnalytics();
+	}
+
 	if (bUnattended)
 	{
+		ErrorReport.SetUserComment( NSLOCTEXT( "CrashReportClient", "UnattendedMode", "Sent in the unattended mode" ) );
 		RunCrashReportClientUnattended(MainLoop, ErrorReport);
 	}
 	else
@@ -136,7 +147,7 @@ void RunCrashReportClient(const TCHAR* CommandLine)
 		FCrashReportClientStyle::Initialize();
 
 		// Create the main implementation object
-		TSharedRef<FCrashReportClient> CrashReportClient = MakeShareable(new FCrashReportClient(ErrorReport, CrashedAppName));
+		TSharedRef<FCrashReportClient> CrashReportClient = MakeShareable(new FCrashReportClient(ErrorReport));
 
 		// open up the app window	
 		TSharedRef<SCrashReportClient> ClientControl = SNew(SCrashReportClient, CrashReportClient);
@@ -184,8 +195,8 @@ void RunCrashReportClient(const TCHAR* CommandLine)
 #endif // !CRASH_REPORT_UNATTENDED_ONLY
 	}
 
+	FPlatformProcess::Sleep( 2.0f );
 	GEngineLoop.AppExit();
-
 	FPlatformErrorReport::ShutDown();
 }
 
@@ -195,5 +206,4 @@ void CrashReportClientCheck(bool bCondition, const TCHAR* Location)
 	{
 		UE_LOG(CrashReportClientLog, Warning, TEXT("CHECK FAILED at %s"), Location);
 	}
-}
-
+}
\ No newline at end of file
diff --git a/Engine/Source/Programs/CrashReporter/CrashReportClient/Private/CrashReportUtil.h b/Engine/Source/Programs/CrashReporter/CrashReportClient/Private/CrashReportUtil.h
--- a/Engine/Source/Programs/CrashReporter/CrashReportClient/Private/CrashReportUtil.h
+++ b/Engine/Source/Programs/CrashReporter/CrashReportClient/Private/CrashReportUtil.h
@@ -100,34 +100,40 @@ FunctorDirectoryVisitor<Functor> MakeDirectoryVisitor(Functor&& FunctorInstance)
 	return FunctorDirectoryVisitor<Functor>(MoveTemp(FunctorInstance));
 }
 
-/**
- * Create a multi line string to display from an exception and callstack
- * @param Exception Exception description string
- * @param Assertion Assertion description string
- * @param Callstack List of callstack entry strings
- * @return Multiline text
- */
-inline FText FormatReportDescription(const FString& Exception, const FString& Assertion, const TArray<FString>& Callstack)
+struct FCrashReportUtil
 {
-	FString Diagnostic = Exception + "\n\n";
-
-	if( !Assertion.IsEmpty() )
+	/**
+	 * Create a multi line string to display from an exception and callstack
+	 * @param Exception Exception description string
+	 * @param Assertion Assertion description string
+	 * @param Callstack List of callstack entry strings
+	 * @return Multiline text
+	 */
+	static inline FText FormatReportDescription( const FString& Exception, const FString& Assertion, const TArray<FString>& Callstack )
 	{
-		TArray<FString> MultilineAssertion;
-		Assertion.ParseIntoArray( &MultilineAssertion, TEXT( "#" ), true );
+		FString Diagnostic = Exception + "\n\n";
 
-		for( const auto& AssertionLine : MultilineAssertion )
+		if( !Assertion.IsEmpty() )
 		{
-			Diagnostic += AssertionLine;
+			TArray<FString> MultilineAssertion;
+			Assertion.ParseIntoArray( &MultilineAssertion, TEXT( "#" ), true );
+
+			for( const auto& AssertionLine : MultilineAssertion )
+			{
+				Diagnostic += AssertionLine;
+				Diagnostic += "\n";
+			}
+
 			Diagnostic += "\n";
 		}
 
-		Diagnostic += "\n";
+		for( const auto& Line : Callstack )
+		{
+			Diagnostic += Line + "\n";
+		}
+		return FText::FromString( Diagnostic );
 	}
 
-	for (const auto& Line: Callstack)
-	{
-		Diagnostic += Line + "\n";
-	}
-	return FText::FromString(Diagnostic);
-}
+	/** Formats processed diagnostic text by adding additional information about machine and user. */
+	static FText FormatDiagnosticText( const FText& DiagnosticText, const FString MachineId, const FString EpicAccountId, const FString UserNameNoDot );
+};
diff --git a/Engine/Source/Programs/CrashReporter/CrashReportClient/Private/GenericErrorReport.cpp b/Engine/Source/Programs/CrashReporter/CrashReportClient/Private/GenericErrorReport.cpp
--- a/Engine/Source/Programs/CrashReporter/CrashReportClient/Private/GenericErrorReport.cpp
+++ b/Engine/Source/Programs/CrashReporter/CrashReportClient/Private/GenericErrorReport.cpp
@@ -5,6 +5,7 @@
 #include "GenericErrorReport.h"
 #include "XmlFile.h"
 #include "CrashReportUtil.h"
+#include "CrashDescription.h"
 
 // ----------------------------------------------------------------
 // Helpers
@@ -63,8 +64,19 @@ bool FGenericErrorReport::SetUserComment(const FText& UserComment)
 
 	DynamicSignaturesNode->AppendChildNode(TEXT("Parameter3"), UserComment.ToString());
 
+	FString MachineIDandUserID;
+	// Set global user name ID: will be added to the report
+	if( FRocketSupport::IsRocket() )
+	{
+		MachineIDandUserID = FString::Printf( TEXT( "!MachineId:%s!EpicAccountId:%s" ), *GetCrashDescription().MachineId, *GetCrashDescription().EpicAccountId );
+	}
+	else
+	{
+		MachineIDandUserID = FString::Printf( TEXT( "!MachineId:%s!Name:%s" ), *GetCrashDescription().MachineId, *GetCrashDescription().UserName );
+	}
+
 	// Also write a user ID into the report
-	DynamicSignaturesNode->AppendChildNode(TEXT("Parameter4"), GCrashUserId);
+	DynamicSignaturesNode->AppendChildNode(TEXT("Parameter4"), MachineIDandUserID);
 
 	// Re-save over the top
 	return XmlFile.Save(XmlFilePath);
@@ -163,7 +175,7 @@ bool FGenericErrorReport::TryReadDiagnosticsFile(FText& OutReportDescription)
 			break;
 		}
 	}
-	OutReportDescription = FormatReportDescription(Exception, TEXT(""), Callstack);
+	OutReportDescription = FCrashReportUtil::FormatReportDescription( Exception, TEXT( "" ), Callstack );
 	return true;
 }
 
@@ -180,3 +192,9 @@ bool FGenericErrorReport::FindFirstReportFileWithExtension(FString& OutFilename,
 	return false;
 }
 
+FString FGenericErrorReport::FindCrashedAppName() const
+{
+	UE_LOG( LogTemp, Warning, TEXT( "FGenericPlatformMemory::Init not implemented on this platform" ) );
+	return FString( TEXT( "GenericAppName" ) );
+}
+
diff --git a/Engine/Source/Programs/CrashReporter/CrashReportClient/Private/GenericErrorReport.h b/Engine/Source/Programs/CrashReporter/CrashReportClient/Private/GenericErrorReport.h
--- a/Engine/Source/Programs/CrashReporter/CrashReportClient/Private/GenericErrorReport.h
+++ b/Engine/Source/Programs/CrashReporter/CrashReportClient/Private/GenericErrorReport.h
@@ -60,8 +60,8 @@ public:
 	bool LoadWindowsReportXmlFile( FString& OutString ) const;
 
 	/**
-	 * @param Description (exception and callstack) to fill in if succesful
-	 * @return Whether the file was found and succesfully read
+	 * @param Description (exception and callstack) to fill in if successful
+	 * @return Whether the file was found and successfully read
 	 */
 	bool TryReadDiagnosticsFile(FText& OutReportDescription);
 
@@ -82,6 +82,11 @@ public:
 		return FPaths::GetCleanFilename(ReportDirectory);
 	}
 
+	/**
+	 * Get the name of the crashed app from the report
+	 */
+	FString FindCrashedAppName() const;
+
 	/**
 	 * Is there anything to upload?
 	 */
@@ -90,13 +95,14 @@ public:
 		return ReportFilenames.Num() != 0;
 	}
 
-protected:
+
 	/**
-	 * Look throught the list of report files to find one with the given extension
+	 * Look thought the list of report files to find one with the given extension
 	 * @return Whether a file with the extension was found
 	 */
 	bool FindFirstReportFileWithExtension(FString& OutFilename, const TCHAR* Extension) const;
 
+protected:
 	/** Full path to the directory the report files are in */
 	FString ReportDirectory;
 
diff --git a/Engine/Source/Programs/CrashReporter/CrashReportClient/Private/Windows/WindowsErrorReport.cpp b/Engine/Source/Programs/CrashReporter/CrashReportClient/Private/Windows/WindowsErrorReport.cpp
--- a/Engine/Source/Programs/CrashReporter/CrashReportClient/Private/Windows/WindowsErrorReport.cpp
+++ b/Engine/Source/Programs/CrashReporter/CrashReportClient/Private/Windows/WindowsErrorReport.cpp
@@ -114,13 +114,7 @@ FText FWindowsErrorReport::DiagnoseReport() const
 	const auto& Exception = CrashDebugHelper->CrashInfo.Exception;
 	const FString Assertion = FWindowsReportParser::Find( ReportDirectory, TEXT( "AssertLog=" ) );
 
-	return FormatReportDescription( Exception.ExceptionString, Assertion, Exception.CallStackString );
-}
-
-FString FWindowsErrorReport::FindCrashedAppName() const
-{
-	const FString CrashedAppName = FWindowsReportParser::Find( ReportDirectory, TEXT( "AppName=" ) );
-	return CrashedAppName;
+	return FCrashReportUtil::FormatReportDescription( Exception.ExceptionString, Assertion, Exception.CallStackString );
 }
 
 FString FWindowsErrorReport::FindMostRecentErrorReport()
diff --git a/Engine/Source/Programs/CrashReporter/CrashReportClient/Private/Windows/WindowsErrorReport.h b/Engine/Source/Programs/CrashReporter/CrashReportClient/Private/Windows/WindowsErrorReport.h
--- a/Engine/Source/Programs/CrashReporter/CrashReportClient/Private/Windows/WindowsErrorReport.h
+++ b/Engine/Source/Programs/CrashReporter/CrashReportClient/Private/Windows/WindowsErrorReport.h
@@ -39,11 +39,6 @@ public:
 	 */
 	FText DiagnoseReport() const;
 
-	/**
-	 * Get the name of the crashed app from the report
-	 */
-	FString FindCrashedAppName() const;
-
 	/**
 	 * Look for the most recent Windows Error Report
 	 * @return Full path to the most recent report, or an empty string if none found
