Fix CrashReportClient in unattended (headless) mode.

- Also fixes TTP #345300 (packaging failure).

diff --git a/Engine/Source/Programs/CrashReporter/CrashReportClient/Private/CrashDescription.cpp b/Engine/Source/Programs/CrashReporter/CrashReportClient/Private/CrashDescription.cpp
--- a/Engine/Source/Programs/CrashReporter/CrashReportClient/Private/CrashDescription.cpp
+++ b/Engine/Source/Programs/CrashReporter/CrashReportClient/Private/CrashDescription.cpp
@@ -8,7 +8,9 @@
 #include "CrashReportAnalytics.h"
 #include "Runtime/Analytics/Analytics/Public/Analytics.h"
 #include "Runtime/Analytics/Analytics/Public/Interfaces/IAnalyticsProvider.h"
-#include "DesktopPlatformModule.h"
+#if !CRASH_REPORT_UNATTENDED_ONLY
+	#include "DesktopPlatformModule.h"
+#endif // CRASH_REPORT_UNATTENDED_ONLY
 
 
 
@@ -180,12 +182,17 @@ void FCrashDescription::InitializeEngineVersion()
 
 void FCrashDescription::InitializeIDs()
 {
+#if !CRASH_REPORT_UNATTENDED_ONLY
 	IDesktopPlatform* DesktopPlatform = FDesktopPlatformModule::Get();
 
 	MachineId = DesktopPlatform->GetMachineId().ToString( EGuidFormats::Digits );
 
 	// The Epic ID can be looked up from this ID.
 	EpicAccountId = DesktopPlatform->GetEpicAccountId();
+#else
+	MachineId = FGuid::NewGuid().ToString(EGuidFormats::Digits);
+	EpicAccountId = TEXT("BogusUnattendedAccount");
+#endif // CRASH_REPORT_UNATTENDED_ONLY
 
 	// Get the user name only for non-UE4 releases.
 	if( !FRocketSupport::IsRocket() )
diff --git a/Engine/Source/Programs/CrashReporter/CrashReportClient/Private/CrashReportAnalytics.cpp b/Engine/Source/Programs/CrashReporter/CrashReportClient/Private/CrashReportAnalytics.cpp
--- a/Engine/Source/Programs/CrashReporter/CrashReportClient/Private/CrashReportAnalytics.cpp
+++ b/Engine/Source/Programs/CrashReporter/CrashReportClient/Private/CrashReportAnalytics.cpp
@@ -6,7 +6,6 @@
 #include "EngineBuildSettings.h"
 #include "Runtime/Analytics/Analytics/Public/Analytics.h"
 #include "Runtime/Analytics/Analytics/Public/Interfaces/IAnalyticsProvider.h"
-#include "DesktopPlatformModule.h"
 
 // WARNING! Copied from EngineAnalytics.cpp
 
diff --git a/Engine/Source/Programs/CrashReporter/CrashReportClient/Private/CrashReportClient.cpp b/Engine/Source/Programs/CrashReporter/CrashReportClient/Private/CrashReportClient.cpp
--- a/Engine/Source/Programs/CrashReporter/CrashReportClient/Private/CrashReportClient.cpp
+++ b/Engine/Source/Programs/CrashReporter/CrashReportClient/Private/CrashReportClient.cpp
@@ -7,7 +7,6 @@
 
 #include "CrashReportClient.h"
 #include "UniquePtr.h"
-#include "DesktopPlatformModule.h"
 
 #define LOCTEXT_NAMESPACE "CrashReportClient"
 
