CrashReport - Updated commented code

diff --git a/Engine/Source/Programs/CrashReporter/CrashReportClient/Private/CrashReportAnalytics.cpp b/Engine/Source/Programs/CrashReporter/CrashReportClient/Private/CrashReportAnalytics.cpp
--- a/Engine/Source/Programs/CrashReporter/CrashReportClient/Private/CrashReportAnalytics.cpp
+++ b/Engine/Source/Programs/CrashReporter/CrashReportClient/Private/CrashReportAnalytics.cpp
@@ -99,8 +99,7 @@ void FCrashReportAnalytics::Initialize()
 			DefaultEngineAnalyticsConfig );
 		if( Analytics.IsValid() )
 		{
-			//IDesktopPlatform* DesktopPlatform = FDesktopPlatformModule::Get();
-			//const FString MachineId = DesktopPlatform->GetMachineId().ToString( EGuidFormats::Digits );
+			//const FString MachineId = FPlatformMisc::GetMachineId().ToString( EGuidFormats::Digits );
 			//Analytics->SetUserID( MachineId );
 			Analytics->SetUserID(FPlatformMisc::GetUniqueDeviceId());
 			Analytics->StartSession();
