On OS X in the crash handlers disable calls to functions which eventually try and allocate memory, resulting in the application hanging before the crash reporter can launch.

diff --git a/Engine/Source/Programs/CrashReporter/CrashReportClient/Private/Mac/CrashReportClientMainMac.cpp b/Engine/Source/Programs/CrashReporter/CrashReportClient/Private/Mac/CrashReportClientMainMac.cpp
--- a/Engine/Source/Programs/CrashReporter/CrashReportClient/Private/Mac/CrashReportClientMainMac.cpp
+++ b/Engine/Source/Programs/CrashReporter/CrashReportClient/Private/Mac/CrashReportClientMainMac.cpp
@@ -21,10 +21,12 @@ void CrashReporterCrashHandler(const FGenericCrashContext & GenericContext)
 	{
 		GWarn->Flush();
 	}
+	// @todo Disable calls to functions which eventually try and allocated memory, resulting in the application hanging before the crash reporter can launch
+	// We need to setup the memory allocator system to handle the case where we can't call the system malloc, free etc. to re-enable this.
 	if (GError)
 	{
 		GError->Flush();
-		GError->HandleError();
+		//GError->HandleError();
 	}
 	FPlatformMisc::RequestExit(true);
 }
diff --git a/Engine/Source/Runtime/Core/Private/Mac/MacPlatformMisc.cpp b/Engine/Source/Runtime/Core/Private/Mac/MacPlatformMisc.cpp
--- a/Engine/Source/Runtime/Core/Private/Mac/MacPlatformMisc.cpp
+++ b/Engine/Source/Runtime/Core/Private/Mac/MacPlatformMisc.cpp
@@ -1083,7 +1083,9 @@ static void DefaultCrashHandler(FMacCrashContext const& Context)
 	if (GError)
 	{
 		GError->Flush();
-		GError->HandleError();
+		// @todo Disable calls to functions which eventually try and allocated memory, resulting in the application hanging before the crash reporter can launch
+		// We need to setup the memory allocator system to handle the case where we can't call the system malloc, free etc. to re-enable this.
+		//GError->HandleError();
 	}
 	
 	return Context.GenerateCrashInfoAndLaunchReporter();
diff --git a/Engine/Source/Runtime/Launch/Private/Mac/LaunchMac.cpp b/Engine/Source/Runtime/Launch/Private/Mac/LaunchMac.cpp
--- a/Engine/Source/Runtime/Launch/Private/Mac/LaunchMac.cpp
+++ b/Engine/Source/Runtime/Launch/Private/Mac/LaunchMac.cpp
@@ -33,12 +33,14 @@ void EngineCrashHandler(const FGenericCrashContext & GenericContext)
 	{
 		GWarn->Flush();
 	}
+	// @todo Disable calls to functions which eventually try and allocated memory, resulting in the application hanging before the crash reporter can launch
+	// We need to setup the memory allocator system to handle the case where we can't call the system malloc, free etc. to re-enable this.
 	if (GError)
 	{
 		GError->Flush();
-		GError->HandleError();
+		//GError->HandleError();
 	}
-	LaunchStaticShutdownAfterError();
+	//LaunchStaticShutdownAfterError();
 	return Context.GenerateCrashInfoAndLaunchReporter();
 }
 
