Pop the CrashReportClient to the front when it finishes launching so that it appears over the crashed application on Mac OS X.

diff --git a/Engine/Source/Programs/CrashReporter/CrashReportClient/Private/Mac/CrashReportClientMainMac.cpp b/Engine/Source/Programs/CrashReporter/CrashReportClient/Private/Mac/CrashReportClientMainMac.cpp
--- a/Engine/Source/Programs/CrashReporter/CrashReportClient/Private/Mac/CrashReportClientMainMac.cpp
+++ b/Engine/Source/Programs/CrashReporter/CrashReportClient/Private/Mac/CrashReportClientMainMac.cpp
@@ -76,6 +76,7 @@ static FString GSavedCommandLine;
     NSAppleEventManager* appleEventManager = [NSAppleEventManager sharedAppleEventManager];
     [appleEventManager setEventHandler:self andSelector:@selector(handleQuitEvent:withReplyEvent:) forEventClass:kCoreEventClass andEventID:kAEQuitApplication];
 	
+	FPlatformMisc::ActivateApplication();
 	RunGameThread(self, @selector(runGameThread:));
 }
 
