Rather than cancelling OS quit requests on OS X and then quitting behind Cocoa's back, use the proper deferral mechanism so that we have Cocoa wait for the separate game thread to exit before it tears the application down. UE4 applications shouldn't cancel shutdown, restart or logout requests anymore.

diff --git a/Engine/Source/Programs/CrashReporter/CrashReportClient/Private/Mac/CrashReportClientMainMac.cpp b/Engine/Source/Programs/CrashReporter/CrashReportClient/Private/Mac/CrashReportClientMainMac.cpp
--- a/Engine/Source/Programs/CrashReporter/CrashReportClient/Private/Mac/CrashReportClientMainMac.cpp
+++ b/Engine/Source/Programs/CrashReporter/CrashReportClient/Private/Mac/CrashReportClientMainMac.cpp
@@ -63,12 +63,15 @@ static FString GSavedCommandLine;
 
 - (NSApplicationTerminateReply)applicationShouldTerminate:(NSApplication *)Sender;
 {
-	if(!GIsRequestingExit)
+	if(!GIsRequestingExit || ([NSThread gameThread] && [NSThread gameThread] != [NSThread mainThread]))
 	{
 		[self requestQuit:self];
-		return NSTerminateCancel;
+		return NSTerminateLater;
+	}
+	else
+	{
+		return NSTerminateNow;
 	}
-	return NSTerminateNow;
 }
 
 - (void)applicationDidFinishLaunching:(NSNotification *)Notification
diff --git a/Engine/Source/Programs/SlateViewer/Private/Mac/SlateViewerMainMac.cpp b/Engine/Source/Programs/SlateViewer/Private/Mac/SlateViewerMainMac.cpp
--- a/Engine/Source/Programs/SlateViewer/Private/Mac/SlateViewerMainMac.cpp
+++ b/Engine/Source/Programs/SlateViewer/Private/Mac/SlateViewerMainMac.cpp
@@ -32,13 +32,16 @@ static FString GSavedCommandLine;
 
 - (NSApplicationTerminateReply)applicationShouldTerminate:(NSApplication *)Sender;
 {
-	if(!GIsRequestingExit)
+	if(!GIsRequestingExit || ([NSThread gameThread] && [NSThread gameThread] != [NSThread mainThread]))
 	{
 		GIsRequestingExit = true;
 		FPlatformMisc::UpdateCachedMacMenuState = nullptr;
-		return NSTerminateCancel;
+		return NSTerminateLater;
+	}
+	else
+	{
+		return NSTerminateNow;
 	}
-	return NSTerminateNow;
 }
 
 - (void)applicationDidFinishLaunching:(NSNotification *)Notification
diff --git a/Engine/Source/Programs/SymbolDebugger/Private/Mac/SymbolDebuggerMainMac.cpp b/Engine/Source/Programs/SymbolDebugger/Private/Mac/SymbolDebuggerMainMac.cpp
--- a/Engine/Source/Programs/SymbolDebugger/Private/Mac/SymbolDebuggerMainMac.cpp
+++ b/Engine/Source/Programs/SymbolDebugger/Private/Mac/SymbolDebuggerMainMac.cpp
@@ -31,13 +31,16 @@ static FString GSavedCommandLine;
 
 - (NSApplicationTerminateReply)applicationShouldTerminate:(NSApplication *)Sender;
 {
-	if(!GIsRequestingExit)
+	if(!GIsRequestingExit || ([NSThread gameThread] && [NSThread gameThread] != [NSThread mainThread]))
 	{
 		GIsRequestingExit = true;
 		FPlatformMisc::UpdateCachedMacMenuState = nullptr;
-		return NSTerminateCancel;
+		return NSTerminateLater;
+	}
+	else
+	{
+		return NSTerminateNow;
 	}
-	return NSTerminateNow;
 }
 
 - (void)applicationDidFinishLaunching:(NSNotification *)Notification
diff --git a/Engine/Source/Programs/UnrealFrontend/Private/Mac/MacUnrealFrontendMain.cpp b/Engine/Source/Programs/UnrealFrontend/Private/Mac/MacUnrealFrontendMain.cpp
--- a/Engine/Source/Programs/UnrealFrontend/Private/Mac/MacUnrealFrontendMain.cpp
+++ b/Engine/Source/Programs/UnrealFrontend/Private/Mac/MacUnrealFrontendMain.cpp
@@ -64,12 +64,15 @@ static FString GSavedCommandLine;
 
 - (NSApplicationTerminateReply)applicationShouldTerminate:(NSApplication *)Sender;
 {
-	if(!GIsRequestingExit)
+	if(!GIsRequestingExit || ([NSThread gameThread] && [NSThread gameThread] != [NSThread mainThread]))
 	{
 		[self requestQuit:self];
-		return NSTerminateCancel;
+		return NSTerminateLater;
+	}
+	else
+	{
+		return NSTerminateNow;
 	}
-	return NSTerminateNow;
 }
 
 - (void)applicationDidFinishLaunching:(NSNotification *)Notification
diff --git a/Engine/Source/Runtime/Core/Private/Mac/CocoaThread.cpp b/Engine/Source/Runtime/Core/Private/Mac/CocoaThread.cpp
--- a/Engine/Source/Runtime/Core/Private/Mac/CocoaThread.cpp
+++ b/Engine/Source/Runtime/Core/Private/Mac/CocoaThread.cpp
@@ -356,6 +356,15 @@ FCocoaRunLoopSource* FCocoaRunLoopSource::GameRunLoopSource = nullptr;
 	FCocoaRunLoopSource::RegisterGameRunLoop([GameRunLoop getCFRunLoop]);
 	
 	[super main];
+	
+	if (GIsRequestingExit)
+	{
+		MainThreadCall(^{
+			[NSApp replyToApplicationShouldTerminate:YES];
+		}, NSDefaultRunLoopMode, false);
+	}
+	
+	[self release];
 }
 
 @end
@@ -429,6 +438,11 @@ void RunGameThread(id Target, SEL Selector)
 	[GameThread start];
 #else
 	[Target performSelector:Selector withObject:nil];
+	
+	if (GIsRequestingExit)
+	{
+		[NSApp replyToApplicationShouldTerminate:YES];
+	}
 #endif
 }
 
diff --git a/Engine/Source/Runtime/Launch/Private/Mac/LaunchMac.cpp b/Engine/Source/Runtime/Launch/Private/Mac/LaunchMac.cpp
--- a/Engine/Source/Runtime/Launch/Private/Mac/LaunchMac.cpp
+++ b/Engine/Source/Runtime/Launch/Private/Mac/LaunchMac.cpp
@@ -157,12 +157,18 @@ void EngineCrashHandler(const FGenericCrashContext & GenericContext)
 
 - (NSApplicationTerminateReply)applicationShouldTerminate:(NSApplication *)Sender;
 {
-	if(!GIsRequestingExit)
+	if(!GIsRequestingExit || ([NSThread gameThread] && [NSThread gameThread] != [NSThread mainThread]))
 	{
-		[self requestQuit:self];
-		return NSTerminateCancel;
+		if (!GIsRequestingExit)
+		{
+			[self requestQuit:self];
+		}
+		return NSTerminateLater;
+	}
+	else
+	{
+		return NSTerminateNow;
 	}
-	return NSTerminateNow;
 }
 
 - (void) runGameThread:(id)Arg
