CrashReport - Added more data and updated comments

diff --git a/Engine/Source/Programs/CrashReporter/CrashReportClient/Private/CrashDescription.cpp b/Engine/Source/Programs/CrashReporter/CrashReportClient/Private/CrashDescription.cpp
--- a/Engine/Source/Programs/CrashReporter/CrashReportClient/Private/CrashDescription.cpp
+++ b/Engine/Source/Programs/CrashReporter/CrashReportClient/Private/CrashDescription.cpp
@@ -1,6 +1,7 @@
 // Copyright 1998-2014 Epic Games, Inc. All Rights Reserved.
 
 #include "CrashReportClientApp.h"
+#include "GenericPlatformContext.h"
 
 #include "EngineVersion.h"
 #include "XmlFile.h"
@@ -10,6 +11,7 @@
 #include "Runtime/Analytics/Analytics/Public/Interfaces/IAnalyticsProvider.h"
 
 FCrashDescription::FCrashDescription() :
+	CrashDescriptionVersion( ECrashDescVersions::VER_1_NewCrashFormat ),
 	BuiltFromCL( -1 ),
 	TimeOfCrash( FDateTime::UtcNow() ),
 	bHasMiniDumpFile( false ),
@@ -22,6 +24,8 @@ FCrashDescription::FCrashDescription() :
 
 
 FCrashDescription::FCrashDescription( FString WERXMLFilepath ) :
+	// WER XML files are forced to be in the first version
+	CrashDescriptionVersion( ECrashDescVersions::VER_1_NewCrashFormat ),
 	BuiltFromCL( -1 ),
 	TimeOfCrash( FDateTime::UtcNow() ),
 	bHasMiniDumpFile( false ),
diff --git a/Engine/Source/Programs/CrashReporter/CrashReportClient/Private/CrashDescription.h b/Engine/Source/Programs/CrashReporter/CrashReportClient/Private/CrashDescription.h
--- a/Engine/Source/Programs/CrashReporter/CrashReportClient/Private/CrashDescription.h
+++ b/Engine/Source/Programs/CrashReporter/CrashReportClient/Private/CrashDescription.h
@@ -2,13 +2,6 @@
 
 #pragma once
 
-/** Enumerates crash description versions. */
-enum class ECrashDescVersions
-{
-	/** Introduces a new crash description format. */
-	VER_1_NewCrashFormat,
-};
-
 /**
  *	Describes a unified crash, should be used by all platforms.
  *	Still lacks some of the properties, they will be added later.
@@ -23,7 +16,7 @@ struct FCrashDescription
 	/** 
 	 * Version of this crash description. 
 	 */
-	static const ECrashDescVersions CrashDescriptionVersion = ECrashDescVersions::VER_1_NewCrashFormat;
+	ECrashDescVersions CrashDescriptionVersion;
 
 	/** An unique report name that this crash belongs to. */
 	FString ReportName;
@@ -31,18 +24,24 @@ struct FCrashDescription
 	/**
 	 * The name of the game that crashed. (AppID)
 	 * @GameName	varchar(64)
+	 * 
+	 * FApp::GetGameName()
 	 */
 	FString GameName;
 
 	/**
 	 * The mode the game was in e.g. editor.
-	 * @EngineMode	varchar(64)	
+	 * @EngineMode	varchar(64)
+	 * 
+	 * FPlatformMisc::GetEngineMode()
 	 */
 	FString EngineMode;
 
 	/**
 	 * The platform that crashed e.g. Win64.
 	 * @PlatformName	varchar(64)	
+	 * 
+	 * FPlatformProperties::PlatformName()
 	 */
 	FString Platform;
 
@@ -51,90 +50,147 @@ struct FCrashDescription
 	 * E.g. 4.3.0.0-2215663+UE4-Releases+4.3
 	 * BuildVersion-BuiltFromCL-BranchName
 	 * @EngineVersion	varchar(64)	
+	 * 
+	 * GEngineVersion.ToString()
+	 * ENGINE_VERSION_STRING
 	 */
 	FEngineVersion EngineVersion;
 
 	/**
 	 * The three component version of the app e.g. 4.4.1
-	 * @BuildVersion	varchar(128)	
+	 * @BuildVersion	varchar(128)
+	 * 
+	 * ENGINE_MAJOR_VERSION.ENGINE_MINOR_VERSION.ENGINE_PATCH_VERSION
 	 */
 	FString BuildVersion;
 
 	/**
 	 * Built from changelist.
-	 * @ChangeListVersion	varchar(64)	
+	 * @ChangeListVersion	varchar(64)
+	 * 
+	 * BUILT_FROM_CHANGELIST
 	 */
 	uint32 BuiltFromCL;
 
 	/**
 	 * The name of the branch this game was built out of.
-	 * @Branch varchar(32) 
+	 * @Branch varchar(32)
+	 * 
+	 * BRANCH_NAME
 	 */
 	FString BranchName;
 
 	/**
 	 * The command line of the application that crashed.
-	 * @CommandLine varchar(512) 
+	 * @CommandLine varchar(512)
+	 * 
+	 * FCommandLine::Get() 
 	 */
 	FString CommandLine;
 
 	/**
 	 * The base directory where the app was running.
-	 * @BaseDir varchar(512) 
+	 * @BaseDir varchar(512)
+	 * 
+	 * FPlatformProcess::BaseDir()
 	 */
 	FString BaseDir;
 
 	/**
 	 * The language ID the application that crashed.
-	 * @LanguageExt varchar(64) 
+	 * @LanguageExt varchar(64)
+	 * 
+	 * FInternationalization::Get().GetCurrentCulture()->GetLCID()
 	 */
 	FString LanguageLCID;
 
 	/** 
 	 * The name of the user that caused this crash.
-	 * @UserName varchar(64)  
+	 * @UserName varchar(64)
+	 * 
+	 * FString( FPlatformProcess::UserName() ).Replace( TEXT( "." ), TEXT( "" ) )
 	 */
 	FString UserName;
 
 	/**
 	 * The unique ID used to identify the machine the crash occurred on.
-	 * @ComputerName varchar(64) 
+	 * @ComputerName varchar(64)
+	 * 
+	 * FPlatformMisc::GetMachineId().ToString( EGuidFormats::Digits )
 	 */
 	FString MachineId;
 
 	/** 
 	 * The Epic account ID for the user who last used the Launcher.
-	 * @EpicAccountId	varchar(64)	
+	 * @EpicAccountId	varchar(64)
+	 * 
+	 * FPlatformMisc::GetEpicAccountId()
 	 */
 	FString EpicAccountId;
 
 	/**
 	 * An array of FStrings representing the callstack of the crash.
-	 * @RawCallStack	varchar(MAX)	
+	 * @RawCallStack	varchar(MAX)
+	 * 
+	 * @TODO
 	 */
 	TArray<FString> CallStack;
 
 	/**
 	 * An array of FStrings showing the source code around the crash.
-	 * @SourceContext varchar(max) 
+	 * @SourceContext varchar(max)
+	 * 
+	 * @TODO
 	 */
 	TArray<FString> SourceContext;
 
 	/**
 	 * An array of FStrings representing the user description of the crash.
-	 * @Description	varchar(512)	
+	 * @Description	varchar(512)
+	 * 
+	 * @TODO
 	 */
 	TArray<FString> UserDescription;
 
 	/**
 	 * The error message, can be assertion message, ensure message or message from the fatal error.
-	 * @Summary	varchar(512)	
+	 * @Summary	varchar(512)
+	 * 
+	 * GErrorMessage
 	 */
 	TArray<FString> ErrorMessage;
 
+	// @TODO yrx 2014-09-10 Add misc data
+
+	/**
+	 *	Generic platform memory stats.
+	 *	@ADD-TO-THE-DATABASE
+	 *	
+	 *	FPlatformMemoryStats::GetStats()
+	 */
+	FGenericPlatformMemoryStats MemoryStats;
+
+	/**
+	 *	Whether this crash was caused by the OOM.
+	 *	@ADD-TO-THE-DATABASE
+	 *	
+	 *	@TODO
+	 */
+	bool bIsOOM;
+
+	/**
+	 * The misc and platform specific payload data.
+	 * @ADD-TO-THE-DATABASE	varchar(max)
+	 * 
+	 * @TODO
+	 */
+	TArray<FString> PayloadData;
+
 	/**
 	 * The UTC time the crash occurred.
-	 * @TimeOfCrash	datetime	
+	 * @TimeOfCrash	datetime
+	 * 
+	 * FDateTime::UtcNow().GetTicks()
 	 */
 	FDateTime TimeOfCrash;
 
