CrashReport - Added query to replace bad EngineMode with 'Unknown'

diff --git a/Engine/Source/Programs/CrashReporter/CrashReportWebSite/Models/Update-FixBadEngineMode.sql b/Engine/Source/Programs/CrashReporter/CrashReportWebSite/Models/Update-FixBadEngineMode.sql
new file mode 100644
--- /dev/null
+++ b/Engine/Source/Programs/CrashReporter/CrashReportWebSite/Models/Update-FixBadEngineMode.sql
@@ -0,0 +1,20 @@
+SELECT [EngineMode]
+FROM [CrashReport].[dbo].[Crashes]
+where charindex('Game', [EngineMode]) = 0 and
+	  charindex('Commandlet', [EngineMode]) = 0 and
+	  charindex('Editor', [EngineMode]) = 0 and
+	  charindex('Tool', [EngineMode]) = 0 and
+	  charindex('Server', [EngineMode]) = 0 and
+	  charindex('Unknown', [EngineMode]) = 0 
+go
+
+
+UPDATE [CrashReport].[dbo].[Crashes]
+SET [EngineMode] = 'Unknown'
+FROM [CrashReport].[dbo].[Crashes]
+where charindex('Game', [EngineMode]) = 0 and
+	  charindex('Commandlet', [EngineMode]) = 0 and
+	  charindex('Editor', [EngineMode]) = 0 and
+	  charindex('Tool', [EngineMode]) = 0 and
+	  charindex('Server', [EngineMode]) = 0 and
+	  charindex('Unknown', [EngineMode]) = 0 
