CrashReport - Added a missing header (https://github.com/EpicGames/UnrealEngine/pull/498)

diff --git a/Engine/Source/Programs/CrashReporter/CrashReportClient/Private/Linux/CrashReportClientMainLinux.cpp b/Engine/Source/Programs/CrashReporter/CrashReportClient/Private/Linux/CrashReportClientMainLinux.cpp
--- a/Engine/Source/Programs/CrashReporter/CrashReportClient/Private/Linux/CrashReportClientMainLinux.cpp
+++ b/Engine/Source/Programs/CrashReporter/CrashReportClient/Private/Linux/CrashReportClientMainLinux.cpp
@@ -1,6 +1,7 @@
 // Copyright 1998-2014 Epic Games, Inc. All Rights Reserved.
 
 #include "CrashReportClientApp.h"
+#include "LinuxPlatformCrashContext.h"
 #include <locale.h>
 
 extern int32 ReportCrash(const FLinuxCrashContext& Context);	// FIXME: handle expose it someplace else?
