CrashReportWWW - MachineId is now stored for the rocket users

diff --git a/Engine/Source/Programs/CrashReporter/CrashReportWebSite/Models/CrashRepository.cs b/Engine/Source/Programs/CrashReporter/CrashReportWebSite/Models/CrashRepository.cs
--- a/Engine/Source/Programs/CrashReporter/CrashReportWebSite/Models/CrashRepository.cs
+++ b/Engine/Source/Programs/CrashReporter/CrashReportWebSite/Models/CrashRepository.cs
@@ -860,6 +860,7 @@ namespace Tools.CrashReporter.CrashReportWebSite.Models
 			// Valid MachineID and EpicAccountId, updated crash from UE4 release
 			else if( !string.IsNullOrEmpty( NewCrashInfo.MachineGuid ) && !string.IsNullOrEmpty( NewCrashInfo.EpicAccountId ) )
 			{
+				NewCrash.ComputerName = NewCrashInfo.MachineGuid;
 				NewCrash.EpicAccountId = NewCrashInfo.EpicAccountId;
 				NewCrash.UserNameId = CrashRepositoryDataContext.FindOrAddUser( UserNameAnonymous );
 			}
