CrashReportWWW - Added users mapping

diff --git a/Engine/Source/Programs/CrashReporter/CrashReportWebSite/Models/CrashReport.dbml b/Engine/Source/Programs/CrashReporter/CrashReportWebSite/Models/CrashReport.dbml
--- a/Engine/Source/Programs/CrashReporter/CrashReportWebSite/Models/CrashReport.dbml
+++ b/Engine/Source/Programs/CrashReporter/CrashReportWebSite/Models/CrashReport.dbml
@@ -126,4 +126,13 @@
       <Association Name="UserGroup_Buggs_UserGroup" Member="Buggs_UserGroups" ThisKey="Id" OtherKey="UserGroupId" Type="Buggs_UserGroup" />
     </Type>
   </Table>
+  <Table Name="dbo.UsersMapping" Member="UsersMappings">
+    <Type Name="UsersMapping">
+      <Column Name="Id" Type="System.Int32" DbType="Int NOT NULL IDENTITY" IsPrimaryKey="true" IsDbGenerated="true" CanBeNull="false" />
+      <Column Name="UserEmail" Type="System.String" DbType="NVarChar(128) NOT NULL" CanBeNull="false" />
+      <Column Name="MachineId" Type="System.String" DbType="NChar(32)" CanBeNull="true" />
+      <Column Name="EpicAccountId" Type="System.String" DbType="NChar(32)" CanBeNull="true" />
+      <Column Name="UserName" Type="System.String" DbType="NVarChar(64)" CanBeNull="true" />
+    </Type>
+  </Table>
 </Database>
\ No newline at end of file
diff --git a/Engine/Source/Programs/CrashReporter/CrashReportWebSite/Models/CrashReport.dbml.layout b/Engine/Source/Programs/CrashReporter/CrashReportWebSite/Models/CrashReport.dbml.layout
--- a/Engine/Source/Programs/CrashReporter/CrashReportWebSite/Models/CrashReport.dbml.layout
+++ b/Engine/Source/Programs/CrashReporter/CrashReportWebSite/Models/CrashReport.dbml.layout
@@ -1,5 +1,5 @@
 ﻿<?xml version="1.0" encoding="utf-8"?>
-<ordesignerObjectsDiagram dslVersion="1.0.0.0" absoluteBounds="0, 0, 13.75, 9.25" name="CrashReport">
+<ordesignerObjectsDiagram dslVersion="1.0.0.0" absoluteBounds="0, 0, 13.75, 9.75" name="CrashReport">
   <DataContextMoniker Name="/CrashReportDataContext" />
   <nestedChildShapes>
     <classShape Id="a37f1cdb-93d7-48ef-8da7-2b7ca09282b0" absoluteBounds="11.5, 1.25, 2, 4.0785139973958326">
@@ -77,7 +77,7 @@
         <classShapeMoniker Id="8a60fc4a-a797-417f-87ed-25b22bd5e420" />
       </nodes>
     </associationConnector>
-    <associationConnector edgePoints="[(11.5 : 2.47199625651042); (10.875 : 2.47199625651042)]" fixedFrom="Algorithm" fixedTo="Algorithm">
+    <associationConnector edgePoints="[(11.5 : 2.47199625651042); (10.875 : 2.47199625651042)]" fixedFrom="NotFixed" fixedTo="NotFixed">
       <AssociationMoniker Name="/CrashReportDataContext/Bugg/Bugg_Buggs_Crash" />
       <nodes>
         <classShapeMoniker Id="a37f1cdb-93d7-48ef-8da7-2b7ca09282b0" />
@@ -126,5 +126,11 @@
         <classShapeMoniker Id="8a60fc4a-a797-417f-87ed-25b22bd5e420" />
       </nodes>
     </associationConnector>
+    <classShape Id="d3fe456a-12f1-46c2-a283-3d1ccff66eff" absoluteBounds="2.75, 7.75, 2, 1.7708968098958344">
+      <DataClassMoniker Name="/CrashReportDataContext/UsersMapping" />
+      <nestedChildShapes>
+        <elementListCompartment Id="74a56d7a-d2a7-480f-8a79-9cfe3115cae7" absoluteBounds="2.7649999999999997, 8.21, 1.9700000000000002, 1.2108968098958333" name="DataPropertiesCompartment" titleTextColor="Black" itemTextColor="Black" />
+      </nestedChildShapes>
+    </classShape>
   </nestedChildShapes>
 </ordesignerObjectsDiagram>
\ No newline at end of file
