CrashReportWWW - Improved crash view listing performance

diff --git a/Engine/Source/Programs/CrashReporter/CrashReportWebSite/Views/Crashes/ViewCrash.ascx b/Engine/Source/Programs/CrashReporter/CrashReportWebSite/Views/Crashes/ViewCrash.ascx
--- a/Engine/Source/Programs/CrashReporter/CrashReportWebSite/Views/Crashes/ViewCrash.ascx
+++ b/Engine/Source/Programs/CrashReporter/CrashReportWebSite/Views/Crashes/ViewCrash.ascx
@@ -83,96 +83,96 @@
 	</tr>
 	<% if( Model.Results.ToList() != null )
 	   { %>
-		<%
-		int Iteration = 0;
-		foreach( Crash CurrentCrash in ( IEnumerable )Model.Results )
-		{
-			Buggs_Crash BuggCrash = null;
-			try
-			{
+	<% 
+				int Iteration = 0;
+				foreach( Crash CurrentCrash in (IEnumerable)Model.Results )
+				{
+						Buggs_Crash BuggCrash = null;
+						try
+						{
 				BuggCrash = CurrentCrash.Buggs_Crashes.FirstOrDefault();
-			}
-			catch
-			{
-				BuggCrash = null;
-			}
-
-			Iteration++;
-			string CrashRowColor = "grey";
-			string CrashColorDescription = "Incoming Crash";
-
-			if( string.IsNullOrWhiteSpace( CurrentCrash.FixedChangeList ) && string.IsNullOrWhiteSpace( CurrentCrash.TTPID ) )
-			{
-				CrashRowColor = "#FFFF88"; // yellow
-				CrashColorDescription = "This crash has not been fixed or assigned a TTP";
-			}
-
-			if( !string.IsNullOrWhiteSpace( CurrentCrash.FixedChangeList ) )
-			{
-				// Green
-				CrashRowColor = "#008C00"; //green
-				CrashColorDescription = "This crash has been fixed in CL# " + CurrentCrash.FixedChangeList;
-			}
-
-			if( ( BuggCrash != null ) && !string.IsNullOrWhiteSpace( CurrentCrash.TTPID ) && string.IsNullOrWhiteSpace( CurrentCrash.FixedChangeList ) )
-			{
-				CrashRowColor = "#D01F3C"; // red
-				CrashColorDescription = "This crash has occurred more than once and been assigned a TTP: " + CurrentCrash.TTPID + " but has not been fixed.";
-			}
-
-			if( CurrentCrash.Status == "Tester" )
-			{
-				CrashRowColor = "#718698"; // red
-				CrashColorDescription = "This crash status has been set to Tester";
-			}
-			%>
+						}
+						catch
+						{
+							BuggCrash = null;
+						}
+
+						Iteration++;
+						string CrashRowColor = "grey";
+						string CrashColorDescription = "Incoming Crash";
+
+						if( string.IsNullOrWhiteSpace( CurrentCrash.FixedChangeList ) && string.IsNullOrWhiteSpace( CurrentCrash.TTPID ) )
+						{
+							CrashRowColor = "#FFFF88"; // yellow
+							CrashColorDescription = "This crash has not been fixed or assigned a TTP";
+						}
+
+						if( !string.IsNullOrWhiteSpace( CurrentCrash.FixedChangeList ) )
+						{
+							// Green
+							CrashRowColor = "#008C00"; //green
+							CrashColorDescription = "This crash has been fixed in CL# " + CurrentCrash.FixedChangeList;
+						}
+
+						if( ( BuggCrash != null ) && !string.IsNullOrWhiteSpace( CurrentCrash.TTPID ) && string.IsNullOrWhiteSpace( CurrentCrash.FixedChangeList ) )
+						{
+							CrashRowColor = "#D01F3C"; // red
+							CrashColorDescription = "This crash has occurred more than once and been assigned a TTP: " + CurrentCrash.TTPID + " but has not been fixed.";
+						}
+
+						if( CurrentCrash.Status == "Tester" )
+						{
+							CrashRowColor = "#718698"; // red
+							CrashColorDescription = "This crash status has been set to Tester";
+						}
+	%>
 
 			<tr class='CrashRow <%=CurrentCrash.User.UserGroup %>'>
 				<td class="CrashTd" style="background-color: <%=CrashRowColor %>;" title = "<%=CrashColorDescription %>"><INPUT TYPE="CHECKBOX" Value="<%=Iteration %>"NAME="<%=CurrentCrash.Id%>" id="<%=CurrentCrash.Id %>" class='input CrashCheckBox' ></td> 
 				<td class="Id"><%=Html.ActionLink(CurrentCrash.Id.ToString(), "Show", new { controller = "crashes", id = CurrentCrash.Id }, null)%> <br /> <em style="font-size:x-small;"> <%= ( BuggCrash != null ) ? BuggCrash.BuggId.ToString() : ""  %></em></td> 
-				<td class="TimeOfCrash">
-					<%=CurrentCrash.GetTimeOfCrash()[0]%>
-					<br />
-					<%=CurrentCrash.GetTimeOfCrash()[1]%>
-				</td>
-				<td class="Username"><%=CurrentCrash.User.UserName%></td>
-				<td class="CallStack" >
-					<div style="clip : auto; ">
-						<div id='<%=CurrentCrash.Id %>-TrimmedCallStackBox' class='TrimmedCallStackBox'>
-							<% foreach( CallStackEntry Entry in CurrentCrash.GetCallStackEntries( 0, 4 ) )
-							   { %>
-								<span class="function-name">
-									<%=Url.CallStackSearchLink( Html.Encode( Entry.GetTrimmedFunctionName( 45 ) ), Model )%>
+		<td class="TimeOfCrash">
+			<%=CurrentCrash.GetTimeOfCrash()[0]%>
+			<br />
+			<%=CurrentCrash.GetTimeOfCrash()[1]%>
+		</td>
+		<td class="Username"><%=CurrentCrash.UserName%></td>
+		<td class="CallStack">
+			<div style="clip: auto;">
+				<div id='<%=CurrentCrash.Id %>-TrimmedCallStackBox' class='TrimmedCallStackBox'>
+					<% foreach( CallStackEntry Entry in CurrentCrash.GetCallStackEntries( 0, 4 ) )
+					   { %>
+					<span class="function-name">
+						<%=Url.CallStackSearchLink( Html.Encode( Entry.GetTrimmedFunctionName( 45 ) ), Model )%>
 								</span><br />
-							<% } %>
-						</div>
-						<a class='FullCallStackTrigger'><span class='FullCallStackTriggerText'>Full Callstack</span>
-							<div id='<%=CurrentCrash.Id %>-FullCallStackBox' class='FullCallStackBox'>
-								<% foreach( CallStackEntry Entry in CurrentCrash.GetCallStackEntries( 0, 40 ) )
-								   {%>
-										<span class="FunctionName">
-											<%=Html.Encode( Entry.GetTrimmedFunctionName( 60 ) )%>
-										</span>
-										<br />
-								<% } %>
-							</div>
-						</a>
+					<% } %>
+				</div>
+				<a class='FullCallStackTrigger'><span class='FullCallStackTriggerText'>Full Callstack</span>
+					<div id='<%=CurrentCrash.Id %>-FullCallStackBox' class='FullCallStackBox'>
+						<% foreach( CallStackEntry Entry in CurrentCrash.GetCallStackEntries( 0, 40 ) )
+						   {%>
+						<span class="FunctionName">
+							<%=Html.Encode( Entry.GetTrimmedFunctionName( 60 ) )%>
+						</span>
+						<br />
+						<% } %>
 					</div>
+				</a>
+			</div>
 
-				</td>
-				<td class="Game"><%=CurrentCrash.GameName%></td>
-				<td class="Mode"><%=CurrentCrash.EngineMode%></td>
-				<td class="FixedChangeList"><%=CurrentCrash.FixedChangeList%></td>
-				<td class="Ttp"><%=CurrentCrash.TTPID%>&nbsp;</td>
-				<td class="Branch"><%=CurrentCrash.Branch%>&nbsp;</td>
-				<td class="Description"><span class="TableData"><%=CurrentCrash.Description%>&nbsp;</span></td>
-				<td class="Summary"><%=CurrentCrash.Summary%></td>
-				<td class="ChangeListVersion"><%=CurrentCrash.ChangeListVersion%></td>
-				<td class="Computer"><%=CurrentCrash.ComputerName%></td>
-				<td class="Platform"><%=CurrentCrash.PlatformName%></td>
-				<td class="Status"><%=CurrentCrash.Status%></td>
-				<td class="Module"><%=CurrentCrash.Module%></td>
-			</tr>
+		</td>
+		<td class="Game"><%=CurrentCrash.GameName%></td>
+		<td class="Mode"><%=CurrentCrash.EngineMode%></td>
+		<td class="FixedChangeList"><%=CurrentCrash.FixedChangeList%></td>
+		<td class="Ttp"><%=CurrentCrash.TTPID%>&nbsp;</td>
+		<td class="Branch"><%=CurrentCrash.Branch%>&nbsp;</td>
+		<td class="Description"><span class="TableData"><%=CurrentCrash.Description%>&nbsp;</span></td>
+		<td class="Summary"><%=CurrentCrash.Summary%></td>
+		<td class="ChangeListVersion"><%=CurrentCrash.ChangeListVersion%></td>
+		<td class="Computer"><%=CurrentCrash.ComputerName%></td>
+		<td class="Platform"><%=CurrentCrash.PlatformName%></td>
+		<td class="Status"><%=CurrentCrash.Status%></td>
+		<td class="Module"><%=CurrentCrash.Module%></td>
+	</tr>
 
 		<% } %>
 	<% } %>
