CrashReportWWW - Improved buggs view listing performance

diff --git a/Engine/Source/Programs/CrashReporter/CrashReportWebSite/Models/BuggRepository.cs b/Engine/Source/Programs/CrashReporter/CrashReportWebSite/Models/BuggRepository.cs
--- a/Engine/Source/Programs/CrashReporter/CrashReportWebSite/Models/BuggRepository.cs
+++ b/Engine/Source/Programs/CrashReporter/CrashReportWebSite/Models/BuggRepository.cs
@@ -125,12 +125,22 @@ namespace Tools.CrashReporter.CrashReportWebSite.Models
 			List<string> FunctionCalls = new List<string>();
 			try
 			{
-				FunctionCalls =
+				/*FunctionCalls =
 				(
 					from Id in Ids
 					join Function in BuggsDataContext.FunctionCalls on Id equals Function.Id
 					select Function.Call
-				).ToList();
+				).ToList();*/
+				
+				{
+					List<FunctionCall> Funcs = BuggsDataContext.FunctionCalls.Where( FuncCall => Ids.Contains(FuncCall.Id) ).ToList();
+					// Order by Ids
+					foreach( int Id in Ids )
+					{
+						var Found = Funcs.Find( FC => FC.Id == Id );
+						FunctionCalls.Add( Found.Call );
+					}		
+				}
 			}
 			catch( Exception Ex )
 			{
