CrashReportWWW - Reordered some code to fix the issue with updating the user group

diff --git a/Engine/Source/Programs/CrashReporter/CrashReportWebSite/Views/Users/Index.aspx b/Engine/Source/Programs/CrashReporter/CrashReportWebSite/Views/Users/Index.aspx
--- a/Engine/Source/Programs/CrashReporter/CrashReportWebSite/Views/Users/Index.aspx
+++ b/Engine/Source/Programs/CrashReporter/CrashReportWebSite/Views/Users/Index.aspx
@@ -34,10 +34,10 @@
 	</div>
 
 	<div id="UserNames">
-		<form action="<%="[REDACTED_USER_HOME]/" + Model.UserGroup %>" method="POST" id="UserNamesForm" style="text-align: center">
-			<%foreach( string UserName in Model.Users )
-			{%>
-			<span>
+		<%foreach( string UserName in Model.Users )
+		{%>
+		<div>
+			<form action="<%="[REDACTED_USER_HOME]/" + Model.UserGroup %>" method="POST" id="UserNamesForm" style="text-align: center">
 				<%=UserName%>
 
 				<span id="set-user-group"></span>
@@ -49,9 +49,9 @@
 					<%} %>
 				</select>
 				<br />
-			</span>
-			<% } %>
-		</form>
+			</form>
+		</div>
+		<% } %>
 	</div>
 </asp:Content>
 
