CrashReportWWW - No change

diff --git a/Engine/Source/Programs/CrashReporter/CrashReportWebSite/Models/BuggsViewModel.cs b/Engine/Source/Programs/CrashReporter/CrashReportWebSite/Models/BuggsViewModel.cs
--- a/Engine/Source/Programs/CrashReporter/CrashReportWebSite/Models/BuggsViewModel.cs
+++ b/Engine/Source/Programs/CrashReporter/CrashReportWebSite/Models/BuggsViewModel.cs
@@ -89,6 +89,5 @@ namespace Tools.CrashReporter.CrashReportWebSite.Models
 					} );
 			}
 		}
-		/// // @TODO yrx 2014-07-28 Grab build versions from the database?
 	}
 }
\ No newline at end of file
diff --git a/Engine/Source/Programs/CrashReporter/CrashReportWebSite/Models/CrashReport.dbml b/Engine/Source/Programs/CrashReporter/CrashReportWebSite/Models/CrashReport.dbml
--- a/Engine/Source/Programs/CrashReporter/CrashReportWebSite/Models/CrashReport.dbml
+++ b/Engine/Source/Programs/CrashReporter/CrashReportWebSite/Models/CrashReport.dbml
@@ -98,24 +98,14 @@
       <Column Name="EpicAccountId" Type="System.String" DbType="varchar(64)" CanBeNull="true" />
       <Column Name="EngineVersion" Type="System.String" DbType="varchar(64)" CanBeNull="true" />
       <Association Name="Crash_Buggs_Crash" Member="Buggs_Crashes" ThisKey="Id" OtherKey="CrashId" Type="Buggs_Crash" />
-      <Association Name="Crash_Crash_FunctionCall" Member="Crash_FunctionCalls" ThisKey="Id" OtherKey="CrashId" Type="Crash_FunctionCall" />
       <Association Name="User_Crash" Member="UserById" Storage="_User" ThisKey="UserNameId" OtherKey="Id" Type="User" IsForeignKey="true" />
       <Association Name="User_Crash1" Member="UserByName" ThisKey="UserName" OtherKey="UserName" Type="User" IsForeignKey="true" />
     </Type>
   </Table>
-  <Table Name="dbo.Crashes_FunctionCalls" Member="Crash_FunctionCalls">
-    <Type Name="Crash_FunctionCall">
-      <Column Name="CrashId" Type="System.Int32" DbType="Int NOT NULL" IsPrimaryKey="true" CanBeNull="false" />
-      <Column Name="FunctionCallId" Type="System.Int32" DbType="Int NOT NULL" IsPrimaryKey="true" CanBeNull="false" />
-      <Association Name="Crash_Crash_FunctionCall" Member="Crash" ThisKey="CrashId" OtherKey="Id" Type="Crash" IsForeignKey="true" />
-      <Association Name="FunctionCall_Crash_FunctionCall" Member="FunctionCall" ThisKey="FunctionCallId" OtherKey="Id" Type="FunctionCall" IsForeignKey="true" />
-    </Type>
-  </Table>
   <Table Name="dbo.FunctionCalls" Member="FunctionCalls">
     <Type Name="FunctionCall">
       <Column Name="Id" Type="System.Int32" DbType="Int NOT NULL IDENTITY" IsPrimaryKey="true" IsDbGenerated="true" CanBeNull="false" />
       <Column Name="Call" Type="System.String" DbType="VarChar(MAX)" CanBeNull="true" />
-      <Association Name="FunctionCall_Crash_FunctionCall" Member="Crash_FunctionCalls" ThisKey="Id" OtherKey="FunctionCallId" Type="Crash_FunctionCall" />
     </Type>
   </Table>
   <Table Name="dbo.UserGroups" Member="UserGroups">
