CrashReportWWW - Improved crash view listing performance

diff --git a/Engine/Source/Programs/CrashReporter/CrashReportWebSite/Views/Crashes/ViewCrash.ascx b/Engine/Source/Programs/CrashReporter/CrashReportWebSite/Views/Crashes/ViewCrash.ascx
--- a/Engine/Source/Programs/CrashReporter/CrashReportWebSite/Views/Crashes/ViewCrash.ascx
+++ b/Engine/Source/Programs/CrashReporter/CrashReportWebSite/Views/Crashes/ViewCrash.ascx
@@ -82,54 +82,54 @@
 		<th style='width: 12em;'><%=Url.TableHeader( "Module", "Module", Model )%></th>
 	</tr>
 		<%
-		if( Model.Results.ToList() != null )
-		{
-			using( FAutoScopedLogTimer LogTimer = new FAutoScopedLogTimer( this.GetType().ToString() + "(" + Model.Results.Count() + ")" ) )
+			if( Model.Results.ToList() != null )
 			{
-		int Iteration = 0;
-		foreach( Crash CurrentCrash in ( IEnumerable )Model.Results )
-		{
-					using( FScopedLogTimer LogTimer2 = new FScopedLogTimer( "VCInnerLoop" + "(" + Iteration + ")" ) )
+				using( FAutoScopedLogTimer LogTimer = new FAutoScopedLogTimer( this.GetType().ToString() + "(" + Model.Results.Count() + ")" ) )
+				{
+					int Iteration = 0;
+					foreach( Crash CurrentCrash in (IEnumerable)Model.Results )
 					{
-			Buggs_Crash BuggCrash = null;
-			try
-			{
-							//BuggCrash = CurrentCrash.Buggs_Crashes.FirstOrDefault();
-			}
-			catch
-			{
-				BuggCrash = null;
-			}
-
-			Iteration++;
-			string CrashRowColor = "grey";
-			string CrashColorDescription = "Incoming Crash";
-
-			if( string.IsNullOrWhiteSpace( CurrentCrash.FixedChangeList ) && string.IsNullOrWhiteSpace( CurrentCrash.TTPID ) )
-			{
-				CrashRowColor = "#FFFF88"; // yellow
-				CrashColorDescription = "This crash has not been fixed or assigned a TTP";
-			}
-
-			if( !string.IsNullOrWhiteSpace( CurrentCrash.FixedChangeList ) )
-			{
-				// Green
-				CrashRowColor = "#008C00"; //green
-				CrashColorDescription = "This crash has been fixed in CL# " + CurrentCrash.FixedChangeList;
-			}
-
-			if( ( BuggCrash != null ) && !string.IsNullOrWhiteSpace( CurrentCrash.TTPID ) && string.IsNullOrWhiteSpace( CurrentCrash.FixedChangeList ) )
-			{
-				CrashRowColor = "#D01F3C"; // red
-				CrashColorDescription = "This crash has occurred more than once and been assigned a TTP: " + CurrentCrash.TTPID + " but has not been fixed.";
-			}
-
-			if( CurrentCrash.Status == "Tester" )
-			{
-				CrashRowColor = "#718698"; // red
-				CrashColorDescription = "This crash status has been set to Tester";
-			}
-			%>
+						using( FScopedLogTimer LogTimer2 = new FScopedLogTimer( "VCInnerLoop" + "(" + Iteration + ")" ) )
+						{
+							Buggs_Crash BuggCrash = null;
+							try
+							{
+								//BuggCrash = CurrentCrash.Buggs_Crashes.FirstOrDefault();
+							}
+							catch
+							{
+								BuggCrash = null;
+							}
+
+							Iteration++;
+							string CrashRowColor = "grey";
+							string CrashColorDescription = "Incoming Crash";
+
+							if( string.IsNullOrWhiteSpace( CurrentCrash.FixedChangeList ) && string.IsNullOrWhiteSpace( CurrentCrash.TTPID ) )
+							{
+								CrashRowColor = "#FFFF88"; // yellow
+								CrashColorDescription = "This crash has not been fixed or assigned a TTP";
+							}
+
+							if( !string.IsNullOrWhiteSpace( CurrentCrash.FixedChangeList ) )
+							{
+								// Green
+								CrashRowColor = "#008C00"; //green
+								CrashColorDescription = "This crash has been fixed in CL# " + CurrentCrash.FixedChangeList;
+							}
+
+							if( ( BuggCrash != null ) && !string.IsNullOrWhiteSpace( CurrentCrash.TTPID ) && string.IsNullOrWhiteSpace( CurrentCrash.FixedChangeList ) )
+							{
+								CrashRowColor = "#D01F3C"; // red
+								CrashColorDescription = "This crash has occurred more than once and been assigned a TTP: " + CurrentCrash.TTPID + " but has not been fixed.";
+							}
+
+							if( CurrentCrash.Status == "Tester" )
+							{
+								CrashRowColor = "#718698"; // red
+								CrashColorDescription = "This crash status has been set to Tester";
+							}
+		%>
 
 	<tr class='CrashRow'>
 		<td class="CrashTd" style="background-color: <%=CrashRowColor %>;" title="<%=CrashColorDescription %>">
