CrashReportWWW - Updated project file

diff --git a/Engine/Source/Programs/CrashReporter/CrashReportWebSite/CrashReportWebSite.csproj b/Engine/Source/Programs/CrashReporter/CrashReportWebSite/CrashReportWebSite.csproj
--- a/Engine/Source/Programs/CrashReporter/CrashReportWebSite/CrashReportWebSite.csproj
+++ b/Engine/Source/Programs/CrashReporter/CrashReportWebSite/CrashReportWebSite.csproj
@@ -12,7 +12,7 @@
     <AppDesignerFolder>Properties</AppDesignerFolder>
     <RootNamespace>Tools.CrashReporter.CrashReportWebSite</RootNamespace>
     <AssemblyName>CrashReportWebSite</AssemblyName>
-    <TargetFrameworkVersion>v4.0</TargetFrameworkVersion>
+    <TargetFrameworkVersion>v4.5</TargetFrameworkVersion>
     <MvcBuildViews>false</MvcBuildViews>
     <SccProjectName>
     </SccProjectName>
@@ -43,6 +43,7 @@
     <ErrorReport>prompt</ErrorReport>
     <WarningLevel>4</WarningLevel>
     <DocumentationFile>bin\CrashReportWebSite.XML</DocumentationFile>
+    <Prefer32Bit>false</Prefer32Bit>
   </PropertyGroup>
   <ItemGroup>
     <Reference Include="System" />
@@ -56,9 +57,9 @@
     <Reference Include="System.Web.ApplicationServices" />
     <Reference Include="System.Web.DynamicData" />
     <Reference Include="System.Web.Entity" />
-    <Reference Include="System.Web.Extensions" />
     <Reference Include="System.Web" />
     <Reference Include="System.Web.Abstractions" />
+    <Reference Include="System.Web.Extensions" />
     <Reference Include="System.Web.Helpers, Version=2.0.0.0, Culture=neutral, PublicKeyToken=31bf3856ad364e35, processorArchitecture=MSIL">
       <SpecificVersion>False</SpecificVersion>
       <HintPath>C:\Program Files (x86)\Microsoft ASP.NET\ASP.NET Web Pages\v1.0\Assemblies\System.Web.Helpers.dll</HintPath>
@@ -101,6 +102,7 @@
     </Compile>
     <Compile Include="Models\CrashViewModel.cs" />
     <Compile Include="Models\FormHelper.cs" />
+    <Compile Include="Models\PerformanceTimers.cs" />
     <Compile Include="Models\UsersViewModel.cs" />
     <Compile Include="Properties\Settings.Designer.cs">
       <AutoGen>True</AutoGen>
@@ -266,10 +268,12 @@
     <CodeAnalysisRuleSet>MinimumRecommendedRules.ruleset</CodeAnalysisRuleSet>
     <DocumentationFile>bin\CrashReportWebSite.XML</DocumentationFile>
     <RunCodeAnalysis>true</RunCodeAnalysis>
+    <Prefer32Bit>false</Prefer32Bit>
   </PropertyGroup>
   <PropertyGroup Condition="'$(Configuration)|$(Platform)' == 'Release|AnyCPU'">
     <OutputPath>bin\</OutputPath>
     <DocumentationFile>bin\CrashReportWebSite.XML</DocumentationFile>
+    <Prefer32Bit>false</Prefer32Bit>
   </PropertyGroup>
   <Import Project="$(MSBuildBinPath)\Microsoft.CSharp.targets" />
   <Import Project="$(VSToolsPath)\WebApplications\Microsoft.WebApplication.targets" Condition="'$(VSToolsPath)' != ''" />
