CrashReportClient - UE-1526/CrashReportClient sometimes doesn't show up on crash (set bCompileAgainstCoreUObject to true)

diff --git a/Engine/Source/Programs/CrashReporter/CrashReportClient/CrashReportClient.Build.cs b/Engine/Source/Programs/CrashReporter/CrashReportClient/CrashReportClient.Build.cs
--- a/Engine/Source/Programs/CrashReporter/CrashReportClient/CrashReportClient.Build.cs
+++ b/Engine/Source/Programs/CrashReporter/CrashReportClient/CrashReportClient.Build.cs
@@ -19,6 +19,7 @@ public class CrashReportClient : ModuleRules
 			new string[] 
 			{
 				"Core",
+				"CoreUObject",
 				"CrashDebugHelper",
 				"HTTP",
 				"Json",
diff --git a/Engine/Source/Programs/CrashReporter/CrashReportClient/CrashReportClient.Target.cs b/Engine/Source/Programs/CrashReporter/CrashReportClient/CrashReportClient.Target.cs
--- a/Engine/Source/Programs/CrashReporter/CrashReportClient/CrashReportClient.Target.cs
+++ b/Engine/Source/Programs/CrashReporter/CrashReportClient/CrashReportClient.Target.cs
@@ -76,7 +76,7 @@ public class CrashReportClientTarget : TargetRules
 
 		// CrashReportClient doesn't ever compile with the engine linked in
 		UEBuildConfiguration.bCompileAgainstEngine = false;
-		UEBuildConfiguration.bCompileAgainstCoreUObject = false;
+		UEBuildConfiguration.bCompileAgainstCoreUObject = true;
 		UEBuildConfiguration.bUseLoggingInShipping = true;
 		UEBuildConfiguration.bCompileSteamOSS = false;
 
