CrashReportWWW - Better filtering by user name or user email, should work with the analitycs database, optimized CrashRepository.PopulateUserInfo, removed more of the usused code, some documentation, updated stats

diff --git a/Engine/Source/Programs/CrashReporter/CrashReportWebSite/Controllers/BuggsController.cs b/Engine/Source/Programs/CrashReporter/CrashReportWebSite/Controllers/BuggsController.cs
--- a/Engine/Source/Programs/CrashReporter/CrashReportWebSite/Controllers/BuggsController.cs
+++ b/Engine/Source/Programs/CrashReporter/CrashReportWebSite/Controllers/BuggsController.cs
@@ -46,11 +46,11 @@ namespace Tools.CrashReporter.CrashReportWebSite.Controllers
 		/// The Show action.
 		/// </summary>
 		/// <param name="BuggsForm">The form of user data passed up from the client.</param>
-		/// <param name="id">The unique id of the Bugg.</param>
+		/// <param name="BuggId">The unique BuggId of the Bugg.</param>
 		/// <returns>The view to display a Bugg on the client.</returns>
-		public ActionResult Show( FormCollection BuggsForm, int id )
+		public ActionResult Show( FormCollection BuggsForm, int BuggId )
 		{
-			using( FAutoScopedLogTimer LogTimer = new FAutoScopedLogTimer( this.GetType().ToString() ) )
+			using( FAutoScopedLogTimer LogTimer = new FAutoScopedLogTimer( this.GetType().ToString() + "(BuggId=" + BuggId + ")" ) )
 			{
 				// Set the display properties based on the radio buttons
 				bool DisplayModuleNames = false;
@@ -89,7 +89,7 @@ namespace Tools.CrashReporter.CrashReportWebSite.Controllers
 				Bugg Bugg = new Bugg();
 
 				BuggViewModel Model = new BuggViewModel();
-				Bugg = LocalBuggRepository.GetBugg( id );
+				Bugg = LocalBuggRepository.GetBugg( BuggId );
 				if( Bugg == null )
 				{
 					return RedirectToAction( "" );
@@ -103,19 +103,19 @@ namespace Tools.CrashReporter.CrashReportWebSite.Controllers
 					if( !string.IsNullOrEmpty( BuggsForm["SetStatus"] ) )
 					{
 						Bugg.Status = BuggsForm["SetStatus"];
-						LocalCrashRepository.SetBuggStatus( Bugg.Status, id );
+						LocalCrashRepository.SetBuggStatus( Bugg.Status, BuggId );
 					}
 
 					if( !string.IsNullOrEmpty( BuggsForm["SetFixedIn"] ) )
 					{
 						Bugg.FixedChangeList = BuggsForm["SetFixedIn"];
-						LocalCrashRepository.SetBuggFixedChangeList( Bugg.FixedChangeList, id );
+						LocalCrashRepository.SetBuggFixedChangeList( Bugg.FixedChangeList, BuggId );
 					}
 
 					if( !string.IsNullOrEmpty( BuggsForm["SetTTP"] ) )
 					{
 						Bugg.TTPID = BuggsForm["SetTTP"];
-						LocalCrashRepository.SetBuggTTPID( Bugg.TTPID, id );
+						LocalCrashRepository.SetBuggTTPID( Bugg.TTPID, BuggId );
 					}
 
 					if( !string.IsNullOrEmpty( BuggsForm["Description"] ) )
@@ -147,7 +147,7 @@ namespace Tools.CrashReporter.CrashReportWebSite.Controllers
 					NewCrash.CallStackContainer = NewCrash.GetCallStack();
 				}
 
-				using( FScopedLogTimer LogTimer2 = new FScopedLogTimer( "BuggsController.Show.PopulateUserInfo" ) )
+				using( FScopedLogTimer LogTimer2 = new FScopedLogTimer( "BuggsController.Show.PopulateUserInfo" + "(BuggId=" + BuggId + ")" ) )
 				{
 					// Add in the users for each crash in the Bugg
 					foreach( Crash CrashInstance in Model.Crashes )
diff --git a/Engine/Source/Programs/CrashReporter/CrashReportWebSite/Controllers/CrashesController.cs b/Engine/Source/Programs/CrashReporter/CrashReportWebSite/Controllers/CrashesController.cs
--- a/Engine/Source/Programs/CrashReporter/CrashReportWebSite/Controllers/CrashesController.cs
+++ b/Engine/Source/Programs/CrashReporter/CrashReportWebSite/Controllers/CrashesController.cs
@@ -47,22 +47,21 @@ namespace Tools.CrashReporter.CrashReportWebSite.Controllers
 							if( !string.IsNullOrEmpty( CrashesForm["SetStatus"] ) )
 							{
 								CurrentCrash.Status = CrashesForm["SetStatus"];
-								LocalCrashRepository.SetCrashStatus( CurrentCrash.Status, Id );
 							}
 
 							if( !string.IsNullOrEmpty( CrashesForm["SetFixedIn"] ) )
 							{
 								CurrentCrash.FixedChangeList = CrashesForm["SetFixedIn"];
-								LocalCrashRepository.SetCrashFixedChangeList( CurrentCrash.FixedChangeList, Id );
 							}
 
 							if( !string.IsNullOrEmpty( CrashesForm["SetTTP"] ) )
 							{
 								CurrentCrash.TTPID = CrashesForm["SetTTP"];
-								LocalCrashRepository.SetCrashTTPID( CurrentCrash.TTPID, Id );
 							}
 						}
 					}
+
+					LocalCrashRepository.SubmitChanges();
 				}
 
 				// <STATUS>
@@ -82,21 +81,16 @@ namespace Tools.CrashReporter.CrashReportWebSite.Controllers
 		/// Show detailed information about a crash.
 		/// </summary>
 		/// <param name="CrashesForm">A form of user data passed up from the client.</param>
-		/// <param name="Id">The unique id of the crash we wish to show the details of.</param>
+		/// <param name="CrashId">The unique NewCrashId of the crash we wish to show the details of.</param>
 		/// <returns>A view to show crash details.</returns>
-		public ActionResult Show( FormCollection CrashesForm, int? Id )
+		public ActionResult Show( FormCollection CrashesForm, int CrashId )
 		{
-			using( FAutoScopedLogTimer LogTimer = new FAutoScopedLogTimer( this.GetType().ToString() ) )
+			using( FAutoScopedLogTimer LogTimer = new FAutoScopedLogTimer( this.GetType().ToString() + "(CrashId=" + CrashId + ")" ) )
 			{
-				if( !Id.HasValue )
-				{
-					return RedirectToAction( "" );
-				}
-
 				CallStackContainer CurrentCallStack = null;
 
 				// Update the selected crash based on the form contents
-				Crash CurrentCrash = LocalCrashRepository.GetCrash( Id.Value );
+				Crash CurrentCrash = LocalCrashRepository.GetCrash( CrashId );
 
 				if( CurrentCrash == null )
 				{
@@ -109,21 +103,18 @@ namespace Tools.CrashReporter.CrashReportWebSite.Controllers
 				if( !string.IsNullOrEmpty( FormValue ) )
 				{
 					CurrentCrash.Status = FormValue;
-					//LocalCrashRepository.SetCrashStatus(FormValue, Id.Value);
 				}
 
 				FormValue = CrashesForm["SetFixedIn"];
 				if( !string.IsNullOrEmpty( FormValue ) )
 				{
 					CurrentCrash.FixedChangeList = FormValue;
-					//LocalCrashRepository.SetCrashFixedChangeList(FormValue, Id.Value);
 				}
 
 				FormValue = CrashesForm["SetTTP"];
 				if( !string.IsNullOrEmpty( FormValue ) )
 				{
 					CurrentCrash.TTPID = FormValue;
-					//LocalCrashRepository.SetCrashTTPID(FormValue, Id.Value);
 				}
 
 				// Valid to set description to an empty string
@@ -131,7 +122,6 @@ namespace Tools.CrashReporter.CrashReportWebSite.Controllers
 				if( FormValue != null )
 				{
 					CurrentCrash.Description = FormValue;
-					//LocalCrashRepository.SetCrashDescription(FormValue, Id.Value);
 				}
 
 				CurrentCallStack = new CallStackContainer( CurrentCrash );
@@ -156,11 +146,11 @@ namespace Tools.CrashReporter.CrashReportWebSite.Controllers
 		/// <summary>
 		/// Add a crash passed in the payload as Xml to the database.
 		/// </summary>
-		/// <param name="id">Unused.</param>
-		/// <returns>The row id of the newly added crash.</returns>
-		public ActionResult AddCrash( int id )
+		/// <param name="NewCrashId">Unused.</param>
+		/// <returns>The row NewCrashId of the newly added crash.</returns>
+		public ActionResult AddCrash( int NewCrashId )
 		{
-			using( FAutoScopedLogTimer LogTimer = new FAutoScopedLogTimer( this.GetType().ToString() ) )
+			using( FAutoScopedLogTimer LogTimer = new FAutoScopedLogTimer( this.GetType().ToString() + "(NewCrashId=" + NewCrashId + ")" ) )
 			{
 				CrashReporterResult NewCrashResult = new CrashReporterResult();
 				NewCrashResult.ID = -1;
diff --git a/Engine/Source/Programs/CrashReporter/CrashReportWebSite/Models/BuggRepository.cs b/Engine/Source/Programs/CrashReporter/CrashReportWebSite/Models/BuggRepository.cs
--- a/Engine/Source/Programs/CrashReporter/CrashReportWebSite/Models/BuggRepository.cs
+++ b/Engine/Source/Programs/CrashReporter/CrashReportWebSite/Models/BuggRepository.cs
@@ -64,26 +64,26 @@ namespace Tools.CrashReporter.CrashReportWebSite.Models
 		/// <returns>A Bugg representing a group of crashes.</returns>
 		public Bugg GetBugg( int Id )
 		{
-			using( FAutoScopedLogTimer LogTimer = new FAutoScopedLogTimer( this.GetType().ToString() ) )
+			using( FAutoScopedLogTimer LogTimer = new FAutoScopedLogTimer( this.GetType().ToString() + "(" + Id + ")" ) )
 			{
-			Bugg Result = null;
+				Bugg Result = null;
 
-			try
-			{
-				Result =
-				(
-					from BuggDetail in BuggsDataContext.Buggs
-					where BuggDetail.Id == Id
-					select BuggDetail
-				).FirstOrDefault();
-			}
-			catch( Exception Ex )
-			{
-				Debug.WriteLine( "Exception in GetBugg: " + Ex.ToString() );
-			}
+				try
+				{
+					Result =
+					(
+						from BuggDetail in BuggsDataContext.Buggs
+						where BuggDetail.Id == Id
+						select BuggDetail
+					).FirstOrDefault();
+				}
+				catch( Exception Ex )
+				{
+					Debug.WriteLine( "Exception in GetBugg: " + Ex.ToString() );
+				}
 
-			return Result;
-		}
+				return Result;
+			}
 		}
 
 		/// <summary>
@@ -95,18 +95,18 @@ namespace Tools.CrashReporter.CrashReportWebSite.Models
 		{
 			using( FAutoScopedLogTimer LogTimer = new FAutoScopedLogTimer( this.GetType().ToString() ) )
 			{
-			try
-			{
-				FunctionCall FunctionCall = BuggsDataContext.FunctionCalls.Where( FunctionCallInstance => FunctionCallInstance.Call.Contains( FunctionCallName ) ).First();
-				return FunctionCall.Id;
-			}
-			catch( Exception Ex )
-			{
-				Debug.WriteLine( "Exception in GetFunctionCallId: " + Ex.ToString() );
-			}
+				try
+				{
+					FunctionCall FunctionCall = BuggsDataContext.FunctionCalls.Where( FunctionCallInstance => FunctionCallInstance.Call.Contains( FunctionCallName ) ).First();
+					return FunctionCall.Id;
+				}
+				catch( Exception Ex )
+				{
+					Debug.WriteLine( "Exception in GetFunctionCallId: " + Ex.ToString() );
+				}
 
-			return 0;
-		}
+				return 0;
+			}
 		}
 
 		/// <summary>
@@ -130,41 +130,26 @@ namespace Tools.CrashReporter.CrashReportWebSite.Models
 		/// <returns>A list of strings that make up a callstack.</returns>
 		public List<string> GetFunctionCalls( List<int> Ids )
 		{
-			using( FAutoScopedLogTimer LogTimer = new FAutoScopedLogTimer( this.GetType().ToString() + "(" + Ids.Count + ")" ) )
+			using( FAutoScopedLogTimer LogTimer = new FAutoScopedLogTimer( this.GetType().ToString() + "(Ids.Count=" + Ids.Count + ")" ) )
 			{
-			List<string> FunctionCalls = new List<string>();
-			try
-			{
-// 					using( FAutoScopedLogTimer LogTimer2 = new FAutoScopedLogTimer( this.GetType().ToString() + "(" + Ids.Count + ") slow" ) )
-// 					{
-// 						var Query =
-// 											(
-// 												from Id in Ids
-// 												join Function in BuggsDataContext.FunctionCalls on Id equals Function.Id
-// 												select Function.Call
-// 											);
-// 						FunctionCalls = Query.ToList();
-// 					}
-
-					using( FAutoScopedLogTimer LogTimer3 = new FAutoScopedLogTimer( this.GetType().ToString() + "(" + Ids.Count + ") fast" ) )
+				List<string> FunctionCalls = new List<string>();
+				try
+				{
+					List<FunctionCall> Funcs = BuggsDataContext.FunctionCalls.Where( FuncCall => Ids.Contains( FuncCall.Id ) ).ToList();
+					// Order by Ids
+					foreach( int Id in Ids )
 					{
-						List<FunctionCall> Funcs = BuggsDataContext.FunctionCalls.Where( FuncCall => Ids.Contains(FuncCall.Id) ).ToList();
-						// Order by Ids
-						foreach( int Id in Ids )
-						{
-							var Found = Funcs.Find( FC => FC.Id == Id );
-							FunctionCalls.Add( Found.Call );
-						}		
+						var Found = Funcs.Find( FC => FC.Id == Id );
+						FunctionCalls.Add( Found.Call );
 					}
+				}
+				catch( Exception Ex )
+				{
+					Debug.WriteLine( "Exception in GetFunctionCalls: " + Ex.ToString() );
+				}
 
+				return FunctionCalls;
 			}
-			catch( Exception Ex )
-			{
-				Debug.WriteLine( "Exception in GetFunctionCalls: " + Ex.ToString() );
-			}
-
-			return FunctionCalls;
-		}
 		}
 
 		/// <summary>
@@ -174,81 +159,82 @@ namespace Tools.CrashReporter.CrashReportWebSite.Models
 		/// <param name="Crashes">The set of crashes to add to the Bugg.</param>
 		public void UpdateBuggData( Bugg Bugg, IEnumerable<Crash> Crashes )
 		{
-			using( FAutoScopedLogTimer LogTimer = new FAutoScopedLogTimer( this.GetType().ToString() ) )
+			using( FAutoScopedLogTimer LogTimer = new FAutoScopedLogTimer( this.GetType().ToString()+ "(" + Bugg.Id + ")" ) )
 			{
+				FLogger.WriteEvent( "UpdateBuggData Bugg.Id=" + Bugg.Id );
 
-			DateTime? TimeOfFirstCrash = null;
-			DateTime? TimeOfLastCrash = null;
-			string BuildVersion = null;
-			List<int> UserNameIds = new List<int>();
-			int CrashCount = 0;
+				DateTime? TimeOfFirstCrash = null;
+				DateTime? TimeOfLastCrash = null;
+				string BuildVersion = null;
+				List<int> UserNameIds = new List<int>();
+				int CrashCount = 0;
 
-			// Set or return min date max date while we're iterating through the crashes
-			bool bHasChanges = false;
-			try
-			{
-				foreach( Crash CurrentCrash in Crashes )
+				// Set or return min date max date while we're iterating through the crashes
+				bool bHasChanges = false;
+				try
 				{
-					CrashCount++;
-					if( TimeOfFirstCrash == null || CurrentCrash.TimeOfCrash < TimeOfFirstCrash )
+					foreach( Crash CurrentCrash in Crashes )
 					{
-						TimeOfFirstCrash = CurrentCrash.TimeOfCrash;
-					}
+						CrashCount++;
+						if( TimeOfFirstCrash == null || CurrentCrash.TimeOfCrash < TimeOfFirstCrash )
+						{
+							TimeOfFirstCrash = CurrentCrash.TimeOfCrash;
+						}
 
-					if( TimeOfLastCrash == null || CurrentCrash.TimeOfCrash > TimeOfLastCrash )
-					{
-						TimeOfLastCrash = CurrentCrash.TimeOfCrash;
-					}
+						if( TimeOfLastCrash == null || CurrentCrash.TimeOfCrash > TimeOfLastCrash )
+						{
+							TimeOfLastCrash = CurrentCrash.TimeOfCrash;
+						}
 
-					if( BuildVersion == null || CurrentCrash.BuildVersion.CompareTo( BuildVersion ) > 0 )
-					{
-						BuildVersion = CurrentCrash.BuildVersion;
-					}
+						if( BuildVersion == null || CurrentCrash.BuildVersion.CompareTo( BuildVersion ) > 0 )
+						{
+							BuildVersion = CurrentCrash.BuildVersion;
+						}
 
-					// Handle user count
-					if(Bugg.Id != 0)
-					{
-						if(!UserNameIds.Contains(CurrentCrash.User.Id))
+						// Handle user count
+						if( Bugg.Id != 0 )
 						{
-							// Add username to local users variable that we will use later to set user count and associate users and user groups to the bugg.
-							UserNameIds.Add(CurrentCrash.User.Id);
+							if( !UserNameIds.Contains( CurrentCrash.User.Id ) )
+							{
+								// Add username to local users variable that we will use later to set user count and associate users and user groups to the bugg.
+								UserNameIds.Add( CurrentCrash.User.Id );
+							}
 						}
-					}
 
-					CurrentCrash.FixedChangeList = Bugg.FixedChangeList;
-					CurrentCrash.TTPID = Bugg.TTPID;
-					CurrentCrash.Status = Bugg.Status;
+						CurrentCrash.FixedChangeList = Bugg.FixedChangeList;
+						CurrentCrash.TTPID = Bugg.TTPID;
+						CurrentCrash.Status = Bugg.Status;
 
-					if( Bugg.Id != 0 )
-					{
-						bHasChanges |= LinkCrashToBugg( Bugg, CurrentCrash );
+						if( Bugg.Id != 0 )
+						{
+							bHasChanges |= LinkCrashToBugg( Bugg, CurrentCrash );
+						}
 					}
-				}
 
-				if( CrashCount > 1 )
-				{
-					Bugg.TimeOfLastCrash = TimeOfLastCrash;
-					Bugg.TimeOfFirstCrash = TimeOfFirstCrash;
-					Bugg.NumberOfUsers = UserNameIds.Count();
-					Bugg.NumberOfCrashes = CrashCount;
-					Bugg.BuildVersion = BuildVersion;
+					if( CrashCount > 1 )
+					{
+						Bugg.TimeOfLastCrash = TimeOfLastCrash;
+						Bugg.TimeOfFirstCrash = TimeOfFirstCrash;
+						Bugg.NumberOfUsers = UserNameIds.Count();
+						Bugg.NumberOfCrashes = CrashCount;
+						Bugg.BuildVersion = BuildVersion;
 
-					foreach( int UserNameId in UserNameIds )
+						foreach( int UserNameId in UserNameIds )
+						{
+							// Link user to Bugg
+							bHasChanges |= LinkUserToBugg( Bugg, UserNameId );
+						}
+					}
+
+					if( bHasChanges )
 					{
-						// Link user to Bugg
-						bHasChanges |= LinkUserToBugg( Bugg, UserNameId );
+						BuggsDataContext.SubmitChanges();
 					}
 				}
-
-				if( bHasChanges )
+				catch( Exception Ex )
 				{
-					BuggsDataContext.SubmitChanges();
+					Debug.WriteLine( "Exception in UpdateBuggData: " + Ex.ToString() );
 				}
-			}
-			catch( Exception Ex )
-			{
-				Debug.WriteLine( "Exception in UpdateBuggData: " + Ex.ToString() );
-			}
 
 			}
 		}
@@ -334,39 +320,39 @@ namespace Tools.CrashReporter.CrashReportWebSite.Models
 		/// <returns>A filtered set of Buggs.</returns>
 		public IQueryable<Bugg> Search( IQueryable<Bugg> Results, string Query )
 		{
-			using( FAutoScopedLogTimer LogTimer = new FAutoScopedLogTimer( this.GetType().ToString() ) )
+			using( FAutoScopedLogTimer LogTimer = new FAutoScopedLogTimer( this.GetType().ToString() + "(Query=" + Query + ")" ) )
 			{
-			// Also may want to revisit how we search since this could get inefficient for a big search set.
-			IQueryable<Bugg> Buggs;
-			try
-			{
-				string QueryString = HttpUtility.HtmlDecode( Query.ToString() );
-				if( QueryString == null )
+				// Also may want to revisit how we search since this could get inefficient for a big search set.
+				IQueryable<Bugg> Buggs;
+				try
 				{
-					QueryString = "";
-				}
+					string QueryString = HttpUtility.HtmlDecode( Query.ToString() );
+					if( QueryString == null )
+					{
+						QueryString = "";
+					}
 
-				// Take out terms starting with a - 
-				string[] Terms = QueryString.Split( "-, ;+".ToCharArray(), StringSplitOptions.RemoveEmptyEntries );
-				string TermsToUse = "";
-				foreach( string Term in Terms )
-				{
-					string CallId = GetFunctionCallId( Term ).ToString();
-					if( !TermsToUse.Contains( CallId ) )
+					// Take out terms starting with a - 
+					string[] Terms = QueryString.Split( "-, ;+".ToCharArray(), StringSplitOptions.RemoveEmptyEntries );
+					string TermsToUse = "";
+					foreach( string Term in Terms )
 					{
-						TermsToUse = TermsToUse + "+" + CallId;
+						string CallId = GetFunctionCallId( Term ).ToString();
+						if( !TermsToUse.Contains( CallId ) )
+						{
+							TermsToUse = TermsToUse + "+" + CallId;
+						}
 					}
-				}
 
-				Buggs = ( IQueryable<Bugg> )Results.Search( new string[] { "Pattern" }, TermsToUse.Split( "+".ToCharArray() ) );
-			}
-			catch( Exception Ex )
-			{
-				Debug.WriteLine( "Exception in Search: " + Ex.ToString() );
-				Buggs = Results;
+					Buggs = (IQueryable<Bugg>)Results.Search( new string[] { "Pattern" }, TermsToUse.Split( "+".ToCharArray() ) );
+				}
+				catch( Exception Ex )
+				{
+					Debug.WriteLine( "Exception in Search: " + Ex.ToString() );
+					Buggs = Results;
+				}
+				return Buggs;
 			}
-			return Buggs;
-		}
 		}
 
 		/// <summary>
@@ -378,78 +364,78 @@ namespace Tools.CrashReporter.CrashReportWebSite.Models
 		{
 			using( FAutoScopedLogTimer LogTimer = new FAutoScopedLogTimer( this.GetType().ToString() ) )
 			{
-			// Right now we take a Result IQueryable starting with ListAll() Buggs then we widdle away the result set by tacking on 
-			// Linq queries. Essentially it's Results.ListAll().Where().Where().Where().Where().Where().Where()
-			// Could possibly create more optimized queries when we know exactly what we're querying
-			// The downside is that if we add more parameters each query may need to be updated.... Or we just add new case statements
-			// The other downside is that there's less code reuse, but that may be worth it.
+				// Right now we take a Result IQueryable starting with ListAll() Buggs then we widdle away the result set by tacking on 
+				// Linq queries. Essentially it's Results.ListAll().Where().Where().Where().Where().Where().Where()
+				// Could possibly create more optimized queries when we know exactly what we're querying
+				// The downside is that if we add more parameters each query may need to be updated.... Or we just add new case statements
+				// The other downside is that there's less code reuse, but that may be worth it.
 
-			IQueryable<Bugg> Results = null;
-			int Skip = ( FormData.Page - 1 ) * FormData.PageSize;
-			int Take = FormData.PageSize;
+				IQueryable<Bugg> Results = null;
+				int Skip = ( FormData.Page - 1 ) * FormData.PageSize;
+				int Take = FormData.PageSize;
 
-			// Get every Bugg.
-			Results = ListAll();
+				// Get every Bugg.
+				Results = ListAll();
 
-			// Look at all Buggs that are still 'open' i.e. the last crash occurred in our date range.
-			Results = FilterByDate( Results, FormData.DateFrom, FormData.DateTo );
+				// Look at all Buggs that are still 'open' i.e. the last crash occurred in our date range.
+				Results = FilterByDate( Results, FormData.DateFrom, FormData.DateTo );
 
-			// Filter results by build version.
-			Results = FilterByBuildVersion( Results, FormData.BuildVersion );
+				// Filter results by build version.
+				Results = FilterByBuildVersion( Results, FormData.BuildVersion );
 
-			// Run at the end
-			if( !string.IsNullOrEmpty( FormData.SearchQuery ) )
-			{
-				Results = Search( Results, FormData.SearchQuery );
-			}
+				// Run at the end
+				if( !string.IsNullOrEmpty( FormData.SearchQuery ) )
+				{
+					Results = Search( Results, FormData.SearchQuery );
+				}
 
-			// Filter by Crash Type
-			if( FormData.CrashType != "All" )
-			{
-				switch( FormData.CrashType )
+				// Filter by Crash Type
+				if( FormData.CrashType != "All" )
 				{
-					case "Crashes":
-						Results = Results.Where( BuggInstance => BuggInstance.CrashType == 1 );
-						break;
-					case "Assert":
-						Results = Results.Where( BuggInstance => BuggInstance.CrashType == 2 );
-						break;
-					case "Ensure":
-						Results = Results.Where( BuggInstance => BuggInstance.CrashType == 3 );
-						break;
-					case "CrashesAsserts":
-						Results = Results.Where( BuggInstance => BuggInstance.CrashType == 1 || BuggInstance.CrashType == 2 );
-						break;
+					switch( FormData.CrashType )
+					{
+						case "Crashes":
+							Results = Results.Where( BuggInstance => BuggInstance.CrashType == 1 );
+							break;
+						case "Assert":
+							Results = Results.Where( BuggInstance => BuggInstance.CrashType == 2 );
+							break;
+						case "Ensure":
+							Results = Results.Where( BuggInstance => BuggInstance.CrashType == 3 );
+							break;
+						case "CrashesAsserts":
+							Results = Results.Where( BuggInstance => BuggInstance.CrashType == 1 || BuggInstance.CrashType == 2 );
+							break;
+					}
 				}
-			}
 
-			// Get UserGroup ResultCounts
-			Dictionary<string, int> GroupCounts = GetCountsByGroup( Results );
+				// Get UserGroup ResultCounts
+				Dictionary<string, int> GroupCounts = GetCountsByGroup( Results );
 
-			// Filter by user group if present
-			Results = FilterByUserGroup( Results, FormData.UserGroup );
+				// Filter by user group if present
+				Results = FilterByUserGroup( Results, FormData.UserGroup );
 
-			// Pass in the results and return them sorted properly
-			IEnumerable<Bugg> SortedResults = GetSortedResults( Results, FormData.SortTerm, ( FormData.SortOrder == "Descending" ), FormData.DateFrom, FormData.DateTo );
+				// Pass in the results and return them sorted properly
+				IEnumerable<Bugg> SortedResults = GetSortedResults( Results, FormData.SortTerm, ( FormData.SortOrder == "Descending" ), FormData.DateFrom, FormData.DateTo );
 
-			// Grab just the results we want to display on this page
-			SortedResults = SortedResults.Skip( Skip ).Take( Take ).ToList();
+				// Grab just the results we want to display on this page
+				SortedResults = SortedResults.Skip( Skip ).Take( Take ).ToList();
 
-			return new BuggsViewModel 
-			{
-				Results = SortedResults,
-				PagingInfo = new PagingInfo { CurrentPage = FormData.Page, PageSize = FormData.PageSize, TotalResults = Results.Count() },
-				SortTerm = FormData.SortTerm,
-				SortOrder = FormData.SortOrder,
-				UserGroup = FormData.UserGroup,
-				CrashType = FormData.CrashType,
-				SearchQuery = FormData.SearchQuery,
-				DateFrom = (long)(FormData.DateFrom - CrashesViewModel.Epoch).TotalMilliseconds,
-				DateTo = (long)(FormData.DateTo - CrashesViewModel.Epoch).TotalMilliseconds,
-				BuildVersion = FormData.BuildVersion,
-				GroupCounts = GroupCounts,
-			};
-		}
+				return new BuggsViewModel
+				{
+					Results = SortedResults,
+					PagingInfo = new PagingInfo { CurrentPage = FormData.Page, PageSize = FormData.PageSize, TotalResults = Results.Count() },
+					SortTerm = FormData.SortTerm,
+					SortOrder = FormData.SortOrder,
+					UserGroup = FormData.UserGroup,
+					CrashType = FormData.CrashType,
+					SearchQuery = FormData.SearchQuery,
+					DateFrom = (long)( FormData.DateFrom - CrashesViewModel.Epoch ).TotalMilliseconds,
+					DateTo = (long)( FormData.DateTo - CrashesViewModel.Epoch ).TotalMilliseconds,
+					BuildVersion = FormData.BuildVersion,
+					GroupCounts = GroupCounts,
+				};
+			}
 		}
 
 		/// <summary>
@@ -459,37 +445,38 @@ namespace Tools.CrashReporter.CrashReportWebSite.Models
 		/// <returns>A dictionary of user group names, and the count of Buggs for each group.</returns>
 		public Dictionary<string, int> GetCountsByGroup( IQueryable<Bugg> Buggs )
 		{
+			// @TODO yrx 2014-11-06 Optimize
 			using( FAutoScopedLogTimer LogTimer = new FAutoScopedLogTimer( this.GetType().ToString() ) )
 			{
-			Dictionary<string, int> Results = new Dictionary<string, int>();
+				Dictionary<string, int> Results = new Dictionary<string, int>();
 
-			try
-			{
-				Results =
-				(
-					from BuggDetail in Buggs
-					join BuggsUserGroupDetail in BuggsDataContext.Buggs_UserGroups on BuggDetail.Id equals BuggsUserGroupDetail.BuggId
-					join UserGroupDetail in BuggsDataContext.UserGroups on BuggsUserGroupDetail.UserGroupId equals UserGroupDetail.Id
-					group 0 by UserGroupDetail.Name into GroupCount
-					select new { Key = GroupCount.Key, Count = GroupCount.Count() }
-				).ToDictionary( x => x.Key, y => y.Count );
-
-				// Add in all groups, even though there are no crashes associated
-				IEnumerable<string> UserGroups = ( from UserGroupDetail in BuggsDataContext.UserGroups select UserGroupDetail.Name );
-				foreach( string UserGroupName in UserGroups )
+				try
 				{
-					if( !Results.Keys.Contains( UserGroupName ) )
+					Results =
+					(
+						from BuggDetail in Buggs
+						join BuggsUserGroupDetail in BuggsDataContext.Buggs_UserGroups on BuggDetail.Id equals BuggsUserGroupDetail.BuggId
+						join UserGroupDetail in BuggsDataContext.UserGroups on BuggsUserGroupDetail.UserGroupId equals UserGroupDetail.Id
+						group 0 by UserGroupDetail.Name into GroupCount
+						select new { Key = GroupCount.Key, Count = GroupCount.Count() }
+					).ToDictionary( x => x.Key, y => y.Count );
+
+					// Add in all groups, even though there are no crashes associated
+					IEnumerable<string> UserGroups = ( from UserGroupDetail in BuggsDataContext.UserGroups select UserGroupDetail.Name );
+					foreach( string UserGroupName in UserGroups )
 					{
-						Results[UserGroupName] = 0;
+						if( !Results.Keys.Contains( UserGroupName ) )
+						{
+							Results[UserGroupName] = 0;
+						}
 					}
 				}
-			}
-			catch( Exception Ex )
-			{
-				Debug.WriteLine( "Exception in GetCountsByGroup: " + Ex.ToString() );
-			}
+				catch( Exception Ex )
+				{
+					Debug.WriteLine( "Exception in GetCountsByGroup: " + Ex.ToString() );
+				}
 
-			return Results;
+				return Results;
 		}
 		}
 
@@ -502,11 +489,8 @@ namespace Tools.CrashReporter.CrashReportWebSite.Models
 		/// <returns>The set of Buggs between the earliest and latest date.</returns>
 		public IQueryable<Bugg> FilterByDate( IQueryable<Bugg> Results, DateTime DateFrom, DateTime DateTo )
 		{
-			using( FAutoScopedLogTimer LogTimer = new FAutoScopedLogTimer( this.GetType().ToString() ) )
-			{
 			IQueryable<Bugg> BuggsInTimeFrame = Results.Where( Bugg => Bugg.TimeOfLastCrash >= DateFrom && Bugg.TimeOfLastCrash <= AddOneDayToDate( DateTo ) );
-				return BuggsInTimeFrame;
-			}
+			return BuggsInTimeFrame;			
 		}
 
 		/// <summary>
@@ -517,8 +501,6 @@ namespace Tools.CrashReporter.CrashReportWebSite.Models
 		/// <returns>The set of Buggs that matches specified build version</returns>
 		public IQueryable<Bugg> FilterByBuildVersion( IQueryable<Bugg> Results, string BuildVersion )
 		{
-			using( FAutoScopedLogTimer LogTimer = new FAutoScopedLogTimer( this.GetType().ToString() ) )
-			{
 			IQueryable<Bugg> BuggsForBuildVersion = Results;
 
 			// Filter by BuildVersion
@@ -530,9 +512,8 @@ namespace Tools.CrashReporter.CrashReportWebSite.Models
 					where MyBugg.BuildVersion.Contains( BuildVersion )
 					select MyBugg
 				);
-				}			
-				return BuggsForBuildVersion;
 			}
+			return BuggsForBuildVersion;	
 		} 
 
 		/// <summary>
@@ -543,30 +524,30 @@ namespace Tools.CrashReporter.CrashReportWebSite.Models
 		/// <returns>The set of Buggs by users in the requested user group.</returns>
 		public IQueryable<Bugg> FilterByUserGroup( IQueryable<Bugg> SetOfBuggs, string UserGroup )
 		{
-			using( FAutoScopedLogTimer LogTimer = new FAutoScopedLogTimer( this.GetType().ToString() ) )
+			using( FAutoScopedLogTimer LogTimer = new FAutoScopedLogTimer( this.GetType().ToString() + "(UserGroup=" + UserGroup + ")" ) )
 			{
-			IQueryable<Bugg> NewSetOfBuggs = null;
+				IQueryable<Bugg> NewSetOfBuggs = null;
 
-			try
-			{
-				NewSetOfBuggs =
-				(
-					from BuggDetail in SetOfBuggs
-					join BuggsUserGroupDetail in BuggsDataContext.Buggs_UserGroups on BuggDetail.Id equals BuggsUserGroupDetail.BuggId
-					join UserGroupDetail in BuggsDataContext.UserGroups on BuggsUserGroupDetail.UserGroupId equals UserGroupDetail.Id
-					where UserGroupDetail.Name.Contains(UserGroup)
-					select BuggDetail
-				);
+				try
+				{
+					NewSetOfBuggs =
+					(
+						from BuggDetail in SetOfBuggs
+						join BuggsUserGroupDetail in BuggsDataContext.Buggs_UserGroups on BuggDetail.Id equals BuggsUserGroupDetail.BuggId
+						join UserGroupDetail in BuggsDataContext.UserGroups on BuggsUserGroupDetail.UserGroupId equals UserGroupDetail.Id
+						where UserGroupDetail.Name.Contains( UserGroup )
+						select BuggDetail
+					);
 
-			}
-			catch( Exception Ex )
-			{
-				Debug.WriteLine( "Exception in FilterByUserGroup: " + Ex.ToString() );
-			}
+				}
+				catch( Exception Ex )
+				{
+					Debug.WriteLine( "Exception in FilterByUserGroup: " + Ex.ToString() );
+				}
 
 
-			return NewSetOfBuggs;
-		}
+				return NewSetOfBuggs;
+			}
 		}
 
 		private DateTime AddOneDayToDate( DateTime Date )
@@ -587,81 +568,81 @@ namespace Tools.CrashReporter.CrashReportWebSite.Models
 		{
 			using( FAutoScopedLogTimer LogTimer = new FAutoScopedLogTimer( this.GetType().ToString() ) )
 			{
-			try
-			{
+				try
+				{
 
-				var IntermediateResults =
-				(
-					from BuggCrashDetail in BuggsDataContext.Buggs_Crashes
-					where BuggCrashDetail.Crash.TimeOfCrash >= DateFrom && BuggCrashDetail.Crash.TimeOfCrash <= AddOneDayToDate( DateTo )
-					group BuggCrashDetail by BuggCrashDetail.BuggId into CrashesGrouped
-					join BuggDetail in Results on CrashesGrouped.Key equals BuggDetail.Id
-					select new { Bugg = BuggDetail, Count = CrashesGrouped.Count() }
-				);
+					var IntermediateResults =
+					(
+						from BuggCrashDetail in BuggsDataContext.Buggs_Crashes
+						where BuggCrashDetail.Crash.TimeOfCrash >= DateFrom && BuggCrashDetail.Crash.TimeOfCrash <= AddOneDayToDate( DateTo )
+						group BuggCrashDetail by BuggCrashDetail.BuggId into CrashesGrouped
+						join BuggDetail in Results on CrashesGrouped.Key equals BuggDetail.Id
+						select new { Bugg = BuggDetail, Count = CrashesGrouped.Count() }
+					);
 
 					using( FScopedLogTimer LogTimer2 = new FScopedLogTimer( "BuggRepository.GetSortedResults.CrashesInTimeFrame" ) )
 					{
-				foreach( var Result in IntermediateResults )
-				{
-					Result.Bugg.CrashesInTimeFrame = Result.Count;
-				}
+						foreach( var Result in IntermediateResults )
+						{
+							Result.Bugg.CrashesInTimeFrame = Result.Count;
+						}
 					}
 
 
-				switch( SortTerm )
-				{
-					case "CrashesInTimeFrame":
-						IntermediateResults = CrashRepository.OrderBy(IntermediateResults, BuggCrashInstance => BuggCrashInstance.Count, bSortDescending);
-						break;
+					switch( SortTerm )
+					{
+						case "CrashesInTimeFrame":
+							IntermediateResults = CrashRepository.OrderBy( IntermediateResults, BuggCrashInstance => BuggCrashInstance.Count, bSortDescending );
+							break;
 
-					case "Id":
-						IntermediateResults = CrashRepository.OrderBy(IntermediateResults, BuggCrashInstance => BuggCrashInstance.Bugg.Id, bSortDescending);
-						break;
+						case "Id":
+							IntermediateResults = CrashRepository.OrderBy( IntermediateResults, BuggCrashInstance => BuggCrashInstance.Bugg.Id, bSortDescending );
+							break;
 
-					case "BuildVersion":
-						IntermediateResults = CrashRepository.OrderBy( IntermediateResults, BuggCrashInstance => BuggCrashInstance.Bugg.BuildVersion, bSortDescending );
-						break;
+						case "BuildVersion":
+							IntermediateResults = CrashRepository.OrderBy( IntermediateResults, BuggCrashInstance => BuggCrashInstance.Bugg.BuildVersion, bSortDescending );
+							break;
 
-					case "LatestCrash":
-						IntermediateResults = CrashRepository.OrderBy(IntermediateResults, BuggCrashInstance => BuggCrashInstance.Bugg.TimeOfLastCrash, bSortDescending);
-						break;
+						case "LatestCrash":
+							IntermediateResults = CrashRepository.OrderBy( IntermediateResults, BuggCrashInstance => BuggCrashInstance.Bugg.TimeOfLastCrash, bSortDescending );
+							break;
 
-					case "FirstCrash":
-						IntermediateResults = CrashRepository.OrderBy(IntermediateResults, BuggCrashInstance => BuggCrashInstance.Bugg.TimeOfFirstCrash, bSortDescending);
-						break;
+						case "FirstCrash":
+							IntermediateResults = CrashRepository.OrderBy( IntermediateResults, BuggCrashInstance => BuggCrashInstance.Bugg.TimeOfFirstCrash, bSortDescending );
+							break;
 
-					case "NumberOfCrashes":
-						IntermediateResults = CrashRepository.OrderBy(IntermediateResults, BuggCrashInstance => BuggCrashInstance.Bugg.NumberOfCrashes, bSortDescending);
-						break;
+						case "NumberOfCrashes":
+							IntermediateResults = CrashRepository.OrderBy( IntermediateResults, BuggCrashInstance => BuggCrashInstance.Bugg.NumberOfCrashes, bSortDescending );
+							break;
 
-					case "NumberOfUsers":
-						IntermediateResults = CrashRepository.OrderBy(IntermediateResults, BuggCrashInstance => BuggCrashInstance.Bugg.NumberOfUsers, bSortDescending);
-						break;
+						case "NumberOfUsers":
+							IntermediateResults = CrashRepository.OrderBy( IntermediateResults, BuggCrashInstance => BuggCrashInstance.Bugg.NumberOfUsers, bSortDescending );
+							break;
 
-					case "Pattern":
-						IntermediateResults = CrashRepository.OrderBy(IntermediateResults, BuggCrashInstance => BuggCrashInstance.Bugg.Pattern, bSortDescending);
-						break;
+						case "Pattern":
+							IntermediateResults = CrashRepository.OrderBy( IntermediateResults, BuggCrashInstance => BuggCrashInstance.Bugg.Pattern, bSortDescending );
+							break;
 
-					case "Status":
-						IntermediateResults = CrashRepository.OrderBy(IntermediateResults, BuggCrashInstance => BuggCrashInstance.Bugg.Status, bSortDescending);
-						break;
+						case "Status":
+							IntermediateResults = CrashRepository.OrderBy( IntermediateResults, BuggCrashInstance => BuggCrashInstance.Bugg.Status, bSortDescending );
+							break;
 
-					case "FixedChangeList":
-						IntermediateResults = CrashRepository.OrderBy(IntermediateResults, BuggCrashInstance => BuggCrashInstance.Bugg.FixedChangeList, bSortDescending);
-						break;
+						case "FixedChangeList":
+							IntermediateResults = CrashRepository.OrderBy( IntermediateResults, BuggCrashInstance => BuggCrashInstance.Bugg.FixedChangeList, bSortDescending );
+							break;
 
-					case "TTPID":
-						IntermediateResults = CrashRepository.OrderBy(IntermediateResults, BuggCrashInstance => BuggCrashInstance.Bugg.TTPID, bSortDescending);
-						break;
+						case "TTPID":
+							IntermediateResults = CrashRepository.OrderBy( IntermediateResults, BuggCrashInstance => BuggCrashInstance.Bugg.TTPID, bSortDescending );
+							break;
 
+					}
+					return IntermediateResults.Select( x => x.Bugg );
 				}
-				return IntermediateResults.Select( x => x.Bugg );
-			}
-			catch( Exception Ex )
-			{
-				Debug.WriteLine( "Exception in GetSortedResults: " + Ex.ToString() );
-			}
-			return Results;
+				catch( Exception Ex )
+				{
+					Debug.WriteLine( "Exception in GetSortedResults: " + Ex.ToString() );
+				}
+				return Results;
 		}
 	}
 }
diff --git a/Engine/Source/Programs/CrashReporter/CrashReportWebSite/Models/CachedDataService.cs b/Engine/Source/Programs/CrashReporter/CrashReportWebSite/Models/CachedDataService.cs
--- a/Engine/Source/Programs/CrashReporter/CrashReportWebSite/Models/CachedDataService.cs
+++ b/Engine/Source/Programs/CrashReporter/CrashReportWebSite/Models/CachedDataService.cs
@@ -70,7 +70,7 @@ namespace Tools.CrashReporter.CrashReportWebSite.Models
 		/// <returns>A parsed callstack.</returns>
 		public CallStackContainer GetCallStack( Crash CurrentCrash )
 		{
-			using( FAutoScopedLogTimer LogTimer = new FAutoScopedLogTimer( this.GetType().ToString() ) )
+			using( FAutoScopedLogTimer LogTimer = new FAutoScopedLogTimer( this.GetType().ToString() + "(CrashId=" + CurrentCrash.Id + ")" ) )
 			{
 				string Key = CacheKeyPrefix + CallstackKeyPrefix + CurrentCrash.Id;
 				CallStackContainer CallStack = (CallStackContainer)CacheInstance[Key];
@@ -116,7 +116,7 @@ namespace Tools.CrashReporter.CrashReportWebSite.Models
 		/// <returns>A list of function names.</returns>
 		public List<string> GetFunctionCalls( string Pattern )
 		{
-			using( FAutoScopedLogTimer LogTimer = new FAutoScopedLogTimer( this.GetType().ToString() ) )
+			using( FAutoScopedLogTimer LogTimer = new FAutoScopedLogTimer( this.GetType().ToString() + "(Count=" + CacheInstance.Count + ")" ) )
 			{
 				string Key = CacheKeyPrefix + FunctionCallKeyPrefix + Pattern;
 				List<string> FunctionCalls = (List<string>)CacheInstance[Key];
diff --git a/Engine/Source/Programs/CrashReporter/CrashReportWebSite/Models/CrashDataFormater.cs b/Engine/Source/Programs/CrashReporter/CrashReportWebSite/Models/CrashDataFormater.cs
--- a/Engine/Source/Programs/CrashReporter/CrashReportWebSite/Models/CrashDataFormater.cs
+++ b/Engine/Source/Programs/CrashReporter/CrashReportWebSite/Models/CrashDataFormater.cs
@@ -136,7 +136,7 @@ namespace Tools.CrashReporter.CrashReportWebSite.Models
 		/// <param name="CurrentCrash"></param>
 		public CallStackContainer( Crash CurrentCrash )
 		{
-			using( FAutoScopedLogTimer LogTimer = new FAutoScopedLogTimer( this.GetType().ToString() ) )
+			using( FAutoScopedLogTimer LogTimer = new FAutoScopedLogTimer( this.GetType().ToString() + "(CrashId=" + CurrentCrash.Id + ")" ) )
 			{
 				ParseCallStack( CurrentCrash ); 
 			}
diff --git a/Engine/Source/Programs/CrashReporter/CrashReportWebSite/Models/CrashReport.cs b/Engine/Source/Programs/CrashReporter/CrashReportWebSite/Models/CrashReport.cs
--- a/Engine/Source/Programs/CrashReporter/CrashReportWebSite/Models/CrashReport.cs
+++ b/Engine/Source/Programs/CrashReporter/CrashReportWebSite/Models/CrashReport.cs
@@ -19,6 +19,36 @@ namespace Tools.CrashReporter.CrashReportWebSite.Models
 		}
 	}
 
+	/// <summary> Users mapping. </summary>
+	public partial class UsersMapping
+	{
+		/// <summary> Helper method, display this UsersMapping as a human readable string. Debugging purpose. </summary>
+		public override string ToString()
+		{
+			return string.Format( "{0} [{1}]", UserName, UserEmail );
+		}
+	}
+
+	/// <summary> User. </summary>
+	public partial class User
+	{
+		/// <summary> Helper method, display this User as a human readable string. Debugging purpose. </summary>
+		public override string ToString()
+		{
+			return string.Format( "{0} [{1}/{2}]", UserName, Id, UserGroupId );
+		}
+	}
+
+	/// <summary> UserGroup. </summary>
+	public partial class UserGroup
+	{
+		/// <summary> Helper method, display this UserGroup as a human readable string. Debugging purpose. </summary>
+		public override string ToString()
+		{
+			return string.Format( "{0} [{1}]", Name, Id );
+		}
+	}
+
 	/// <summary>
 	/// Derived information to the default Bugg information from the database.
 	/// </summary>
@@ -95,9 +125,9 @@ namespace Tools.CrashReporter.CrashReportWebSite.Models
 		/// Return the top lines of a callstack.
 		/// </summary>
 		/// <returns>A list of callstack entries.</returns>
-		public List<string> GetFunctionCalls( /*int Max*/ )
+		public List<string> GetFunctionCalls()
 		{
-			using( FAutoScopedLogTimer LogTimer = new FAutoScopedLogTimer( this.GetType().ToString() ) )
+			using( FAutoScopedLogTimer LogTimer = new FAutoScopedLogTimer( this.GetType().ToString() + "(Id=" + this.Id + ")" ) )
 			{
 				BuggRepository LocalBuggRepository = new BuggRepository();
 				List<string> Results = LocalBuggRepository.GetFunctionCalls( Pattern );
@@ -217,7 +247,7 @@ namespace Tools.CrashReporter.CrashReportWebSite.Models
 		/// <returns>Lines of processed callstack entries.</returns>
 		public List<CallStackEntry> GetCallStackEntries( int StartIndex, int Count )
 		{
-			using( FAutoScopedLogTimer LogTimer = new FAutoScopedLogTimer( this.GetType().ToString() ) )
+			using( FAutoScopedLogTimer LogTimer = new FAutoScopedLogTimer( this.GetType().ToString() + "(Count=" + Count + ")" ) )
 			{
 				IEnumerable<CallStackEntry> Results = new List<CallStackEntry>() { new CallStackEntry() };
 
diff --git a/Engine/Source/Programs/CrashReporter/CrashReportWebSite/Models/CrashRepository.cs b/Engine/Source/Programs/CrashReporter/CrashReportWebSite/Models/CrashRepository.cs
--- a/Engine/Source/Programs/CrashReporter/CrashReportWebSite/Models/CrashRepository.cs
+++ b/Engine/Source/Programs/CrashReporter/CrashReportWebSite/Models/CrashRepository.cs
@@ -63,7 +63,7 @@ namespace Tools.CrashReporter.CrashReportWebSite.Models
 		/// <returns>The list of users in the requested user group.</returns>
 		public List<string> GetUsersForGroup( string UserGroupName )
 		{
-			using( FAutoScopedLogTimer LogTimer = new FAutoScopedLogTimer( this.GetType().ToString() ) )
+			using( FAutoScopedLogTimer LogTimer = new FAutoScopedLogTimer( this.GetType().ToString() + "(UserGroupName=" + UserGroupName + ")" ) )
 			{
 				List<string> Users = new List<string>();
 				try
@@ -93,7 +93,7 @@ namespace Tools.CrashReporter.CrashReportWebSite.Models
 		/// <returns>The crash with the requested id.</returns>
 		public Crash GetCrash( int Id )
 		{
-			using( FAutoScopedLogTimer LogTimer = new FAutoScopedLogTimer( this.GetType().ToString() ) )
+			using( FAutoScopedLogTimer LogTimer = new FAutoScopedLogTimer( this.GetType().ToString() + "(CrashId" + Id + ")" ) )
 			{
 				try
 				{
@@ -105,8 +105,6 @@ namespace Tools.CrashReporter.CrashReportWebSite.Models
 					);
 
 					return Crashes.FirstOrDefault();
-
-					//CrashRepositoryDataContext
 				}
 				catch( Exception Ex )
 				{
@@ -117,84 +115,6 @@ namespace Tools.CrashReporter.CrashReportWebSite.Models
 			}
 		}
 
-		/// <summary>
-		/// Sets the status of a crash.
-		/// </summary>
-		/// <param name="Status">The new status of the crash (from a predetermined set).</param>
-		/// <param name="Id">The id of the crash to update the status.</param>
-		public void SetCrashStatus( string Status, int Id )
-		{
-			try
-			{
-				string Query = "UPDATE Crashes SET Status = {0} WHERE Id = {1}";
-				CrashRepositoryDataContext.ExecuteCommand( Query, Status, Id );
-			}
-			catch( Exception Ex )
-			{
-				Debug.WriteLine( "Exception in SetCrashStatus: " + Ex.ToString() );
-			}
-		}
-
-		/// <summary>
-		/// Sets the changelist the crash was fixed in.
-		/// </summary>
-		/// <param name="FixedChangeList">The string determining the changelist the crash was fixed.</param>
-		/// <param name="Id">The id of the fixed crash.</param>
-		public void SetCrashFixedChangeList( string FixedChangeList, int Id )
-		{
-			try
-			{
-				string Query = "UPDATE Crashes SET FixedChangeList = {0} WHERE Id = {1}";
-				CrashRepositoryDataContext.ExecuteCommand( Query, FixedChangeList, Id );
-			}
-			catch( Exception Ex )
-			{
-				Debug.WriteLine( "Exception in SetCrashFixedChangeList: " + Ex.ToString() );
-			}
-		}
-
-		/// <summary>
-		/// Sets the TTP ID associated with the crash.
-		/// </summary>
-		/// <param name="TTPID">A string representing a TTP.</param>
-		/// <param name="Id">The id of the crash to update.</param>
-		public void SetCrashTTPID( string TTPID, int Id )
-		{
-			try
-			{
-				string Query = "UPDATE Crashes SET TTPID = {0} WHERE Id = {1}";
-				CrashRepositoryDataContext.ExecuteCommand( Query, TTPID, Id );
-			}
-			catch( Exception Ex )
-			{
-				Debug.WriteLine( "Exception in SetCrashTTPID: " + Ex.ToString() );
-			}
-		}
-
-		/// <summary>
-		/// Sets the description of the crash.
-		/// </summary>
-		/// <param name="Description">A description of the crash.</param>
-		/// <param name="Id">The id of the crash to update.</param>
-		public void SetCrashDescription( string Description, int Id )
-		{
-			try
-			{
-				//string Query = "UPDATE Crashes SET Description = {0} WHERE Id = {1}";
-				//CrashRepositoryDataContext.ExecuteCommand( Query, Description, Id );
-
-				Crash MyCrash = this.CrashRepositoryDataContext.Crashes.Single( Crash => Crash.Id == Id );
-				MyCrash.Description = Description;
-				this.CrashRepositoryDataContext.SubmitChanges();
-
-			}
-			catch( Exception Ex )
-			{
-				Debug.WriteLine( "Exception in SetDescription: " + Ex.ToString() );
-			}
-		}
-
-
 		/// <summary>
 		/// Sets the status for all crashes in a Bugg.
 		/// </summary>
@@ -264,18 +184,13 @@ namespace Tools.CrashReporter.CrashReportWebSite.Models
 		/// <param name="CrashInstance">An instance of a crash we wish to augment with additional data.</param>
 		public void PopulateUserInfo( Crash CrashInstance )
 		{
-			using( FAutoScopedLogTimer LogTimer = new FAutoScopedLogTimer( this.GetType().ToString() ) )
+			using( FAutoScopedLogTimer LogTimer = new FAutoScopedLogTimer( this.GetType().ToString() + "(CrashId=" + CrashInstance.Id + ")" ) )
 			{
 				try
 				{
-					var query = (
-						from UserDetail in CrashRepositoryDataContext.Users
-						where UserDetail.Id == CrashInstance.User.Id
-						join UserGroupDetail in CrashRepositoryDataContext.UserGroups on UserDetail.UserGroupId equals UserGroupDetail.Id
-						select UserGroupDetail.Name
-					);
-
-					CrashInstance.UserGroupName = query.First();
+					int UserGroupId = CrashInstance.User.UserGroupId;
+					var Result = CrashRepositoryDataContext.UserGroups.Where( i => i.Id == UserGroupId ).First();
+					CrashInstance.UserGroupName = Result.Name;
 
 				}
 				catch( Exception Ex )
@@ -294,16 +209,7 @@ namespace Tools.CrashReporter.CrashReportWebSite.Models
 			IQueryable<Crash> Crashes = null;
 			try
 			{
-				Crashes =
-				(
-					from CrashDetail in CrashRepositoryDataContext.Crashes
-					select CrashDetail
-				).OrderByDescending( CrashDetail => CrashDetail.TimeOfCrash );
-
-				Crashes/*IQueryable<Crash> Crashes2*/ = CrashRepositoryDataContext.Crashes.OrderByDescending( CrashDetail => CrashDetail.TimeOfCrash );
-
-				UsersMapping Me = CrashRepositoryDataContext.UsersMappings.FirstOrDefault( mappeduser => mappeduser.UserEmail == "[REDACTED_EMAIL]" );
-
+				Crashes = CrashRepositoryDataContext.Crashes.OrderByDescending( CrashDetail => CrashDetail.TimeOfCrash );
 			}
 			catch( Exception Ex )
 			{
@@ -347,30 +253,32 @@ namespace Tools.CrashReporter.CrashReportWebSite.Models
 		/// <param name="Query">The query to apply.</param>
 		/// <returns>A set of results filtered by the query.</returns>
 		public IQueryable<Crash> Search( IQueryable<Crash> Results, string Query )
-		{
-			IQueryable<Crash> Crashes = null;
-			try
+		{using( FAutoScopedLogTimer LogTimer = new FAutoScopedLogTimer( this.GetType().ToString() + "(Query=" + Query + ")" ) )
 			{
-				string[] Terms = Query.Split( "-, ;+".ToCharArray(), StringSplitOptions.RemoveEmptyEntries );
-				string TermsToUse = "";
-				foreach( string Term in Terms )
+				IQueryable<Crash> Crashes = null;
+				try
 				{
-					if( !TermsToUse.Contains( Term ) )
+					string[] Terms = Query.Split( "-, ;+".ToCharArray(), StringSplitOptions.RemoveEmptyEntries );
+					string TermsToUse = "";
+					foreach( string Term in Terms )
 					{
-						TermsToUse = TermsToUse + "+" + Term;
+						if( !TermsToUse.Contains( Term ) )
+						{
+							TermsToUse = TermsToUse + "+" + Term;
+						}
 					}
+
+					// Search the results by the search terms using IQueryable search
+					Crashes = (IQueryable<Crash>)Results.Search( TermsToUse.Split( "+".ToCharArray() ) );
+				}
+				catch( Exception Ex )
+				{
+					Debug.WriteLine( "Exception in Search: " + Ex.ToString() );
+					Crashes = Results;
 				}
 
-				// Search the results by the search terms using IQueryable search
-				Crashes = ( IQueryable<Crash> )Results.Search( TermsToUse.Split( "+".ToCharArray() ) );
-			}
-			catch( Exception Ex )
-			{
-				Debug.WriteLine( "Exception in Search: " + Ex.ToString() );
-				Crashes = Results;
+				return Crashes;
 			}
-
-			return Crashes;
 		}
 
 		/// <summary>
@@ -393,6 +301,7 @@ namespace Tools.CrashReporter.CrashReportWebSite.Models
 		{
 			using( FAutoScopedLogTimer LogTimer = new FAutoScopedLogTimer( this.GetType().ToString() ) )
 			{
+				UsersMapping UniqueUser = null;
 				IQueryable<Crash> Results = null;
 				int Skip = ( FormData.Page - 1 ) * FormData.PageSize;
 				int Take = FormData.PageSize;
@@ -403,26 +312,38 @@ namespace Tools.CrashReporter.CrashReportWebSite.Models
 				// Grab Results 
 				if( !string.IsNullOrEmpty( FormData.SearchQuery ) )
 				{
-					string DecodedQuery = HttpUtility.HtmlDecode( FormData.SearchQuery );
-					if( !string.IsNullOrEmpty( DecodedQuery ) )
-					{
-						// Check if we are looking for user name.
-						string[] Params = DecodedQuery.Split( new string[] { "User:" }, StringSplitOptions.RemoveEmptyEntries );
-						if( Params.Length >= 1 )
+					string DecodedQuery = HttpUtility.HtmlDecode( FormData.SearchQuery ).ToLower();
+					using( FScopedLogTimer LogTimer2 = new FScopedLogTimer( "CrashRepository.GetResults.FindUserFromQuery" + "(Query=" + DecodedQuery + ")" ) )
+					{						
+						if( !string.IsNullOrEmpty( DecodedQuery ) )
 						{
-							IQueryable<UsersMapping> FoundUsers = CrashRepositoryDataContext.UsersMappings.Where( mappeduser => mappeduser.UserName.Contains( Params[0] ) );
-
-							List<string> EpicIds = new List<string>( 32 );
-							foreach( UsersMapping MappedUser in FoundUsers )
+							// Check if we are looking for user name.
+							string[] Params = DecodedQuery.Split( new string[] { "user:" }, StringSplitOptions.None );
+							if( Params.Length == 2 )
 							{
-								EpicIds.Add( MappedUser.EpicAccountId );
+								/*IQueryable<UsersMapping>*/
+								// Make sure that type of [dbo].[UsersMapping] is the same as [analyticsdb-01.dmz.epicgames.net].[CrashReport].[dbo].[UsersMapping]
+								IEnumerable<UsersMapping> FoundUsers = CrashRepositoryDataContext.ExecuteQuery<UsersMapping>
+								( @"SELECT * FROM [analyticsdb-01.dmz.epicgames.net].[CrashReport].[dbo].[UsersMapping] WHERE lower(UserName) = {0} OR lower(UserEmail) = {0}", Params[1] );
+
+								foreach( UsersMapping TheUser in FoundUsers )
+								{
+									UniqueUser = TheUser;
+									break;
+								}
+								if( UniqueUser != null )
+								{
+									Results = Results.Where( CrashInstance => CrashInstance.EpicAccountId == UniqueUser.EpicAccountId );
+								}
+								else
+								{
+									Results = Results.Where( CrashInstance => CrashInstance.EpicAccountId == "SomeValueThatIsNotPresentInTheDatabase" );
+								}
+							}
+							else
+							{
+								Results = Search( Results, DecodedQuery );
 							}
-
-							Results = Results.Where( CrashInstance => EpicIds.Contains( CrashInstance.EpicAccountId ) );
-						}
-						else
-						{
-							Results = Search( Results, DecodedQuery );
 						}
 					}
 				}
@@ -526,11 +447,14 @@ namespace Tools.CrashReporter.CrashReportWebSite.Models
 				// Grab just the results we want to display on this page
 				Results = Results.Skip( Skip ).Take( Take );
 
-				// Process call stack for display
-				foreach( Crash CrashInstance in Results )
+				using( FScopedLogTimer LogTimer3 = new FScopedLogTimer( "CrashRepository.GetResults.GetCallstacks" ) )
 				{
-					// Put callstacks into an list so we can access them line by line in the view
-					CrashInstance.CallStackContainer = GetCallStack( CrashInstance );
+					// Process call stack for display
+					foreach( Crash CrashInstance in Results )
+					{
+						// Put callstacks into an list so we can access them line by line in the view
+						CrashInstance.CallStackContainer = GetCallStack( CrashInstance );
+					}
 				}
 
 				return new CrashesViewModel
@@ -547,6 +471,7 @@ namespace Tools.CrashReporter.CrashReportWebSite.Models
 					BranchName = FormData.BranchName,
 					GameName = FormData.GameName,
 					GroupCounts = GroupCounts,
+					RealUserName = UniqueUser != null ? UniqueUser.ToString() : null,
 				};
 			}
 		}
@@ -557,6 +482,7 @@ namespace Tools.CrashReporter.CrashReportWebSite.Models
 		/// <returns>A dictionary of user group names, and the count of crashes for each group.</returns>
 		public Dictionary<string, int> GetCountsByGroup()
 		{
+			// @TODO yrx 2014-11-06 Optimize?
 			using( FAutoScopedLogTimer LogTimer = new FAutoScopedLogTimer( this.GetType().ToString() ) )
 			{
 				Dictionary<string, int> Results = new Dictionary<string, int>();
@@ -602,6 +528,7 @@ namespace Tools.CrashReporter.CrashReportWebSite.Models
 		/// <returns>A dictionary of user group names, and the count of Buggs for each group.</returns>
 		public Dictionary<string, int> GetCountsByGroupFromCrashes( IQueryable<Crash> Crashes )
 		{
+			// @TODO yrx 2014-11-06 Optimize?
 			using( FAutoScopedLogTimer LogTimer = new FAutoScopedLogTimer( this.GetType().ToString() ) )
 			{
 				Dictionary<string, int> Results = new Dictionary<string, int>();
@@ -646,7 +573,7 @@ namespace Tools.CrashReporter.CrashReportWebSite.Models
 		/// <returns>A dictionary of week vs. crash count.</returns>
 		public Dictionary<DateTime, int> GetWeeklyCountsByGroup( IQueryable<Crash> Crashes, int UserGroupId, int UndefinedUserGroupId )
 		{
-			using( FAutoScopedLogTimer LogTimer = new FAutoScopedLogTimer( this.GetType().ToString() ) )
+			using( FAutoScopedLogTimer LogTimer = new FAutoScopedLogTimer( this.GetType().ToString() + "(UserGroupId=" + UserGroupId + ")" ) )
 			{
 				Dictionary<DateTime, int> Results = new Dictionary<DateTime, int>();
 
@@ -680,7 +607,7 @@ namespace Tools.CrashReporter.CrashReportWebSite.Models
 		/// <returns>A dictionary of day vs. crash count.</returns>
 		public Dictionary<DateTime, int> GetDailyCountsByGroup( IQueryable<Crash> Crashes, int UserGroupId, int UndefinedUserGroupId )
 		{
-			using( FAutoScopedLogTimer LogTimer = new FAutoScopedLogTimer( this.GetType().ToString() ) )
+			using( FAutoScopedLogTimer LogTimer = new FAutoScopedLogTimer( this.GetType().ToString() + "(UserGroupId=" + UserGroupId + ")" ) )
 			{
 				Dictionary<DateTime, int> Results = new Dictionary<DateTime, int>();
 
@@ -741,7 +668,7 @@ namespace Tools.CrashReporter.CrashReportWebSite.Models
 		/// <returns>A sorted set of crashes.</returns>
 		public IQueryable<Crash> GetSortedResults( IQueryable<Crash> Results, string SortTerm, bool bSortByDescending )
 		{
-			using( FAutoScopedLogTimer LogTimer = new FAutoScopedLogTimer( this.GetType().ToString() ) )
+			using( FAutoScopedLogTimer LogTimer = new FAutoScopedLogTimer( this.GetType().ToString() + "(SortTerm=" + SortTerm + ")" ) )
 			{
 				switch( SortTerm )
 				{
@@ -857,7 +784,7 @@ namespace Tools.CrashReporter.CrashReportWebSite.Models
 		/// <remarks>All user group interaction is done this way to remove any dependencies on pre-populated tables.</remarks>
 		public int FindOrAddUserGroup( string UserGroupName )
 		{
-			using( FAutoScopedLogTimer LogTimer = new FAutoScopedLogTimer( this.GetType().ToString() ) )
+			using( FAutoScopedLogTimer LogTimer = new FAutoScopedLogTimer( this.GetType().ToString() + "(UserGroupName=" + UserGroupName + ")" ) )
 			{
 				int UserGroupNameId = 0;
 				try
diff --git a/Engine/Source/Programs/CrashReporter/CrashReportWebSite/Models/CrashesViewModel.cs b/Engine/Source/Programs/CrashReporter/CrashReportWebSite/Models/CrashesViewModel.cs
--- a/Engine/Source/Programs/CrashReporter/CrashReportWebSite/Models/CrashesViewModel.cs
+++ b/Engine/Source/Programs/CrashReporter/CrashReportWebSite/Models/CrashesViewModel.cs
@@ -54,7 +54,11 @@ namespace Tools.CrashReporter.CrashReportWebSite.Models
 
 		/// <summary>The user supplied parameters.</summary>
 		public FormCollection FormCollection { get; set; }
+
 		/// <summary>The set of statuses a crash could have its status set to.</summary>
 		public IEnumerable<string> SetStatus { get { return new List<string>( new string[] { "Unset", "Reviewed", "New", "Coder", "EngineQA", "GameQA" } ); } }
+
+		/// <summary> The real user name, displayed only for searches. </summary>
+		public string RealUserName { get; set; }
 	}
 }
\ No newline at end of file
diff --git a/Engine/Source/Programs/CrashReporter/CrashReportWebSite/Models/FormHelper.cs b/Engine/Source/Programs/CrashReporter/CrashReportWebSite/Models/FormHelper.cs
--- a/Engine/Source/Programs/CrashReporter/CrashReportWebSite/Models/FormHelper.cs
+++ b/Engine/Source/Programs/CrashReporter/CrashReportWebSite/Models/FormHelper.cs
@@ -97,7 +97,9 @@ namespace Tools.CrashReporter.CrashReportWebSite.Models
 		{
 			string MillisecondsString;
 			if (!GetFormParameter(Request, Form, Key, "", out MillisecondsString))
-				return;
+			{ 
+				return; 
+			}
 
 			try
 			{
diff --git a/Engine/Source/Programs/CrashReporter/CrashReportWebSite/Views/Crashes/Index.aspx b/Engine/Source/Programs/CrashReporter/CrashReportWebSite/Views/Crashes/Index.aspx
--- a/Engine/Source/Programs/CrashReporter/CrashReportWebSite/Views/Crashes/Index.aspx
+++ b/Engine/Source/Programs/CrashReporter/CrashReportWebSite/Views/Crashes/Index.aspx
@@ -111,7 +111,10 @@ Crash Reports
 			<input type="radio" name="CrashType" class="CrashType" value="All" <%=( Model.CrashType == "All" ) ? "checked='checked'" : "" %>/> <span title='All Crashes'>All</span>
 		</span>
 
-		<div id="SearchBox"><%=Html.TextBox( "SearchQuery", Model.SearchQuery, new { width = "1000" } )%><input type="submit" value="Search" class='SearchButton' title="For searching for an user use 'user:[name]'" /></div>
+		<div id="SearchBox">
+			<input id="SearchQuery" name="SearchQuery" type="text" value="<%=Model.SearchQuery %>" width="1000" title="For searching for an user use 'user:[name]'" />
+			<input type="submit" value="Search" class='SearchButton' />
+		</div>
 	
 		<script>$.datepicker.setDefaults($.datepicker.regional['']);</script>
 
@@ -157,7 +160,8 @@ Crash Reports
 						UserGroup = Model.UserGroup, 
 						DateFrom = Model.DateFrom, 
 						DateTo = Model.DateTo, 
-						CrashType = Model.CrashType
+						CrashType = Model.CrashType,
+						RealUserName = Model.RealUserName
 					}
 					, 
 					new { style = "color:black; text-decoration:none;" }
@@ -193,7 +197,8 @@ Crash Reports
 			UserGroup = Model.UserGroup, 
 			DateFrom = Model.DateFrom, 
 			DateTo = Model.DateTo, 
-			CrashType = Model.CrashType 
+			CrashType = Model.CrashType,
+			RealUserName = Model.RealUserName
 		} 
 	) )%>
 	<div id="clear"></div>
diff --git a/Engine/Source/Programs/CrashReporter/CrashReportWebSite/Views/Crashes/Show.aspx b/Engine/Source/Programs/CrashReporter/CrashReportWebSite/Views/Crashes/Show.aspx
--- a/Engine/Source/Programs/CrashReporter/CrashReportWebSite/Views/Crashes/Show.aspx
+++ b/Engine/Source/Programs/CrashReporter/CrashReportWebSite/Views/Crashes/Show.aspx
@@ -168,7 +168,9 @@
 				<dd class='even' style='width:8em'><%=Model.Crash.GetCrashTypeAsString()%></dd>
 
 			<dt>User</dt>
-				<dd><%=Html.DisplayFor(m => Model.Crash.User.UserName) %></dd>
+				<dd>
+					<%=Html.DisplayFor(m => Model.Crash.UserName) %>
+				</dd>
 
 			<dt>User Group</dt>
 				<dd><%=Html.DisplayFor(m => Model.Crash.UserGroupName) %></dd>
diff --git a/Engine/Source/Programs/CrashReporter/CrashReportWebSite/Views/Crashes/ViewCrash.ascx b/Engine/Source/Programs/CrashReporter/CrashReportWebSite/Views/Crashes/ViewCrash.ascx
--- a/Engine/Source/Programs/CrashReporter/CrashReportWebSite/Views/Crashes/ViewCrash.ascx
+++ b/Engine/Source/Programs/CrashReporter/CrashReportWebSite/Views/Crashes/ViewCrash.ascx
@@ -89,11 +89,12 @@
 					int Iteration = 0;
 					foreach( Crash CurrentCrash in (IEnumerable)Model.Results )
 					{
-						using( FScopedLogTimer LogTimer2 = new FScopedLogTimer( "VCInnerLoop" + "(" + Iteration + ")" ) )
+						using( FScopedLogTimer LogTimer2 = new FScopedLogTimer( "CurrentCrash" + "(" + CurrentCrash.Id + ")" ) )
 						{
 							Buggs_Crash BuggCrash = null;
 							try
 							{
+								// This is veeery slow.
 								//BuggCrash = CurrentCrash.Buggs_Crashes.FirstOrDefault();
 							}
 							catch
@@ -131,61 +132,63 @@
 							}
 		%>
 
-	<tr class='CrashRow'>
-		<td class="CrashTd" style="background-color: <%=CrashRowColor %>;" title="<%=CrashColorDescription %>">
-			<input type="CHECKBOX" value="<%=Iteration %>" name="<%=CurrentCrash.Id%>" id="<%=CurrentCrash.Id %>" class='input CrashCheckBox'></td>
-		<td class="Id"><%=Html.ActionLink(CurrentCrash.Id.ToString(), "Show", new { controller = "crashes", id = CurrentCrash.Id }, null)%>
-			<br />
-			<em style="font-size: x-small;"><%= ( BuggCrash != null ) ? BuggCrash.BuggId.ToString() : ""  %></em></td>
-				<td class="TimeOfCrash">
-					<%=CurrentCrash.GetTimeOfCrash()[0]%>
-					<br />
-					<%=CurrentCrash.GetTimeOfCrash()[1]%>
-				</td>
-		<td class="Username"><%=CurrentCrash.UserName%></td>
-				<td class="CallStack" >
-					<div style="clip : auto; ">
-						<div id='<%=CurrentCrash.Id %>-TrimmedCallStackBox' class='TrimmedCallStackBox'>
-							<% foreach( CallStackEntry Entry in CurrentCrash.GetCallStackEntries( 0, 4 ) )
-							   { %>
-								<span class="function-name">
-									<%=Url.CallStackSearchLink( Html.Encode( Entry.GetTrimmedFunctionName( 45 ) ), Model )%>
-					</span>
-					<br />
-							<% } %>
-						</div>
-						<a class='FullCallStackTrigger'><span class='FullCallStackTriggerText'>Full Callstack</span>
-							<div id='<%=CurrentCrash.Id %>-FullCallStackBox' class='FullCallStackBox'>
-								<% foreach( CallStackEntry Entry in CurrentCrash.GetCallStackEntries( 0, 40 ) )
-								   {%>
-										<span class="FunctionName">
-											<%=Html.Encode( Entry.GetTrimmedFunctionName( 60 ) )%>
-										</span>
+								<tr class='CrashRow'>
+									<td class="CrashTd" style="background-color: <%=CrashRowColor %>;" title="<%=CrashColorDescription %>">
+										<input type="CHECKBOX" value="<%=Iteration %>" name="<%=CurrentCrash.Id%>" id="<%=CurrentCrash.Id %>" class='input CrashCheckBox'></td>
+									<td class="Id"><%=Html.ActionLink(CurrentCrash.Id.ToString(), "Show", new { controller = "crashes", id = CurrentCrash.Id }, null)%>
 										<br />
-								<% } %>
-							</div>
-						</a>
-					</div>
-
-				</td>
-
-				<td class="Game"><%=CurrentCrash.GameName%></td>
-				<td class="Mode"><%=CurrentCrash.EngineMode%></td>
-				<td class="FixedChangeList"><%=CurrentCrash.FixedChangeList%></td>
-				<td class="Ttp"><%=CurrentCrash.TTPID%>&nbsp;</td>
-				<td class="Branch"><%=CurrentCrash.Branch%>&nbsp;</td>
-				<td class="Description"><span class="TableData"><%=CurrentCrash.Description%>&nbsp;</span></td>
-				<td class="Summary"><%=CurrentCrash.Summary%></td>
-				<td class="ChangeListVersion"><%=CurrentCrash.ChangeListVersion%></td>
-				<td class="Computer"><%=CurrentCrash.ComputerName%></td>
-				<td class="Platform"><%=CurrentCrash.PlatformName%></td>
-				<td class="Status"><%=CurrentCrash.Status%></td>
-				<td class="Module"><%=CurrentCrash.Module%></td>
-			</tr>
-	<% 
-			        }
-		        }
-	        }
-	    } 
-	%>
+										<em style="font-size: x-small;"><%= ( BuggCrash != null ) ? BuggCrash.BuggId.ToString() : ""  %></em></td>
+											<td class="TimeOfCrash">
+												<%=CurrentCrash.GetTimeOfCrash()[0]%>
+												<br />
+												<%=CurrentCrash.GetTimeOfCrash()[1]%>
+											</td>
+									<td class="Username">
+										<%=Model.RealUserName != null ? Model.RealUserName : CurrentCrash.UserName%>
+									</td>
+									<td class="CallStack">
+										<div style="clip: auto;">
+											<div id='<%=CurrentCrash.Id %>-TrimmedCallStackBox' class='TrimmedCallStackBox'>
+												<% foreach( CallStackEntry Entry in CurrentCrash.GetCallStackEntries( 0, 4 ) )
+												{ %>
+													<span class="function-name">
+														<%=Url.CallStackSearchLink( Html.Encode( Entry.GetTrimmedFunctionName( 45 ) ), Model )%>
+													</span>
+													<br />
+												<% } %>
+											</div>
+											<a class='FullCallStackTrigger'><span class='FullCallStackTriggerText'>Full Callstack</span>
+												<div id='<%=CurrentCrash.Id %>-FullCallStackBox' class='FullCallStackBox'>
+													<% foreach( CallStackEntry Entry in CurrentCrash.GetCallStackEntries( 0, 40 ) )
+													{%>
+														<span class="FunctionName">
+															<%=Html.Encode( Entry.GetTrimmedFunctionName( 60 ) )%>
+														</span>
+														<br />
+													<% } %>
+												</div>
+											</a>
+										</div>
+
+									</td>
+
+									<td class="Game"><%=CurrentCrash.GameName%></td>
+									<td class="Mode"><%=CurrentCrash.EngineMode%></td>
+									<td class="FixedChangeList"><%=CurrentCrash.FixedChangeList%></td>
+									<td class="Ttp"><%=CurrentCrash.TTPID%>&nbsp;</td>
+									<td class="Branch"><%=CurrentCrash.Branch%>&nbsp;</td>
+									<td class="Description"><span class="TableData"><%=CurrentCrash.Description%>&nbsp;</span></td>
+									<td class="Summary"><%=CurrentCrash.Summary%></td>
+									<td class="ChangeListVersion"><%=CurrentCrash.ChangeListVersion%></td>
+									<td class="Computer"><%=CurrentCrash.ComputerName%></td>
+									<td class="Platform"><%=CurrentCrash.PlatformName%></td>
+									<td class="Status"><%=CurrentCrash.Status%></td>
+									<td class="Module"><%=CurrentCrash.Module%></td>
+								</tr>
+		<% 
+						}
+					}
+				}
+			} 
+		%>
 </table>
