CrashReportWWW - Fix

diff --git a/Engine/Source/Programs/CrashReporter/CrashReportWebSite/Controllers/BuggsController.cs b/Engine/Source/Programs/CrashReporter/CrashReportWebSite/Controllers/BuggsController.cs
--- a/Engine/Source/Programs/CrashReporter/CrashReportWebSite/Controllers/BuggsController.cs
+++ b/Engine/Source/Programs/CrashReporter/CrashReportWebSite/Controllers/BuggsController.cs
@@ -46,11 +46,11 @@ namespace Tools.CrashReporter.CrashReportWebSite.Controllers
 		/// The Show action.
 		/// </summary>
 		/// <param name="BuggsForm">The form of user data passed up from the client.</param>
-		/// <param name="BuggId">The unique BuggId of the Bugg.</param>
+		/// <param name="id">The unique id of the Bugg.</param>
 		/// <returns>The view to display a Bugg on the client.</returns>
-		public ActionResult Show( FormCollection BuggsForm, int BuggId )
+		public ActionResult Show( FormCollection BuggsForm, int id )
 		{
-			using( FAutoScopedLogTimer LogTimer = new FAutoScopedLogTimer( this.GetType().ToString() + "(BuggId=" + BuggId + ")" ) )
+			using( FAutoScopedLogTimer LogTimer = new FAutoScopedLogTimer( this.GetType().ToString() + "(BuggId=" + id + ")" ) )
 			{
 				// Set the display properties based on the radio buttons
 				bool DisplayModuleNames = false;
@@ -89,7 +89,7 @@ namespace Tools.CrashReporter.CrashReportWebSite.Controllers
 				Bugg Bugg = new Bugg();
 
 				BuggViewModel Model = new BuggViewModel();
-				Bugg = LocalBuggRepository.GetBugg( BuggId );
+				Bugg = LocalBuggRepository.GetBugg( id );
 				if( Bugg == null )
 				{
 					return RedirectToAction( "" );
@@ -103,19 +103,19 @@ namespace Tools.CrashReporter.CrashReportWebSite.Controllers
 					if( !string.IsNullOrEmpty( BuggsForm["SetStatus"] ) )
 					{
 						Bugg.Status = BuggsForm["SetStatus"];
-						LocalCrashRepository.SetBuggStatus( Bugg.Status, BuggId );
+						LocalCrashRepository.SetBuggStatus( Bugg.Status, id );
 					}
 
 					if( !string.IsNullOrEmpty( BuggsForm["SetFixedIn"] ) )
 					{
 						Bugg.FixedChangeList = BuggsForm["SetFixedIn"];
-						LocalCrashRepository.SetBuggFixedChangeList( Bugg.FixedChangeList, BuggId );
+						LocalCrashRepository.SetBuggFixedChangeList( Bugg.FixedChangeList, id );
 					}
 
 					if( !string.IsNullOrEmpty( BuggsForm["SetTTP"] ) )
 					{
 						Bugg.TTPID = BuggsForm["SetTTP"];
-						LocalCrashRepository.SetBuggTTPID( Bugg.TTPID, BuggId );
+						LocalCrashRepository.SetBuggTTPID( Bugg.TTPID, id );
 					}
 
 					if( !string.IsNullOrEmpty( BuggsForm["Description"] ) )
@@ -147,7 +147,7 @@ namespace Tools.CrashReporter.CrashReportWebSite.Controllers
 					NewCrash.CallStackContainer = NewCrash.GetCallStack();
 				}
 
-				using( FScopedLogTimer LogTimer2 = new FScopedLogTimer( "BuggsController.Show.PopulateUserInfo" + "(BuggId=" + BuggId + ")" ) )
+				using( FScopedLogTimer LogTimer2 = new FScopedLogTimer( "BuggsController.Show.PopulateUserInfo" + "(id=" + id + ")" ) )
 				{
 					// Add in the users for each crash in the Bugg
 					foreach( Crash CrashInstance in Model.Crashes )
diff --git a/Engine/Source/Programs/CrashReporter/CrashReportWebSite/Controllers/CrashesController.cs b/Engine/Source/Programs/CrashReporter/CrashReportWebSite/Controllers/CrashesController.cs
--- a/Engine/Source/Programs/CrashReporter/CrashReportWebSite/Controllers/CrashesController.cs
+++ b/Engine/Source/Programs/CrashReporter/CrashReportWebSite/Controllers/CrashesController.cs
@@ -81,16 +81,16 @@ namespace Tools.CrashReporter.CrashReportWebSite.Controllers
 		/// Show detailed information about a crash.
 		/// </summary>
 		/// <param name="CrashesForm">A form of user data passed up from the client.</param>
-		/// <param name="CrashId">The unique NewCrashId of the crash we wish to show the details of.</param>
+		/// <param name="CrashId">The unique id of the crash we wish to show the details of.</param>
 		/// <returns>A view to show crash details.</returns>
-		public ActionResult Show( FormCollection CrashesForm, int CrashId )
+		public ActionResult Show( FormCollection CrashesForm, int id )
 		{
-			using( FAutoScopedLogTimer LogTimer = new FAutoScopedLogTimer( this.GetType().ToString() + "(CrashId=" + CrashId + ")" ) )
+			using( FAutoScopedLogTimer LogTimer = new FAutoScopedLogTimer( this.GetType().ToString() + "(CrashId=" + id + ")" ) )
 			{
 				CallStackContainer CurrentCallStack = null;
 
 				// Update the selected crash based on the form contents
-				Crash CurrentCrash = LocalCrashRepository.GetCrash( CrashId );
+				Crash CurrentCrash = LocalCrashRepository.GetCrash( id );
 
 				if( CurrentCrash == null )
 				{
@@ -146,11 +146,11 @@ namespace Tools.CrashReporter.CrashReportWebSite.Controllers
 		/// <summary>
 		/// Add a crash passed in the payload as Xml to the database.
 		/// </summary>
-		/// <param name="NewCrashId">Unused.</param>
-		/// <returns>The row NewCrashId of the newly added crash.</returns>
-		public ActionResult AddCrash( int NewCrashId )
+		/// <param name="id">Unused.</param>
+		/// <returns>The row id of the newly added crash.</returns>
+		public ActionResult AddCrash( int id )
 		{
-			using( FAutoScopedLogTimer LogTimer = new FAutoScopedLogTimer( this.GetType().ToString() + "(NewCrashId=" + NewCrashId + ")" ) )
+			using( FAutoScopedLogTimer LogTimer = new FAutoScopedLogTimer( this.GetType().ToString() + "(NewCrashId=" + id + ")" ) )
 			{
 				CrashReporterResult NewCrashResult = new CrashReporterResult();
 				NewCrashResult.ID = -1;
diff --git a/Engine/Source/Programs/CrashReporter/CrashReportWebSite/Views/Buggs/ViewBuggs.ascx b/Engine/Source/Programs/CrashReporter/CrashReportWebSite/Views/Buggs/ViewBuggs.ascx
--- a/Engine/Source/Programs/CrashReporter/CrashReportWebSite/Views/Buggs/ViewBuggs.ascx
+++ b/Engine/Source/Programs/CrashReporter/CrashReportWebSite/Views/Buggs/ViewBuggs.ascx
@@ -58,7 +58,7 @@
 			{
 				foreach( Bugg CurrentBugg in Model.Results )
 				{
-					using( FScopedLogTimer LogTimer2 = new FScopedLogTimer( "VBInnerLoop" ) )
+					using( FScopedLogTimer LogTimer2 = new FScopedLogTimer( "CurrentBugg" + "(" + CurrentBugg.Id + ")" ) )
 					{ 
 						string BuggRowColor = "grey";
 						string BuggColorDescription = "Incoming CrashGroup";
