Temporary fix in crashreporter website to stop issues where users have [] squarebrackets in their file system path

diff --git a/Engine/Source/Programs/CrashReporter/CrashReportWebSite/Models/CrashDataFormater.cs b/Engine/Source/Programs/CrashReporter/CrashReportWebSite/Models/CrashDataFormater.cs
--- a/Engine/Source/Programs/CrashReporter/CrashReportWebSite/Models/CrashDataFormater.cs
+++ b/Engine/Source/Programs/CrashReporter/CrashReportWebSite/Models/CrashDataFormater.cs
@@ -312,7 +312,7 @@ namespace Tools.CrashReporter.CrashReportWebSite.Models
 				int PlusOffset = CurrentLine.IndexOf( " + " );
 				int BytesOffset = CurrentLine.IndexOf( " bytes" );
 				int OpenBracketOffset = CurrentLine.IndexOf( '[' );
-				int CloseBracketOffset = CurrentLine.IndexOf( ']' );
+				int CloseBracketOffset = CurrentLine.LastIndexOf( ']' );
 
 				// Parse out the juicy info from the line of the callstack
 				if( PlingOffset > 0 )
@@ -330,7 +330,8 @@ namespace Tools.CrashReporter.CrashReportWebSite.Models
 
 							FilePath = FileLinePath.TrimEnd( "0123456789:".ToCharArray() );
 							int SourceLine = 0;
-							if( int.TryParse( FileLinePath.Substring( FilePath.Length + 1 ), out SourceLine ) )
+						    Debug.Assert(FilePath.Length < FileLinePath.Length,"WRONG SIZE");
+							if( int.TryParse( FileLinePath.Substring(FilePath.Length + 1 ), out SourceLine ) )
 							{
 								LineNumber = SourceLine;
 							}
