CrashReportWebsite - Optimized GetCountsByGroupFromCrashes, down from time-out to 1 second

diff --git a/Engine/Source/Programs/CrashReporter/CrashReportWebSite/Models/CrashRepository.cs b/Engine/Source/Programs/CrashReporter/CrashReportWebSite/Models/CrashRepository.cs
--- a/Engine/Source/Programs/CrashReporter/CrashReportWebSite/Models/CrashRepository.cs
+++ b/Engine/Source/Programs/CrashReporter/CrashReportWebSite/Models/CrashRepository.cs
@@ -524,35 +524,39 @@ namespace Tools.CrashReporter.CrashReportWebSite.Models
 		/// <summary>
 		/// Gets a container of crash counts per user group for the set of crashes passed in.
 		/// </summary>
-		/// <param name="Crashes">The set of crashes to tabulate by user group.</param>
+		/// <param name="QueryableCrashes">The set of crashes to tabulate by user group.</param>
 		/// <returns>A dictionary of user group names, and the count of Buggs for each group.</returns>
-		public Dictionary<string, int> GetCountsByGroupFromCrashes( IQueryable<Crash> Crashes )
+		public Dictionary<string, int> GetCountsByGroupFromCrashes( IQueryable<Crash> QueryableCrashes )
 		{
-			// @TODO yrx 2014-11-06 Optimize?
 			using( FAutoScopedLogTimer LogTimer = new FAutoScopedLogTimer( this.GetType().ToString() ) )
 			{
 				Dictionary<string, int> Results = new Dictionary<string, int>();
 
 				try
 				{
-					Results =
-					(
-						from CrashDetail in Crashes
-						from UserDetail in CrashRepositoryDataContext.Users
-						join UserGroupDetail in CrashRepositoryDataContext.UserGroups on UserDetail.UserGroupId equals UserGroupDetail.Id
-						where CrashDetail.UserNameId == UserDetail.Id || CrashDetail.UserName == UserDetail.UserName
-						group CrashDetail by UserGroupDetail.Name into GroupCount
-						select new { Key = GroupCount.Key, Count = GroupCount.Count() }
-					).ToDictionary( x => x.Key, y => y.Count );
+					var UsersIDsAndGroupIDs = CrashRepositoryDataContext.Users.Select( User => new { UserId = User.Id, UserGroupId = User.UserGroupId } ).ToList();
+					var UserGroupArray = CrashRepositoryDataContext.UserGroups.ToList();
+					UserGroupArray.Sort( ( UG1, UG2 ) => UG1.Name.CompareTo( UG2.Name ) ); 
 
-					// Add in all groups, even though there are no crashes associated
-					IEnumerable<string> UserGroups = ( from UserGroupDetail in CrashRepositoryDataContext.UserGroups select UserGroupDetail.Name );
-					foreach( string UserGroupName in UserGroups )
+					// Initialize all groups to 0.
+					foreach( var UserGroup in UserGroupArray )
 					{
-						if( !Results.Keys.Contains( UserGroupName ) )
-						{
-							Results[UserGroupName] = 0;
-						}
+						Results.Add( UserGroup.Name, 0 );
+					}
+
+					Dictionary<int, string> UserIdToGroupName = new Dictionary<int, string>();
+					foreach( var UserIds in UsersIDsAndGroupIDs )
+					{
+						// Find group name for the user id.
+						string UserGroupName = UserGroupArray.Where( UG => UG.Id == UserIds.UserGroupId ).First().Name;
+						UserIdToGroupName.Add( UserIds.UserId, UserGroupName );
+					}
+
+					List<int> UserIdCrashes = QueryableCrashes.Select( Crash => Crash.UserNameId.Value ).ToList();
+					foreach( int UserId in UserIdCrashes )
+					{
+						string UserGroupName = UserIdToGroupName[UserId];
+						Results[UserGroupName]++;
 					}
 				}
 				catch( Exception Ex )
