Changed titles to reflect move from TTP to Jira. Titles now read "Bug"

diff --git a/Engine/Source/Programs/CrashReporter/CrashReportWebSite/Views/Crashes/Show.aspx b/Engine/Source/Programs/CrashReporter/CrashReportWebSite/Views/Crashes/Show.aspx
--- a/Engine/Source/Programs/CrashReporter/CrashReportWebSite/Views/Crashes/Show.aspx
+++ b/Engine/Source/Programs/CrashReporter/CrashReportWebSite/Views/Crashes/Show.aspx
@@ -77,7 +77,7 @@
 			</select>
 			<input type="submit" name="SetStatusSubmit" value="Set" class="SetButtonGreen" />
 
-			<span id="set-ttp" style="">TTP</span>
+			<span id="set-ttp" style="">Bug</span>
 			<input name="SetTTP" type="text" id="ttp-text-box" />
 			<input type="submit" name="SetTTPSubmit" value="Set" class="SetButtonGreen" />
 
diff --git a/Engine/Source/Programs/CrashReporter/CrashReportWebSite/Views/Crashes/ViewCrash.ascx b/Engine/Source/Programs/CrashReporter/CrashReportWebSite/Views/Crashes/ViewCrash.ascx
--- a/Engine/Source/Programs/CrashReporter/CrashReportWebSite/Views/Crashes/ViewCrash.ascx
+++ b/Engine/Source/Programs/CrashReporter/CrashReportWebSite/Views/Crashes/ViewCrash.ascx
@@ -71,7 +71,7 @@
 		<th style='width: 11em;'><%=Url.TableHeader( "Game", "GameName", Model )%></th> 
 		<th style='width: 11em;'><%=Url.TableHeader( "Mode", "EngineMode", Model )%></th> 
 		<th style='width: 16em;'><%=Url.TableHeader( "FixedCL#", "FixedChangeList", Model )%></th>
-		<th style='width: 9em;'><%=Url.TableHeader( "TTP", "TTPID", Model) %></th>
+		<th style='width: 9em;'><%=Url.TableHeader( "Bug", "TTPID", Model) %></th>
 		<th style='width: 9em;'><%=Url.TableHeader( "Branch", "Branch", Model) %></th>
 		<th>Description</th>
 		<th style='width: 15em'><%=Url.TableHeader( "Message", "Summary", Model )%></th>
