CrashReport - Ignore crashes that don't contain "UE4-" in the directory name

diff --git a/Engine/Source/Programs/CrashReporter/CrashReportClient/Private/Windows/WindowsErrorReport.cpp b/Engine/Source/Programs/CrashReporter/CrashReportClient/Private/Windows/WindowsErrorReport.cpp
--- a/Engine/Source/Programs/CrashReporter/CrashReportClient/Private/Windows/WindowsErrorReport.cpp
+++ b/Engine/Source/Programs/CrashReporter/CrashReportClient/Private/Windows/WindowsErrorReport.cpp
@@ -123,11 +123,12 @@ FString FWindowsErrorReport::FindMostRecentErrorReport()
 
 	auto DirectoryModifiedTime = FDateTime::MinValue();
 	FString ReportDirectory;
-	auto ReportFinder = MakeDirectoryVisitor([&](const TCHAR* FilenameOrDirectory, bool bIsDirectory) {
+	auto ReportFinder = MakeDirectoryVisitor([&](const TCHAR* FilenameOrDirectory, bool bIsDirectory) 
+	{
 		if (bIsDirectory)
 		{
 			auto TimeStamp = PlatformFile.GetTimeStamp(FilenameOrDirectory);
-			if (TimeStamp > DirectoryModifiedTime)
+			if (TimeStamp > DirectoryModifiedTime && FCString::Strstr( FilenameOrDirectory, TEXT("UE4-") ) )
 			{
 				ReportDirectory = FilenameOrDirectory;
 				DirectoryModifiedTime = TimeStamp;
@@ -138,10 +139,7 @@ FString FWindowsErrorReport::FindMostRecentErrorReport()
 
 	TCHAR LocalAppDataPath[MAX_PATH];
 	SHGetFolderPath(0, CSIDL_LOCAL_APPDATA, NULL, 0, LocalAppDataPath);
-
-	PlatformFile.IterateDirectory(
-		*(FString(LocalAppDataPath) / TEXT("Microsoft/Windows/WER/ReportQueue")),
-		ReportFinder);
+	PlatformFile.IterateDirectory( *(FString(LocalAppDataPath) / TEXT("Microsoft/Windows/WER/ReportQueue")), ReportFinder);
 
 	return ReportDirectory;
 }
