Small fix for the machineId code

diff --git a/Engine/Source/Programs/CrashReporter/CrashReportWebSite/Models/CrashRepository.cs b/Engine/Source/Programs/CrashReporter/CrashReportWebSite/Models/CrashRepository.cs
--- a/Engine/Source/Programs/CrashReporter/CrashReportWebSite/Models/CrashRepository.cs
+++ b/Engine/Source/Programs/CrashReporter/CrashReportWebSite/Models/CrashRepository.cs
@@ -849,14 +849,13 @@ namespace Tools.CrashReporter.CrashReportWebSite.Models
 		    NewCrash.ComputerName = NewCrashInfo.MachineGuid;
 
 			// Valid MachineID and UserName, updated crash from non-UE4 release
-			if( !string.IsNullOrEmpty(NewCrashInfo.MachineGuid) && !string.IsNullOrEmpty(NewCrashInfo.UserName) )
+			if(!string.IsNullOrEmpty(NewCrashInfo.UserName))
 			{
 				NewCrash.UserNameId = CrashRepositoryDataContext.FindOrAddUser( NewCrashInfo.UserName );
 			}
 			// Valid MachineID and EpicAccountId, updated crash from UE4 release
-			else if( !string.IsNullOrEmpty( NewCrashInfo.MachineGuid ) && !string.IsNullOrEmpty( NewCrashInfo.EpicAccountId ) )
+			else if(!string.IsNullOrEmpty( NewCrashInfo.EpicAccountId ))
 			{
-				NewCrash.ComputerName = NewCrashInfo.MachineGuid;
 				NewCrash.EpicAccountId = NewCrashInfo.EpicAccountId;
 				NewCrash.UserNameId = CrashRepositoryDataContext.FindOrAddUser( UserNameAnonymous );
 			}
