Core - UE/UE-3219 - CRASH REPORTER: Unreal Frontend UFE crashes not populating on the Crash Reporter website (increased HttpSendTimeout to 5minutes and added compressing crash report files into one .ue4crash file, simplyfied a bit main loop)

diff --git a/Engine/Source/Programs/CrashReporter/CrashReportClient/Private/CrashDescription.cpp b/Engine/Source/Programs/CrashReporter/CrashReportClient/Private/CrashDescription.cpp
--- a/Engine/Source/Programs/CrashReporter/CrashReportClient/Private/CrashDescription.cpp
+++ b/Engine/Source/Programs/CrashReporter/CrashReportClient/Private/CrashDescription.cpp
@@ -78,6 +78,25 @@ FCrashDescription::FCrashDescription( FString WERXMLFilepath ) :
 				EngineVersionComponents++;
 			}
 
+			/*const FXmlNode* Parameter8Node = ProblemSignaturesNode->FindChildNode( TEXT( "Parameter8" ) );
+			if( Parameter8Node )
+			{
+				const FString Parameter8Value = Parameter8Node->GetContent();
+
+				TArray<FString> ParsedParameters8;
+				Parameter8Value.ParseIntoArray( &ParsedParameters8, TEXT( "!" ), false );
+
+				if( ParsedParameters8.Num() > 1 )
+				{
+					CommandLine = ParsedParameters8[1];
+				}
+
+				if( ParsedParameters8.Num() > 2 )
+				{
+					ErrorMessage.Add( ParsedParameters8[2] );
+				}
+			}*/
+
 			const FXmlNode* Parameter9Node = ProblemSignaturesNode->FindChildNode( TEXT( "Parameter9" ) );
 			if( Parameter9Node )
 			{
diff --git a/Engine/Source/Programs/CrashReporter/CrashReportClient/Private/CrashReportClient.cpp b/Engine/Source/Programs/CrashReporter/CrashReportClient/Private/CrashReportClient.cpp
--- a/Engine/Source/Programs/CrashReporter/CrashReportClient/Private/CrashReportClient.cpp
+++ b/Engine/Source/Programs/CrashReporter/CrashReportClient/Private/CrashReportClient.cpp
@@ -8,8 +8,12 @@
 #include "CrashReportClient.h"
 #include "UniquePtr.h"
 
+#include "TaskGraphInterfaces.h"
+
 #define LOCTEXT_NAMESPACE "CrashReportClient"
 
+//#define DO_LOCAL_TESTING 1
+
 #if	DO_LOCAL_TESTING
 	const TCHAR* GServerIP = TEXT( "http://localhost:57005" );
 #else
@@ -31,15 +35,18 @@ FCrashDescription& GetCrashDescription()
 FCrashReportClient::FCrashReportClient(const FPlatformErrorReport& InErrorReport)
 	: AppState(EApplicationState::Ready)
 	, SubmittedCountdown(-1)
-	, bDiagnosticFileSent(false)
+	, DiagnosticText( LOCTEXT("ProcessingReport", "Processing crash report ...") )
+	, DiagnoseReportTask(nullptr)
 	, ErrorReport( InErrorReport )
 	, Uploader(GServerIP)
 	, CancelButtonText(LOCTEXT("Cancel", "Don't Send"))
+	, bBeginUploadCalled(false)
 {
+
 	if (!ErrorReport.TryReadDiagnosticsFile(DiagnosticText) && !FParse::Param(FCommandLine::Get(), TEXT("no-local-diagnosis")))
 	{
-		FDiagnoseReportWorker& Worker = DiagnoseReportTask.GetTask( &DiagnosticText, GetCrashDescription().MachineId, GetCrashDescription().EpicAccountId, GetCrashDescription().UserName, &ErrorReport );
-		DiagnoseReportTask.StartBackgroundTask();
+		DiagnoseReportTask = new FAsyncTask<FDiagnoseReportWorker>( this );
+		DiagnoseReportTask->StartBackgroundTask();
 	}
 	else if( !DiagnosticText.IsEmpty() )
 	{
@@ -90,8 +97,7 @@ FText FCrashReportClient::GetCancelButtonText() const
 
 FText FCrashReportClient::GetDiagnosticText() const
 {
-	static const FText ProcessingReportText = LOCTEXT("ProcessingReport", "Processing crash report ...");
-	return DiagnoseReportTask.IsDone() ? DiagnosticText : ProcessingReportText;
+	return DiagnosticText;
 }
 
 EVisibility FCrashReportClient::SubmitButtonVisibility() const
@@ -132,8 +138,6 @@ void FCrashReportClient::StoreCommentAndUpload()
 	// Call upload even if the report is empty: pending reports will be sent if any
 	ErrorReport.SetUserComment(UserComment);
 
-	Uploader.BeginUpload(ErrorReport);
-
 	SubmittedCountdown = 5;
 	AppState = EApplicationState::CountingDown;
 	// Change the submit button text immediately (also sends diagnostics file if complete)
@@ -152,14 +156,6 @@ bool FCrashReportClient::UIWillCloseTick(float UnusedDeltaTime)
 
 	static const FText CountdownTextFormat = LOCTEXT("CloseApplication", "Close ({0})");
 
-	if( !bDiagnosticFileSent && DiagnoseReportTask.IsDone() )
-	{
-		auto DiagnosticsFilePath = ErrorReport.GetReportDirectory() / GDiagnosticsFilename;
-
-		Uploader.LocalDiagnosisComplete(FPaths::FileExists(DiagnosticsFilePath) ? DiagnosticsFilePath : TEXT(""));
-		bDiagnosticFileSent = true;	
-	}
-
 	if (bCountingDown)
 	{
 		CancelButtonText = FText::Format(CountdownTextFormat, FText::AsNumber(SubmittedCountdown));
@@ -174,9 +170,21 @@ bool FCrashReportClient::UIWillCloseTick(float UnusedDeltaTime)
 		}
 	}
 
+	// We are waiting for diagnose report task to complete.
+	if( DiagnoseReportTask && !DiagnoseReportTask->IsWorkDone() )
+	{
+		return true;
+	}
+	else if( !bBeginUploadCalled )
+	{
+		// Can be called only when we have all files.
+		Uploader.BeginUpload(ErrorReport);
+		bBeginUploadCalled = true;
+	}
+
 	// IsWorkDone will always return true here (since uploader can't finish until the diagnosis has been sent), but it
 	//  has the side effect of joining the worker thread.
-	if( !Uploader.IsFinished() || !DiagnoseReportTask.IsDone() )
+	if( !Uploader.IsFinished() )
 	{
 		// More ticks, please
 		return true;
@@ -192,12 +200,35 @@ FString FCrashReportClient::GetCrashedAppName() const
 	return GetCrashDescription().GameName;
 }
 
+void FCrashReportClient::FinalizeDiagnoseReportWorker( FText InDiagnosticText )
+{
+	DiagnosticText = InDiagnosticText;
+
+	auto DiagnosticsFilePath = ErrorReport.GetReportDirectory() / GDiagnosticsFilename;
+	Uploader.LocalDiagnosisComplete(FPaths::FileExists(DiagnosticsFilePath) ? DiagnosticsFilePath : TEXT(""));
+}
+
 #endif // !CRASH_REPORT_UNATTENDED_ONLY
 
+
+FDiagnoseReportWorker::FDiagnoseReportWorker( FCrashReportClient* InCrashReportClient ) 
+	: CrashReportClient( InCrashReportClient )		
+	, MachineId(GetCrashDescription().MachineId)
+	, EpicAccountId(GetCrashDescription().EpicAccountId)
+	, UserNameNoDot(GetCrashDescription().UserName)
+{}
+
 void FDiagnoseReportWorker::DoWork()
 {
-	const FText ReportText = ErrorReport.DiagnoseReport();
+	const FText ReportText = CrashReportClient->ErrorReport.DiagnoseReport();
 	DiagnosticText = FCrashReportUtil::FormatDiagnosticText( ReportText, MachineId, EpicAccountId, UserNameNoDot );
+
+	// Inform the game thread that we are done.
+	FSimpleDelegateGraphTask::CreateAndDispatchWhenReady
+	(
+		FSimpleDelegateGraphTask::FDelegate::CreateRaw( CrashReportClient, &FCrashReportClient::FinalizeDiagnoseReportWorker, DiagnosticText ),
+		TStatId(), nullptr, ENamedThreads::GameThread
+	);
 }
 
 FText FCrashReportUtil::FormatDiagnosticText( const FText& DiagnosticText, const FString MachineId, const FString EpicAccountId, const FString UserNameNoDot )
diff --git a/Engine/Source/Programs/CrashReporter/CrashReportClient/Private/CrashReportClient.h b/Engine/Source/Programs/CrashReporter/CrashReportClient/Private/CrashReportClient.h
--- a/Engine/Source/Programs/CrashReporter/CrashReportClient/Private/CrashReportClient.h
+++ b/Engine/Source/Programs/CrashReporter/CrashReportClient/Private/CrashReportClient.h
@@ -6,17 +6,16 @@
 #include "PlatformErrorReport.h"
 #include "CrashReportUtil.h"
 
+class FCrashReportClient;
+
 /**
  * Helper task class to process a crash report in the background
  */
-class FDiagnoseReportWorker 
+class FDiagnoseReportWorker  : public FNonAbandonableTask
 {
 public:
-	/** Diagnostics from crash report text to fill in. */
-	FText& DiagnosticText;
-
-	/** Error report files to use. */
-	const FPlatformErrorReport& ErrorReport;
+	/** Pointer to the crash report client, used to store the results. */
+	FCrashReportClient* CrashReportClient;
 
 	/** Machine ID. */
 	const FString MachineId;
@@ -27,14 +26,11 @@ public:
 	/** User name without dot. */
 	const FString UserNameNoDot;
 
+	/** Diagnostic text. */
+	FText DiagnosticText;
+
 	/** Initialization constructor. */
-	FDiagnoseReportWorker( FText* InDiagnosticText, const FString InMachineId, const FString InEpicAccountId, const FString InUserNameNoDot, const FPlatformErrorReport* InErrorReport )
-		: DiagnosticText( *InDiagnosticText )		
-		, ErrorReport( *InErrorReport )
-		, MachineId(InMachineId)
-		, EpicAccountId(InEpicAccountId)
-		, UserNameNoDot(InUserNameNoDot)
-	{}
+	FDiagnoseReportWorker( FCrashReportClient* InCrashReportClient );
 
 	/**
 	 * Do platform-specific work to get information about the crash.
@@ -57,6 +53,8 @@ public:
  */
 class FCrashReportClient : public TSharedFromThis<FCrashReportClient>
 {
+	friend class FDiagnoseReportWorker;
+
 public:
 	/**
 	 * Constructor: sets up background diagnosis
@@ -149,6 +147,9 @@ private:
 	 */
 	void StartUIWillCloseTicker();
 
+	/** Enqueued from the diagnose report worker thread to be executed on the game thread. */
+	void FinalizeDiagnoseReportWorker( FText InDiagnosticText );
+
 	/** State enum to keep track of what the app is doing */
 	struct EApplicationState
 	{
@@ -172,11 +173,8 @@ private:
 	/** Exception and call-stack to show, valid once diagnosis task is complete */
 	FText DiagnosticText;
 
-	/** Flag to remember having told the crash uploader to send the diagnostic file */
-	bool bDiagnosticFileSent;
-
 	/** Background worker to get a callstack from the report */
-	TOneShotTaskUsingDedicatedThread<FDiagnoseReportWorker> DiagnoseReportTask;
+	FAsyncTask<FDiagnoseReportWorker>* DiagnoseReportTask;
 
 	/** Platform code for accessing the report */
 	FPlatformErrorReport ErrorReport;
@@ -186,6 +184,9 @@ private:
 
 	/** Text for Cancel/Close button, including count-down value */
 	FText CancelButtonText;
+
+	/** Wheter BeginUpload has been called. */
+	bool bBeginUploadCalled;
 };
 
 #endif // !CRASH_REPORT_UNATTENDED_ONLY
diff --git a/Engine/Source/Programs/CrashReporter/CrashReportClient/Private/CrashReportClientApp.cpp b/Engine/Source/Programs/CrashReporter/CrashReportClient/Private/CrashReportClientApp.cpp
--- a/Engine/Source/Programs/CrashReporter/CrashReportClient/Private/CrashReportClientApp.cpp
+++ b/Engine/Source/Programs/CrashReporter/CrashReportClient/Private/CrashReportClientApp.cpp
@@ -124,6 +124,9 @@ void RunCrashReportClient(const TCHAR* CommandLine)
 	// Find the report to upload in the command line arguments
 	ParseCommandLine(CommandLine);
 
+	// Increase the HttpSendTimeout to 5 minutes
+	GConfig->SetFloat(TEXT("HTTP"), TEXT("HttpSendTimeout"), 5*60.0f, GEngineIni);
+
 	FPlatformErrorReport::Init();
 	auto ErrorReport = LoadErrorReport();
 	
@@ -197,9 +200,8 @@ void RunCrashReportClient(const TCHAR* CommandLine)
 #endif // !CRASH_REPORT_UNATTENDED_ONLY
 	}
 
-	FPlatformProcess::Sleep( 2.0f );
-	GEngineLoop.AppExit();
 	FPlatformErrorReport::ShutDown();
+	GEngineLoop.AppExit();
 }
 
 void CrashReportClientCheck(bool bCondition, const TCHAR* Location)
diff --git a/Engine/Source/Programs/CrashReporter/CrashReportClient/Private/CrashReportClientUnattended.cpp b/Engine/Source/Programs/CrashReporter/CrashReportClient/Private/CrashReportClientUnattended.cpp
--- a/Engine/Source/Programs/CrashReporter/CrashReportClient/Private/CrashReportClientUnattended.cpp
+++ b/Engine/Source/Programs/CrashReporter/CrashReportClient/Private/CrashReportClientUnattended.cpp
@@ -6,11 +6,9 @@
 FCrashReportClientUnattended::FCrashReportClientUnattended(const FPlatformErrorReport& ErrorReport)
 	: Uploader(GServerIP)
 {
-	Uploader.BeginUpload(ErrorReport);
-
 	// Prevent uploader waiting for a diagnosis file
-	Uploader.LocalDiagnosisSkipped();
-
+	Uploader.LocalDiagnosisSkipped();	
+	Uploader.BeginUpload(ErrorReport);
 	StartTicker();
 }
 
diff --git a/Engine/Source/Programs/CrashReporter/CrashReportClient/Private/CrashReportUtil.h b/Engine/Source/Programs/CrashReporter/CrashReportClient/Private/CrashReportUtil.h
--- a/Engine/Source/Programs/CrashReporter/CrashReportClient/Private/CrashReportUtil.h
+++ b/Engine/Source/Programs/CrashReporter/CrashReportClient/Private/CrashReportUtil.h
@@ -2,62 +2,6 @@
 
 #pragma once
 
-template <class TWorker>
-class TOneShotTaskUsingDedicatedThread : FRunnable
-{
-public:
-	/**
-	 * Constructor
-	 */
-	TOneShotTaskUsingDedicatedThread()
-		: bIsDone(1)
-	{
-	}
-
-	/**
-	 * Destructor
-	 */
-	~TOneShotTaskUsingDedicatedThread()
-	{
-		CRASHREPORTCLIENT_CHECK(!!bIsDone);
-	}
-
-
-	template<typename T1, typename T2, typename T3, typename T4, typename T5>
-	TWorker& GetTask( T1 Arg1, T2 Arg2, T3 Arg3, T4 Arg4, T5 Arg5 )
-	{
-		Worker = new TWorker( Arg1, Arg2, Arg3, Arg4, Arg5 );
-		return *Worker;
-	}
-
-	void StartBackgroundTask()
-	{
-		bIsDone = 0;
-		Thread = FRunnableThread::Create( this, Worker->Name() );
-	}
-
-	/**
-	 * No need to join the thread, since this can't be reused
-	 */
-	bool IsDone() const
-	{
-		return bIsDone == 1;
-	}
-
-private:
-	virtual uint32 Run() override
-	{
-		Worker->DoWork();
-		FPlatformAtomics::InterlockedIncrement(&bIsDone);
-		return 0;
-	}
-
-	TAutoPtr<TWorker> Worker;
-	TScopedPointer<FRunnableThread> Thread;
-	volatile int32 bIsDone;	
-};
-
-
 /**
  * Helper class for MakeDirectoryVisitor
  */
diff --git a/Engine/Source/Programs/CrashReporter/CrashReportClient/Private/CrashUpload.cpp b/Engine/Source/Programs/CrashReporter/CrashReportClient/Private/CrashUpload.cpp
--- a/Engine/Source/Programs/CrashReporter/CrashReportClient/Private/CrashUpload.cpp
+++ b/Engine/Source/Programs/CrashReporter/CrashReportClient/Private/CrashUpload.cpp
@@ -12,16 +12,15 @@
 
 namespace
 {
-const float PingTimeoutSeconds = 5.f;
-// Ignore files bigger than 100MB; mini-dumps are smaller than this, but heap dumps can be very large
-const int MaxFileSizeToUpload = 100 * 1024 * 1024;
+	const float PingTimeoutSeconds = 5.f;
+	// Ignore files bigger than 100MB; mini-dumps are smaller than this, but heap dumps can be very large
+	const int MaxFileSizeToUpload = 100 * 1024 * 1024;
 }
 
 FCrashUpload::FCrashUpload(const FString& ServerAddress)
 	: UrlPrefix(ServerAddress / "CrashReporter")
 	, State(EUploadState::NotSet)
 	, PauseState(EUploadState::Ready)
-	, bDiagnosticsFileSent(false)
 {
 	SendPingRequest();
 }
@@ -54,8 +53,7 @@ void FCrashUpload::BeginUpload(const FPlatformErrorReport& PlatformErrorReport)
 	PendingFiles = ErrorReport.GetFilesToUpload();
 	UE_LOG(CrashReportClientLog, Log, TEXT("Got %d pending files to upload from '%s'"), PendingFiles.Num(), *ErrorReport.GetReportDirectoryLeafName());
 
-	// Pause before posting completed message, to allow for additional files to be uploaded
-	PauseState = EUploadState::PostingReportComplete;
+	PauseState = EUploadState::Finished;
 	if (State == EUploadState::Ready)
 	{
 		BeginUploadImpl();
@@ -75,28 +73,12 @@ void FCrashUpload::LocalDiagnosisComplete(const FString& DiagnosticsFile)
 {
 	if (State >= EUploadState::FirstCompletedState)
 	{
-		// Must be a failure/cancelled state, or the report was a rejected by the server
+		// Must be a failure/canceled state, or the report was a rejected by the server
 		return;
 	}
 
-	bool SendDiagnosticsFile = !bDiagnosticsFileSent && !DiagnosticsFile.IsEmpty();
-
-	CRASHREPORTCLIENT_CHECK(PauseState == EUploadState::PostingReportComplete);
-	PauseState = EUploadState::Finished;
-	if (State == EUploadState::WaitingToPostReportComplete)
-	{
-		if (SendDiagnosticsFile)
-		{
-			PendingFiles.Push(DiagnosticsFile);
-			SetCurrentState(EUploadState::SendingFiles);
-			UploadNextFile();
-		}
-		else
-		{
-			PostReportComplete();
-		}
-	}
-	else if (SendDiagnosticsFile)
+	const bool SendDiagnosticsFile = !DiagnosticsFile.IsEmpty();
+	if (SendDiagnosticsFile)
 	{
 		PendingFiles.Push(DiagnosticsFile);
 	}
@@ -147,19 +129,51 @@ bool FCrashUpload::SendCheckReportRequest()
 	return Request->ProcessRequest();
 }
 
-void FCrashUpload::UploadNextFile()
+enum class ECompressedCrashFileHeader
 {
-	UE_LOG(CrashReportClientLog, Log, TEXT("UploadNextFile: have %d pending files"), PendingFiles.Num());
+	MAGIC = 0x7E1B83C1,
+};
+
+struct FCompressedCrashFile : FNoncopyable
+{
+	int32 CurrentFileIndex;
+	FString Filename;
+	TArray<uint8> Filedata;
+
+	FCompressedCrashFile( int32 InCurrentFileIndex, const FString& InFilename, const TArray<uint8>& InFiledata )
+		: CurrentFileIndex(InCurrentFileIndex)
+		, Filename(InFilename)
+		, Filedata(InFiledata)
+	{
+	}
+
+	/** Serialization operator. */
+	friend FArchive& operator << (FArchive& Ar, FCompressedCrashFile& Data)
+	{
+		Ar << Data.CurrentFileIndex;
+		Data.Filename.SerializeAsANSICharArray( Ar, 260 );
+		Ar << Data.Filedata;
+		return Ar;
+	}
+};
+
+void FCrashUpload::CompressAndSendData()
+{
+	UE_LOG(CrashReportClientLog, Log, TEXT("CompressAndSendData have %d pending files"), PendingFiles.Num());
+
+	// Compress all files into one archive.
+	const int32 BufferSize = 16*1024*1024;
+
+	TArray<uint8> UncompressedData;
+	UncompressedData.Reserve( BufferSize );
+	FMemoryWriter MemoryWriter( UncompressedData, false, true );
+
+	int32 CurrentFileIndex = 0;
 
 	// Loop to keep trying files until a send succeeds or we run out of files
 	while (PendingFiles.Num() != 0)
 	{
 		FString PathOfFileToUpload = PendingFiles.Pop();
-		// Remember if there was already a diagnostics file in the report, so we don't send it twice
-		if (PathOfFileToUpload.EndsWith(GDiagnosticsFilename))
-		{
-			bDiagnosticsFileSent = true;
-		}
 		
 		if (FPlatformFileManager::Get().GetPlatformFile().FileSize(*PathOfFileToUpload) > MaxFileSizeToUpload)
 		{
@@ -173,7 +187,7 @@ void FCrashUpload::UploadNextFile()
 			continue;
 		}
 
-		UE_LOG(CrashReportClientLog, Log, TEXT("UploadNextFile: uploading %d bytes ('%s')"), PostData.Num(), *PathOfFileToUpload);
+		UE_LOG(CrashReportClientLog, Log, TEXT("CompressAndSendData compressing %d bytes ('%s')"), PostData.Num(), *PathOfFileToUpload);
 		FString Filename = FPaths::GetCleanFilename(PathOfFileToUpload);
 		if (Filename == "diagnostics.txt")
 		{
@@ -181,25 +195,56 @@ void FCrashUpload::UploadNextFile()
 			Filename[0] = 'D';
 		}
 
-		// Set up request for upload
-		UE_LOG(CrashReportClientLog, Log, TEXT("Sending HTTP request (posting file)"));
-		auto Request = CreateHttpRequest();
-		Request->SetVerb(TEXT("POST"));
-		Request->SetHeader(TEXT("Content-Type"), TEXT("application/octet-stream"));
-		Request->SetURL(UrlPrefix / TEXT("UploadReportFile"));
-		Request->SetContent(PostData);
-		Request->SetHeader(TEXT("DirectoryName"), *ErrorReport.GetReportDirectoryLeafName());
-		Request->SetHeader(TEXT("FileName"), Filename);
-		Request->SetHeader(TEXT("FileLength"), FString::FromInt(PostData.Num()));
-
-		if (Request->ProcessRequest())
-		{
-			return;
-		}
+		FCompressedCrashFile FileToCompress( CurrentFileIndex, Filename, PostData );
+		CurrentFileIndex++;
 
-		UE_LOG(CrashReportClientLog, Warning, TEXT("Failed to send file upload request"));
+		MemoryWriter << FileToCompress;
+	}
+
+	
+	uint8* CompressedDataRaw = new uint8[BufferSize];
+
+	int32 CompressedSize = BufferSize;
+	int32 UncompressedSize = UncompressedData.Num();
+	const bool bResult = FCompression::CompressMemory( COMPRESS_ZLIB, CompressedDataRaw, CompressedSize, UncompressedData.GetData(), UncompressedSize );
+	if( !bResult )
+	{
+		UE_LOG(CrashReportClientLog, Warning, TEXT("Couldn't compress the crash report files"));
+		SetCurrentState(EUploadState::Cancelled);
+		return;
 	}
-	PostReportComplete();
+	
+	const FString Filename = ErrorReport.GetReportDirectoryLeafName() + TEXT(".ue4crash");
+
+	// Copy compressed data into the array.
+	TArray<uint8> CompressedData;
+	CompressedData.Append( CompressedDataRaw, CompressedSize );
+	delete CompressedDataRaw;
+	CompressedDataRaw = nullptr;
+
+ 	// Set up request for upload
+ 	UE_LOG(CrashReportClientLog, Log, TEXT("Sending HTTP request (posting file)"));
+ 	auto Request = CreateHttpRequest();
+ 	Request->SetVerb(TEXT("POST"));
+ 	Request->SetHeader(TEXT("Content-Type"), TEXT("application/octet-stream"));
+ 	Request->SetURL(UrlPrefix / TEXT("UploadReportFile"));
+ 	Request->SetContent(CompressedData);
+ 	Request->SetHeader(TEXT("DirectoryName"), *ErrorReport.GetReportDirectoryLeafName());
+ 	Request->SetHeader(TEXT("FileName"), Filename);
+ 	Request->SetHeader(TEXT("FileLength"), TTypeToString<int32>::ToString(CompressedData.Num()) );
+	Request->SetHeader(TEXT("CompressedSize"), TTypeToString<int32>::ToString(CompressedSize) );
+	Request->SetHeader(TEXT("UncompressedSize"), TTypeToString<int32>::ToString(UncompressedSize) );
+	Request->SetHeader(TEXT("NumberOfFiles"), TTypeToString<int32>::ToString(CurrentFileIndex) );
+
+ 	if (Request->ProcessRequest())
+ 	{
+ 		return;
+ 	}
+	else
+	{
+		UE_LOG(CrashReportClientLog, Warning, TEXT("Failed to send file upload request"));
+		SetCurrentState(EUploadState::Cancelled);
+	}	
 }
 
 void FCrashUpload::AssignReportIdToPostDataBuffer()
@@ -243,7 +288,7 @@ void FCrashUpload::PostReportComplete()
 
 void FCrashUpload::OnProcessRequestComplete(FHttpRequestPtr HttpRequest, FHttpResponsePtr HttpResponse, bool bSucceeded)
 {
-	UE_LOG(CrashReportClientLog, Log, TEXT("=========================> OnProcessRequestComplete(), State=%s"), ToString(State));
+	UE_LOG(CrashReportClientLog, Log, TEXT("OnProcessRequestComplete(), State=%s bSucceeded=%i"), ToString(State), (int32)bSucceeded );
 	switch (State)
 	{
 	default:
@@ -290,18 +335,21 @@ void FCrashUpload::OnProcessRequestComplete(FHttpRequestPtr HttpRequest, FHttpRe
 		}
 		else
 		{
-			SetCurrentState(EUploadState::SendingFiles);
-			UploadNextFile();
+			SetCurrentState(EUploadState::CompressAndSendData);
+			CompressAndSendData();
 		}
 		break;
 
-	case EUploadState::SendingFiles:
+	case EUploadState::CompressAndSendData:
 		if (!bSucceeded)
 		{
 			UE_LOG(CrashReportClientLog, Warning, TEXT("File upload failed"));
+			SetCurrentState(EUploadState::Cancelled);
+		}
+		else
+		{
+			PostReportComplete();
 		}
-
-		UploadNextFile();
 		break;
 
 	case EUploadState::PostingReportComplete:
@@ -469,7 +517,7 @@ const TCHAR* FCrashUpload::ToString(EUploadState::Type State)
 		case EUploadState::CheckingReportDetail:
 			return TEXT("CheckingReportDetail");
 
-		case EUploadState::SendingFiles:
+		case EUploadState::CompressAndSendData:
 			return TEXT("SendingFiles");
 
 		case EUploadState::WaitingToPostReportComplete:
diff --git a/Engine/Source/Programs/CrashReporter/CrashReportClient/Private/CrashUpload.h b/Engine/Source/Programs/CrashReporter/CrashReportClient/Private/CrashUpload.h
--- a/Engine/Source/Programs/CrashReporter/CrashReportClient/Private/CrashUpload.h
+++ b/Engine/Source/Programs/CrashReporter/CrashReportClient/Private/CrashUpload.h
@@ -71,7 +71,7 @@ private:
 			Ready,
 			CheckingReport,
 			CheckingReportDetail,
-			SendingFiles,
+			CompressAndSendData,
 			WaitingToPostReportComplete,
 			PostingReportComplete,
 			Finished,
@@ -80,7 +80,7 @@ private:
 			UploadError,
 			Cancelled,
 
-			FirstCompletedState = Finished
+			FirstCompletedState = Finished,
 		};
 	};
 
@@ -104,9 +104,9 @@ private:
 	bool SendCheckReportRequest();
 
 	/**
-	 * If there are more files, upload the next, else send an upload complete message
+	 * Compresses all crash report files and sends one compressed file.
 	 */
-	void UploadNextFile();
+	void CompressAndSendData();
 
 	/**
 	 * Convert the report name to single byte non-zero-terminated HTTP post data
@@ -183,9 +183,6 @@ private:
 	/** State to pause at until confirmation has been received to continue */
 	EUploadState::Type PauseState;
 
-	/** Flag to record if a diagnostics report text file has been sent */
-	bool bDiagnosticsFileSent;
-
 	/** Full paths of files still to be uploaded */
 	TArray<FString> PendingFiles;
 
diff --git a/Engine/Source/Programs/CrashReporter/CrashReportClient/Private/MainLoopTiming.cpp b/Engine/Source/Programs/CrashReporter/CrashReportClient/Private/MainLoopTiming.cpp
--- a/Engine/Source/Programs/CrashReporter/CrashReportClient/Private/MainLoopTiming.cpp
+++ b/Engine/Source/Programs/CrashReporter/CrashReportClient/Private/MainLoopTiming.cpp
@@ -4,7 +4,7 @@
 
 #include "MainLoopTiming.h"
 #include "LaunchEngineLoop.h"
-
+#include "TaskGraphInterfaces.h"
 
 FMainLoopTiming::FMainLoopTiming(float IdealTickRate, EMainLoopOptions::Type Options)
 	: IdealFrameTime(1.f / IdealTickRate)
@@ -17,6 +17,7 @@ FMainLoopTiming::FMainLoopTiming(float IdealTickRate, EMainLoopOptions::Type Opt
 void FMainLoopTiming::Tick()
 {
 	// Tick app logic
+	FTaskGraphInterface::Get().ProcessThreadUntilIdle(ENamedThreads::GameThread);
 	FTicker::GetCoreTicker().Tick(ActualDeltaTime);
 
 #if !CRASH_REPORT_UNATTENDED_ONLY
diff --git a/Engine/Source/Programs/CrashReporter/CrashReportClient/Public/CrashReportClientApp.h b/Engine/Source/Programs/CrashReporter/CrashReportClient/Public/CrashReportClientApp.h
--- a/Engine/Source/Programs/CrashReporter/CrashReportClient/Public/CrashReportClientApp.h
+++ b/Engine/Source/Programs/CrashReporter/CrashReportClient/Public/CrashReportClientApp.h
@@ -2,8 +2,6 @@
 
 #pragma once
 
-//#define DO_LOCAL_TESTING 1
-
 #define CRASH_REPORT_UNATTENDED_ONLY			PLATFORM_LINUX
 
 // Pre-compiled header includes
