A fix to prevent timeouts when viewing Anonymous crashgroups

diff --git a/Engine/Source/Programs/CrashReporter/CrashReportWebSite/Models/BuggRepository.cs b/Engine/Source/Programs/CrashReporter/CrashReportWebSite/Models/BuggRepository.cs
--- a/Engine/Source/Programs/CrashReporter/CrashReportWebSite/Models/BuggRepository.cs
+++ b/Engine/Source/Programs/CrashReporter/CrashReportWebSite/Models/BuggRepository.cs
@@ -2,6 +2,7 @@
 
 using System;
 using System.Collections.Generic;
+using System.Data;
 using System.Data.Linq.SqlClient;
 using System.Diagnostics;
 using System.Linq;
@@ -419,11 +420,15 @@ namespace Tools.CrashReporter.CrashReportWebSite.Models
 				IEnumerable<Bugg> SortedResults = GetSortedResults( Results, FormData.SortTerm, ( FormData.SortOrder == "Descending" ), FormData.DateFrom, FormData.DateTo );
 
 				// Grab just the results we want to display on this page
-				SortedResults = SortedResults.Skip( Skip ).Take( Take ).ToList();
 
-				return new BuggsViewModel
+			    var SortedResultsList = SortedResults.ToList();
+                var TotalCountedRecords = SortedResultsList.Count();
+
+			    SortedResultsList = SortedResultsList.GetRange(Skip, TotalCountedRecords >= Skip + Take ? Take : TotalCountedRecords);
+
+			    return new BuggsViewModel
 				{
-					Results = SortedResults,
+					Results = SortedResultsList,
 					PagingInfo = new PagingInfo { CurrentPage = FormData.Page, PageSize = FormData.PageSize, TotalResults = Results.Count() },
 					SortTerm = FormData.SortTerm,
 					SortOrder = FormData.SortOrder,
@@ -550,6 +555,49 @@ namespace Tools.CrashReporter.CrashReportWebSite.Models
 			}
 		}
 
+
+        public DataTable LINQToDataTable<T>(IEnumerable<T> varlist)
+        {
+             DataTable dtReturn = new DataTable();
+
+             // column names 
+             PropertyInfo[] oProps = null;
+
+             if (varlist == null) return dtReturn;
+
+             foreach (T rec in varlist)
+             {
+                  // Use reflection to get property names, to create table, Only first time, others will follow 
+                  if (oProps == null)
+                  {
+                       oProps = ((Type)rec.GetType()).GetProperties();
+                       foreach (PropertyInfo pi in oProps)
+                       {
+                            Type colType = pi.PropertyType;
+
+                            if ((colType.IsGenericType) && (colType.GetGenericTypeDefinition()      
+                            ==typeof(Nullable<>)))
+                             {
+                                 colType = colType.GetGenericArguments()[0];
+                             }
+
+                            dtReturn.Columns.Add(new DataColumn(pi.Name, colType));
+                       }
+                  }
+
+                  DataRow dr = dtReturn.NewRow();
+
+                  foreach (PropertyInfo pi in oProps)
+                  {
+                       dr[pi.Name] = pi.GetValue(rec, null) ?? DBNull.Value;
+                  }
+
+                  dtReturn.Rows.Add(dr);
+             }
+             return dtReturn;
+        }
+
+
 		private DateTime AddOneDayToDate( DateTime Date )
 		{
 			return Date.AddDays( 1 );
@@ -571,72 +619,71 @@ namespace Tools.CrashReporter.CrashReportWebSite.Models
 				try
 				{
 
-					var IntermediateResults =
-					(
-						from BuggCrashDetail in BuggsDataContext.Buggs_Crashes
-						where BuggCrashDetail.Crash.TimeOfCrash >= DateFrom && BuggCrashDetail.Crash.TimeOfCrash <= AddOneDayToDate( DateTo )
-						group BuggCrashDetail by BuggCrashDetail.BuggId into CrashesGrouped
-						join BuggDetail in Results on CrashesGrouped.Key equals BuggDetail.Id
-						select new { Bugg = BuggDetail, Count = CrashesGrouped.Count() }
-					);
-
-					using( FScopedLogTimer LogTimer2 = new FScopedLogTimer( "BuggRepository.GetSortedResults.CrashesInTimeFrame" ) )
-					{
-						foreach( var Result in IntermediateResults )
-						{
-							Result.Bugg.CrashesInTimeFrame = Result.Count;
-						}
-					}
-
-
-					switch( SortTerm )
-					{
-						case "CrashesInTimeFrame":
-							IntermediateResults = CrashRepository.OrderBy( IntermediateResults, BuggCrashInstance => BuggCrashInstance.Count, bSortDescending );
-							break;
-
-						case "Id":
-							IntermediateResults = CrashRepository.OrderBy( IntermediateResults, BuggCrashInstance => BuggCrashInstance.Bugg.Id, bSortDescending );
-							break;
-
-						case "BuildVersion":
-							IntermediateResults = CrashRepository.OrderBy( IntermediateResults, BuggCrashInstance => BuggCrashInstance.Bugg.BuildVersion, bSortDescending );
-							break;
-
-						case "LatestCrash":
-							IntermediateResults = CrashRepository.OrderBy( IntermediateResults, BuggCrashInstance => BuggCrashInstance.Bugg.TimeOfLastCrash, bSortDescending );
-							break;
-
-						case "FirstCrash":
-							IntermediateResults = CrashRepository.OrderBy( IntermediateResults, BuggCrashInstance => BuggCrashInstance.Bugg.TimeOfFirstCrash, bSortDescending );
-							break;
-
-						case "NumberOfCrashes":
-							IntermediateResults = CrashRepository.OrderBy( IntermediateResults, BuggCrashInstance => BuggCrashInstance.Bugg.NumberOfCrashes, bSortDescending );
-							break;
-
-						case "NumberOfUsers":
-							IntermediateResults = CrashRepository.OrderBy( IntermediateResults, BuggCrashInstance => BuggCrashInstance.Bugg.NumberOfUsers, bSortDescending );
-							break;
-
-						case "Pattern":
-							IntermediateResults = CrashRepository.OrderBy( IntermediateResults, BuggCrashInstance => BuggCrashInstance.Bugg.Pattern, bSortDescending );
-							break;
-
-						case "Status":
-							IntermediateResults = CrashRepository.OrderBy( IntermediateResults, BuggCrashInstance => BuggCrashInstance.Bugg.Status, bSortDescending );
-							break;
-
-						case "FixedChangeList":
-							IntermediateResults = CrashRepository.OrderBy( IntermediateResults, BuggCrashInstance => BuggCrashInstance.Bugg.FixedChangeList, bSortDescending );
-							break;
-
-						case "TTPID":
-							IntermediateResults = CrashRepository.OrderBy( IntermediateResults, BuggCrashInstance => BuggCrashInstance.Bugg.TTPID, bSortDescending );
-							break;
-
-					}
-					return IntermediateResults.Select( x => x.Bugg );
+                    var IntermediateResults =
+                    (
+                        from BuggCrashDetail in BuggsDataContext.Buggs_Crashes
+                        where BuggCrashDetail.Crash.TimeOfCrash >= DateFrom && BuggCrashDetail.Crash.TimeOfCrash <= AddOneDayToDate(DateTo)
+                        group BuggCrashDetail by BuggCrashDetail.BuggId into CrashesGrouped
+                        join BuggDetail in Results on CrashesGrouped.Key equals BuggDetail.Id
+                        select new { Bugg = BuggDetail, Count = CrashesGrouped.Count() }
+                    );
+
+                    using (FScopedLogTimer LogTimer2 = new FScopedLogTimer("BuggRepository.GetSortedResults.CrashesInTimeFrame"))
+                    {
+                        foreach (var Result in IntermediateResults)
+                        {
+                            Result.Bugg.CrashesInTimeFrame = Result.Count;
+                        }
+                    }
+                    
+                    switch (SortTerm)
+                    {
+                        case "CrashesInTimeFrame":
+                            IntermediateResults = CrashRepository.OrderBy(IntermediateResults, BuggCrashInstance => BuggCrashInstance.Count, bSortDescending);
+                            break;
+
+                        case "Id":
+                            IntermediateResults = CrashRepository.OrderBy(IntermediateResults, BuggCrashInstance => BuggCrashInstance.Bugg.Id, bSortDescending);
+                            break;
+
+                        case "BuildVersion":
+                            IntermediateResults = CrashRepository.OrderBy(IntermediateResults, BuggCrashInstance => BuggCrashInstance.Bugg.BuildVersion, bSortDescending);
+                            break;
+
+                        case "LatestCrash":
+                            IntermediateResults = CrashRepository.OrderBy(IntermediateResults, BuggCrashInstance => BuggCrashInstance.Bugg.TimeOfLastCrash, bSortDescending);
+                            break;
+
+                        case "FirstCrash":
+                            IntermediateResults = CrashRepository.OrderBy(IntermediateResults, BuggCrashInstance => BuggCrashInstance.Bugg.TimeOfFirstCrash, bSortDescending);
+                            break;
+
+                        case "NumberOfCrashes":
+                            IntermediateResults = CrashRepository.OrderBy(IntermediateResults, BuggCrashInstance => BuggCrashInstance.Bugg.NumberOfCrashes, bSortDescending);
+                            break;
+
+                        case "NumberOfUsers":
+                            IntermediateResults = CrashRepository.OrderBy(IntermediateResults, BuggCrashInstance => BuggCrashInstance.Bugg.NumberOfUsers, bSortDescending);
+                            break;
+
+                        case "Pattern":
+                            IntermediateResults = CrashRepository.OrderBy(IntermediateResults, BuggCrashInstance => BuggCrashInstance.Bugg.Pattern, bSortDescending);
+                            break;
+
+                        case "Status":
+                            IntermediateResults = CrashRepository.OrderBy(IntermediateResults, BuggCrashInstance => BuggCrashInstance.Bugg.Status, bSortDescending);
+                            break;
+
+                        case "FixedChangeList":
+                            IntermediateResults = CrashRepository.OrderBy(IntermediateResults, BuggCrashInstance => BuggCrashInstance.Bugg.FixedChangeList, bSortDescending);
+                            break;
+
+                        case "TTPID":
+                            IntermediateResults = CrashRepository.OrderBy(IntermediateResults, BuggCrashInstance => BuggCrashInstance.Bugg.TTPID, bSortDescending);
+                            break;
+
+                    }
+					return IntermediateResults.Select(x => x.Bugg);
 				}
 				catch( Exception Ex )
 				{
