MinidumpDiagnostics - No longer needs ICU or database

diff --git a/Engine/Source/Programs/CrashReporter/MinidumpDiagnostics/MinidumpDiagnostics.Build.cs b/Engine/Source/Programs/CrashReporter/MinidumpDiagnostics/MinidumpDiagnostics.Build.cs
--- a/Engine/Source/Programs/CrashReporter/MinidumpDiagnostics/MinidumpDiagnostics.Build.cs
+++ b/Engine/Source/Programs/CrashReporter/MinidumpDiagnostics/MinidumpDiagnostics.Build.cs
@@ -18,8 +18,5 @@ public class MinidumpDiagnostics : ModuleRules
 				"Projects"
 			}
 			);
-
-		// Need database support!
-		Definitions.Add( "WITH_DATABASE_SUPPORT=1" );
 	}
 }
diff --git a/Engine/Source/Programs/CrashReporter/MinidumpDiagnostics/MinidumpDiagnostics.Target.cs b/Engine/Source/Programs/CrashReporter/MinidumpDiagnostics/MinidumpDiagnostics.Target.cs
--- a/Engine/Source/Programs/CrashReporter/MinidumpDiagnostics/MinidumpDiagnostics.Target.cs
+++ b/Engine/Source/Programs/CrashReporter/MinidumpDiagnostics/MinidumpDiagnostics.Target.cs
@@ -27,6 +27,8 @@ public class MinidumpDiagnosticsTarget : TargetRules
 			new UEBuildBinaryConfiguration(	InType: UEBuildBinaryType.Executable,
 											InModuleNames: new List<string>() { "MinidumpDiagnostics" })
 			);
+
+		UEBuildConfiguration.bCompileICU = false;
 	}
 
 	public override bool ShouldCompileMonolithic(UnrealTargetPlatform InPlatform, UnrealTargetConfiguration InConfiguration)
@@ -40,8 +42,6 @@ public class MinidumpDiagnosticsTarget : TargetRules
 		ref CPPEnvironmentConfiguration OutCPPEnvironmentConfiguration
 		)
 	{
-		OutCPPEnvironmentConfiguration.Definitions.Add("WITH_DATABASE_SUPPORT=1");
-
 		UEBuildConfiguration.bCompileLeanAndMeanUE = true;
 
 		// Don't need editor
@@ -50,7 +50,8 @@ public class MinidumpDiagnosticsTarget : TargetRules
 		// MinidumpDiagnostics doesn't ever compile with the engine linked in
 		UEBuildConfiguration.bCompileAgainstEngine = false;
 
-		UEBuildConfiguration.bIncludeADO = true;
+		UEBuildConfiguration.bIncludeADO = false;
+		//UEBuildConfiguration.bCompileICU = false;
 
 		// MinidumpDiagnostics.exe has no exports, so no need to verify that a .lib and .exp file was emitted by the linker.
 		OutLinkEnvironmentConfiguration.bHasExports = false;
