CrashReportClient - UEBuildConfiguration.bIncludeADO set to false

diff --git a/Engine/Source/Programs/CrashReporter/CrashReportClient/CrashReportClient.Target.cs b/Engine/Source/Programs/CrashReporter/CrashReportClient/CrashReportClient.Target.cs
--- a/Engine/Source/Programs/CrashReporter/CrashReportClient/CrashReportClient.Target.cs
+++ b/Engine/Source/Programs/CrashReporter/CrashReportClient/CrashReportClient.Target.cs
@@ -80,7 +80,7 @@ public class CrashReportClientTarget : TargetRules
 		UEBuildConfiguration.bUseLoggingInShipping = true;
 		UEBuildConfiguration.bCompileSteamOSS = false;
 
-		UEBuildConfiguration.bIncludeADO = (Target.Platform == UnrealTargetPlatform.Win64 || Target.Platform == UnrealTargetPlatform.Win32);
+		UEBuildConfiguration.bIncludeADO = false;
 		
 		// Do not include ICU for Linux (this is a temporary workaround, separate headless CrashReportClient target should be created, see UECORE-14 for details).
 		if (Target.Platform == UnrealTargetPlatform.Linux)
