CrashReportClient - Minor code refactor regarding usage CRASH_REPORT_UNATTENDED_ONLY, also fixes issues for Linux

diff --git a/Engine/Source/Programs/CrashReporter/CrashReportClient/Private/CrashReportClient.cpp b/Engine/Source/Programs/CrashReporter/CrashReportClient/Private/CrashReportClient.cpp
--- a/Engine/Source/Programs/CrashReporter/CrashReportClient/Private/CrashReportClient.cpp
+++ b/Engine/Source/Programs/CrashReporter/CrashReportClient/Private/CrashReportClient.cpp
@@ -200,37 +200,31 @@ FString FCrashReportClient::GetCrashedAppName() const
 	return GetCrashDescription().GameName;
 }
 
-void FCrashReportClient::FinalizeDiagnoseReportWorker( FText InDiagnosticText )
+void FCrashReportClient::FinalizeDiagnoseReportWorker( FText ReportText )
 {
-	DiagnosticText = InDiagnosticText;
+	DiagnosticText = FCrashReportUtil::FormatDiagnosticText( ReportText, GetCrashDescription().MachineId, GetCrashDescription().EpicAccountId, GetCrashDescription().UserName );
 
 	auto DiagnosticsFilePath = ErrorReport.GetReportDirectory() / GDiagnosticsFilename;
 	Uploader.LocalDiagnosisComplete(FPaths::FileExists(DiagnosticsFilePath) ? DiagnosticsFilePath : TEXT(""));
 }
 
-#endif // !CRASH_REPORT_UNATTENDED_ONLY
-
-
 FDiagnoseReportWorker::FDiagnoseReportWorker( FCrashReportClient* InCrashReportClient ) 
-	: CrashReportClient( InCrashReportClient )		
-	, MachineId(GetCrashDescription().MachineId)
-	, EpicAccountId(GetCrashDescription().EpicAccountId)
-	, UserNameNoDot(GetCrashDescription().UserName)
+	: CrashReportClient( InCrashReportClient )
 {}
 
 void FDiagnoseReportWorker::DoWork()
 {
 	const FText ReportText = CrashReportClient->ErrorReport.DiagnoseReport();
-	DiagnosticText = FCrashReportUtil::FormatDiagnosticText( ReportText, MachineId, EpicAccountId, UserNameNoDot );
-
 	// Inform the game thread that we are done.
 	FSimpleDelegateGraphTask::CreateAndDispatchWhenReady
 	(
-		FSimpleDelegateGraphTask::FDelegate::CreateRaw( CrashReportClient, &FCrashReportClient::FinalizeDiagnoseReportWorker, DiagnosticText ),
+		FSimpleDelegateGraphTask::FDelegate::CreateRaw( CrashReportClient, &FCrashReportClient::FinalizeDiagnoseReportWorker, ReportText ),
 		TStatId(), nullptr, ENamedThreads::GameThread
 	);
 }
 
+#endif // !CRASH_REPORT_UNATTENDED_ONLY
+
 FText FCrashReportUtil::FormatDiagnosticText( const FText& DiagnosticText, const FString MachineId, const FString EpicAccountId, const FString UserNameNoDot )
 {
 	if( FRocketSupport::IsRocket() )
diff --git a/Engine/Source/Programs/CrashReporter/CrashReportClient/Private/CrashReportClient.h b/Engine/Source/Programs/CrashReporter/CrashReportClient/Private/CrashReportClient.h
--- a/Engine/Source/Programs/CrashReporter/CrashReportClient/Private/CrashReportClient.h
+++ b/Engine/Source/Programs/CrashReporter/CrashReportClient/Private/CrashReportClient.h
@@ -6,6 +6,8 @@
 #include "PlatformErrorReport.h"
 #include "CrashReportUtil.h"
 
+#if !CRASH_REPORT_UNATTENDED_ONLY
+
 class FCrashReportClient;
 
 /**
@@ -17,18 +19,6 @@ public:
 	/** Pointer to the crash report client, used to store the results. */
 	FCrashReportClient* CrashReportClient;
 
-	/** Machine ID. */
-	const FString MachineId;
-
-	/** Epic Account ID. */
-	const FString EpicAccountId;
-
-	/** User name without dot. */
-	const FString UserNameNoDot;
-
-	/** Diagnostic text. */
-	FText DiagnosticText;
-
 	/** Initialization constructor. */
 	FDiagnoseReportWorker( FCrashReportClient* InCrashReportClient );
 
@@ -46,8 +36,6 @@ public:
 	}
 };
 
-#if !CRASH_REPORT_UNATTENDED_ONLY
-
 /**
  * Main implementation of the crash report client application
  */
@@ -148,7 +136,7 @@ private:
 	void StartUIWillCloseTicker();
 
 	/** Enqueued from the diagnose report worker thread to be executed on the game thread. */
-	void FinalizeDiagnoseReportWorker( FText InDiagnosticText );
+	void FinalizeDiagnoseReportWorker( FText ReportText );
 
 	/** State enum to keep track of what the app is doing */
 	struct EApplicationState
