CrashReportClient - Removed Cancel button and counting down

diff --git a/Engine/Source/Programs/CrashReporter/CrashReportClient/Private/CrashReportClient.cpp b/Engine/Source/Programs/CrashReporter/CrashReportClient/Private/CrashReportClient.cpp
--- a/Engine/Source/Programs/CrashReporter/CrashReportClient/Private/CrashReportClient.cpp
+++ b/Engine/Source/Programs/CrashReporter/CrashReportClient/Private/CrashReportClient.cpp
@@ -33,14 +33,12 @@ FCrashDescription& GetCrashDescription()
 #if !CRASH_REPORT_UNATTENDED_ONLY
 
 FCrashReportClient::FCrashReportClient(const FPlatformErrorReport& InErrorReport)
-	: AppState(EApplicationState::Ready)
-	, SubmittedCountdown(-1)
-	, DiagnosticText( LOCTEXT("ProcessingReport", "Processing crash report ...") )
+	: DiagnosticText( LOCTEXT("ProcessingReport", "Processing crash report ...") )
 	, DiagnoseReportTask(nullptr)
 	, ErrorReport( InErrorReport )
 	, Uploader(GServerIP)
-	, CancelButtonText(LOCTEXT("Cancel", "Don't Send"))
 	, bBeginUploadCalled(false)
+	, bShouldWindowBeHidden(false)
 {
 
 	if (!ErrorReport.TryReadDiagnosticsFile(DiagnosticText) && !FParse::Param(FCommandLine::Get(), TEXT("no-local-diagnosis")))
@@ -56,25 +54,8 @@ FCrashReportClient::FCrashReportClient(const FPlatformErrorReport& InErrorReport
 
 FReply FCrashReportClient::Submit()
 {
-	if (AppState == EApplicationState::Ready)
-	{
-		StoreCommentAndUpload();
-	}
-
-	return FReply::Handled();
-}
-
-FReply FCrashReportClient::Cancel()
-{
-	// If the AppState is Ready, this performs a Cancel; otherwise the Cancel
-	// button has become the Close button and the ticker has already started.
-	if (AppState == EApplicationState::Ready)
-	{
-		Uploader.Cancel();
-		StartUIWillCloseTicker();
-	}
-	AppState = EApplicationState::Closing;
-
+	StoreCommentAndUpload();
+	bShouldWindowBeHidden = true;
 	return FReply::Handled();
 }
 
@@ -84,92 +65,42 @@ FReply FCrashReportClient::CopyCallstack()
 	return FReply::Handled();
 }
 
-FText FCrashReportClient::GetStatusText() const
-{
-	static const FText ClosingText = LOCTEXT("Closing", "Thank you for reporting this issue - closing automatically");
-	return AppState == EApplicationState::CountingDown ? ClosingText : Uploader.GetStatusText();
-}
-
-FText FCrashReportClient::GetCancelButtonText() const
-{
-	return CancelButtonText;
-}
-
 FText FCrashReportClient::GetDiagnosticText() const
 {
 	return DiagnosticText;
 }
 
-EVisibility FCrashReportClient::SubmitButtonVisibility() const
-{
-	return AppState == EApplicationState::Ready ? EVisibility::Visible : EVisibility::Hidden;
-}
-
 void FCrashReportClient::UserCommentChanged(const FText& Comment, ETextCommit::Type CommitType)
 {
 	UserComment = Comment;
 
 	// Implement Shift+Enter to commit shortcut
-	if (CommitType == ETextCommit::OnEnter &&
-		AppState == EApplicationState::Ready &&
-		FSlateApplication::Get().GetModifierKeys().IsShiftDown())
+	if (CommitType == ETextCommit::OnEnter && FSlateApplication::Get().GetModifierKeys().IsShiftDown())
 	{
 		Submit();
 	}
 }
 
-bool FCrashReportClient::ShouldWindowBeHidden() const
-{
-	return AppState == EApplicationState::Closing;
-}
-
 void FCrashReportClient::RequestCloseWindow(const TSharedRef<SWindow>& Window)
 {
-	Cancel();
+	// We may still processing minidump etc. so start the main ticker.
+	StartTicker();
 }
 
-void FCrashReportClient::StartUIWillCloseTicker()
+void FCrashReportClient::StartTicker()
 {
-	FTicker::GetCoreTicker().AddTicker(FTickerDelegate::CreateRaw(this, &FCrashReportClient::UIWillCloseTick), 1.f);
+	FTicker::GetCoreTicker().AddTicker(FTickerDelegate::CreateSP(this, &FCrashReportClient::Tick), 1.f);
 }
 
 void FCrashReportClient::StoreCommentAndUpload()
 {
 	// Call upload even if the report is empty: pending reports will be sent if any
 	ErrorReport.SetUserComment(UserComment);
-
-	SubmittedCountdown = 5;
-	AppState = EApplicationState::CountingDown;
-	// Change the submit button text immediately (also sends diagnostics file if complete)
-	UIWillCloseTick(0);
-	StartUIWillCloseTicker();
+	StartTicker();
 }
 
-bool FCrashReportClient::UIWillCloseTick(float UnusedDeltaTime)
+bool FCrashReportClient::Tick(float UnusedDeltaTime)
 {
-	bool bCountingDown = AppState == EApplicationState::CountingDown;
-	if (!bCountingDown && AppState != EApplicationState::Closing)
-	{
-		CRASHREPORTCLIENT_CHECK(false);
-		return false;
-	}
-
-	static const FText CountdownTextFormat = LOCTEXT("CloseApplication", "Close ({0})");
-
-	if (bCountingDown)
-	{
-		CancelButtonText = FText::Format(CountdownTextFormat, FText::AsNumber(SubmittedCountdown));
-		if (SubmittedCountdown-- == 0)
-		{
-			AppState = EApplicationState::Closing;
-		}
-		else
-		{
-			// More ticks, please
-			return true;
-		}
-	}
-
 	// We are waiting for diagnose report task to complete.
 	if( DiagnoseReportTask && !DiagnoseReportTask->IsWorkDone() )
 	{
@@ -190,8 +121,7 @@ bool FCrashReportClient::UIWillCloseTick(float UnusedDeltaTime)
 		return true;
 	}
 
-	FPlatformMisc::RequestExit(false /* don't force */);
-	// No more ticks, thank you
+	FPlatformMisc::RequestExit(false);
 	return false;
 }
 
diff --git a/Engine/Source/Programs/CrashReporter/CrashReportClient/Private/CrashReportClient.h b/Engine/Source/Programs/CrashReporter/CrashReportClient/Private/CrashReportClient.h
--- a/Engine/Source/Programs/CrashReporter/CrashReportClient/Private/CrashReportClient.h
+++ b/Engine/Source/Programs/CrashReporter/CrashReportClient/Private/CrashReportClient.h
@@ -56,42 +56,18 @@ public:
 	 */
 	FReply Submit();
 
-	/**
-	 * Respond to the user pressing Cancel or Close
-	 * @return Whether the request was handled
-	 */
-	FReply Cancel();
-
 	/**
 	 * Respond to the user requesting the callstack to be copied to the clipboard
 	 * @return Whether the request was handled
 	 */
 	FReply CopyCallstack();
 
-	/**
-	 * Provide text to display at the bottom of the app
-	 * @return Localized text to display
-	 */
-	FText GetStatusText() const;
-
-	/**
-	 * Tell UI to display 'Cancel' or 'Close'
-	 * @return Localized text to display
-	 */
-	FText GetCancelButtonText() const;
-
 	/**
 	 * Pass on exception and callstack from the platform error report code
 	 * @return Localized text to display
 	 */
 	FText GetDiagnosticText() const;
 
-	/**
-	 * Indicate availability of Submit button
-	 * @return Whether button should be visible
-	 */
-	EVisibility SubmitButtonVisibility() const;
-
 	/**
 	 * Access to name (which usually includes the build config) of the app that crashed
 	 * @return Name of executable
@@ -106,17 +82,16 @@ public:
 	void UserCommentChanged(const FText& Comment, ETextCommit::Type CommitType);
 
 	/**
-	 * Indication of whether the app should be visible
-	 * @return Whether the main window should be hidden
-	 */
-	bool ShouldWindowBeHidden() const;
-
-	/**
-	 * Handle user closing the main window: same behaviour as Cancel
+	 * Handle user closing the main window
 	 * @param Window Main window
 	 */
 	void RequestCloseWindow(const TSharedRef<SWindow>& Window);
 
+	bool ShouldWindowBeHidden() const
+	{
+		return bShouldWindowBeHidden;
+	}
+
 private:
 	/**
 	 * Write the user's comment to the report and begin uploading the entire report 
@@ -128,33 +103,16 @@ private:
 	 * @param DeltaTime Time since last update, unused
 	 * @return Whether the updates should continue
 	 */
-	bool UIWillCloseTick(float DeltaTime);
+	bool Tick(float DeltaTime);
 
 	/**
-	 * Begin calling UIWillCloseTick once a second
+	 * Begin calling Tick once a second
 	 */
-	void StartUIWillCloseTicker();
+	void StartTicker();
 
 	/** Enqueued from the diagnose report worker thread to be executed on the game thread. */
 	void FinalizeDiagnoseReportWorker( FText ReportText );
 
-	/** State enum to keep track of what the app is doing */
-	struct EApplicationState
-	{
-		enum Type
-		{
-			Ready,				/** Waiting for user input */
-			CountingDown,		/** Submit pressed; counting down */
-			Closing				/** UI hidden, possibly waiting for tasks to finish */
-		};
-	};
-
-	/** What the app is currently doing */
-	EApplicationState::Type AppState;
-
-	/** Count-down until closing the app after user has pressed Submit */
-	int SubmittedCountdown;
-
 	/** Comment provided by the user */
 	FText UserComment;
 
@@ -170,11 +128,12 @@ private:
 	/** Object that uploads report files to the server */
 	FCrashUpload Uploader;
 
-	/** Text for Cancel/Close button, including count-down value */
-	FText CancelButtonText;
-
-	/** Wheter BeginUpload has been called. */
+	/** Whether BeginUpload has been called. */
 	bool bBeginUploadCalled;
+
+	/** Whether the main windows should be hidden. */
+	bool bShouldWindowBeHidden;
+
 };
 
 #endif // !CRASH_REPORT_UNATTENDED_ONLY
diff --git a/Engine/Source/Programs/CrashReporter/CrashReportClient/Private/CrashReportClientApp.cpp b/Engine/Source/Programs/CrashReporter/CrashReportClient/Private/CrashReportClientApp.cpp
--- a/Engine/Source/Programs/CrashReporter/CrashReportClient/Private/CrashReportClientApp.cpp
+++ b/Engine/Source/Programs/CrashReporter/CrashReportClient/Private/CrashReportClientApp.cpp
@@ -38,20 +38,6 @@ static FString ReportDirectoryAbsolutePath;
 /** Name of the game passed via the command line. */
 static FString GameNameFromCmd;
 
-/**
- * Upload the crash report with no user interaction
- */
-void RunCrashReportClientUnattended(FMainLoopTiming& MainLoop, const FPlatformErrorReport& ErrorReport)
-{
-	FCrashReportClientUnattended CrashReportClient( ErrorReport );
-
-	// loop until the app is ready to quit
-	while (!GIsRequestingExit)
-	{
-		MainLoop.Tick();
-	}
-}
-
 /**
  * Look for the report to upload, either in the command line or in the platform's report queue
  */
@@ -140,7 +126,13 @@ void RunCrashReportClient(const TCHAR* CommandLine)
 	if (bUnattended)
 	{
 		ErrorReport.SetUserComment( NSLOCTEXT( "CrashReportClient", "UnattendedMode", "Sent in the unattended mode" ) );
-		RunCrashReportClientUnattended(MainLoop, ErrorReport);
+		FCrashReportClientUnattended CrashReportClient( ErrorReport );
+
+		// loop until the app is ready to quit
+		while (!GIsRequestingExit)
+		{
+			MainLoop.Tick();
+		}
 	}
 	else
 	{
diff --git a/Engine/Source/Programs/CrashReporter/CrashReportClient/Private/SCrashReportClient.cpp b/Engine/Source/Programs/CrashReporter/CrashReportClient/Private/SCrashReportClient.cpp
--- a/Engine/Source/Programs/CrashReporter/CrashReportClient/Private/SCrashReportClient.cpp
+++ b/Engine/Source/Programs/CrashReporter/CrashReportClient/Private/SCrashReportClient.cpp
@@ -88,17 +88,6 @@ void SCrashReportClient::Construct(const FArguments& InArgs, TSharedRef<FCrashRe
 			]
 
 			// Stuff anchored to the bottom
-			+SVerticalBox::Slot()
-			.AutoHeight()
-			.Padding(10)
-			[
-				SNew(STextBlock)
-				.ColorAndOpacity(FSlateColor::UseSubduedForeground())
-				.AutoWrapText(true)
-				.TextStyle(FCrashReportClientStyle::Get(), "Status")
-				.Text(Client, &FCrashReportClient::GetStatusText)
-			]
-
 			+SVerticalBox::Slot()
 			.AutoHeight()
 			.Padding(10)
@@ -109,7 +98,7 @@ void SCrashReportClient::Construct(const FArguments& InArgs, TSharedRef<FCrashRe
 				.Padding(3)
 				[
 					SNew(STextBlock)
-					.Text(LOCTEXT("KeyboardShortcuts", "Shift+Enter: submit\nEscape: cancel"))
+					.Text(LOCTEXT("KeyboardShortcuts", "Shift+Enter: submit\n"))
 				]
 
 				+SHorizontalBox::Slot()
@@ -125,17 +114,8 @@ void SCrashReportClient::Construct(const FArguments& InArgs, TSharedRef<FCrashRe
 					[
 						SNew(SButton)
 						.HAlign(HAlign_Center)
-						.Text(LOCTEXT("Submit", "Submit Crash Report"))
+						.Text(LOCTEXT("Send", "Send"))
 						.OnClicked(Client, &FCrashReportClient::Submit)
-						.Visibility(Client, &FCrashReportClient::SubmitButtonVisibility)
-					]
-
-					+SUniformGridPanel::Slot(1, 0)
-					[
-						SNew(SButton)
-						.HAlign(HAlign_Center)
-						.Text(Client, &FCrashReportClient::GetCancelButtonText)
-						.OnClicked(Client, &FCrashReportClient::Cancel)
 					]
 				]
 			]
@@ -178,12 +158,7 @@ void SCrashReportClient::SetDefaultFocus()
 FReply SCrashReportClient::OnUnhandledKeyDown(const FKeyEvent& InKeyEvent)
 {
 	const FKey Key = InKeyEvent.GetKey();
-	if (Key == EKeys::Escape)
-	{
-		CrashReportClient->Cancel();
-		return FReply::Handled();
-	}
-	else if (Key == EKeys::Enter)
+	if (Key == EKeys::Enter)
 	{
 		if (InKeyEvent.IsShiftDown())
 		{
