CrashReportClient - Fix for hiding the main window

diff --git a/Engine/Source/Programs/CrashReporter/CrashReportClient/Private/CrashReportClient.cpp b/Engine/Source/Programs/CrashReporter/CrashReportClient/Private/CrashReportClient.cpp
--- a/Engine/Source/Programs/CrashReporter/CrashReportClient/Private/CrashReportClient.cpp
+++ b/Engine/Source/Programs/CrashReporter/CrashReportClient/Private/CrashReportClient.cpp
@@ -85,6 +85,7 @@ void FCrashReportClient::RequestCloseWindow(const TSharedRef<SWindow>& Window)
 {
 	// We may still processing minidump etc. so start the main ticker.
 	StartTicker();
+	bShouldWindowBeHidden = true;
 }
 
 void FCrashReportClient::StartTicker()
