CrashReportClient - Updated information regarding no debugging symbols

diff --git a/Engine/Source/Programs/CrashReporter/CrashReportClient/Private/Mac/MacErrorReport.cpp b/Engine/Source/Programs/CrashReporter/CrashReportClient/Private/Mac/MacErrorReport.cpp
--- a/Engine/Source/Programs/CrashReporter/CrashReportClient/Private/Mac/MacErrorReport.cpp
+++ b/Engine/Source/Programs/CrashReporter/CrashReportClient/Private/Mac/MacErrorReport.cpp
@@ -6,6 +6,8 @@
 
 #include "MacErrorReport.h"
 
+#define LOCTEXT_NAMESPACE "CrashReportClient"
+
 namespace
 {
 	/** Pointer to dynamically loaded crash diagnosis module */
@@ -116,14 +118,7 @@ FText FMacErrorReport::DiagnoseReport() const
 	
 	if ( !CrashDebugHelper->CreateMinidumpDiagnosticReport(ReportDirectory / DumpFilename) )
 	{
-		if ( FRocketSupport::IsRocket() )
-		{
-			return FText::FromString("We apologize for the inconvenience.\nPlease send this crash report to help improve our software.");
-		}
-		else
-		{
-			return FText::FromString("You do not have any debugging symbols required to display the callstack for this crash.");
-		}
+		return LOCTEXT("NoDebuggingSymbols", "You do not have any debugging symbols required to display the callstack for this crash.");
 	}
 	else
 	{
diff --git a/Engine/Source/Programs/CrashReporter/CrashReportClient/Private/Windows/WindowsErrorReport.cpp b/Engine/Source/Programs/CrashReporter/CrashReportClient/Private/Windows/WindowsErrorReport.cpp
--- a/Engine/Source/Programs/CrashReporter/CrashReportClient/Private/Windows/WindowsErrorReport.cpp
+++ b/Engine/Source/Programs/CrashReporter/CrashReportClient/Private/Windows/WindowsErrorReport.cpp
@@ -94,14 +94,7 @@ FText FWindowsErrorReport::DiagnoseReport() const
 
 	if (!CrashDebugHelper->CreateMinidumpDiagnosticReport(ReportDirectory / DumpFilename))
 	{
-		if ( FRocketSupport::IsRocket() )
-		{
-			return LOCTEXT("NoDebuggingSymbolsRocket", "We apologize for the inconvenience.\nPlease send this crash report to help improve our software.");
-		}
-		else
-		{
-			return LOCTEXT("NoDebuggingSymbols", "You do not have any debugging symbols required to display the callstack for this crash.");
-		}
+		return LOCTEXT("NoDebuggingSymbols", "You do not have any debugging symbols required to display the callstack for this crash.");
 	}
 
 	// Don't write a Diagnostics.txt to disk in UE4 release build. It will be displayed in the UI but not sent to the server.
