CrashReportClient - Tweaks to the UI

diff --git a/Engine/Source/Programs/CrashReporter/CrashReportClient/Private/CrashReportClient.cpp b/Engine/Source/Programs/CrashReporter/CrashReportClient/Private/CrashReportClient.cpp
--- a/Engine/Source/Programs/CrashReporter/CrashReportClient/Private/CrashReportClient.cpp
+++ b/Engine/Source/Programs/CrashReporter/CrashReportClient/Private/CrashReportClient.cpp
@@ -39,6 +39,7 @@ FCrashReportClient::FCrashReportClient(const FPlatformErrorReport& InErrorReport
 	, Uploader(GServerIP)
 	, bBeginUploadCalled(false)
 	, bShouldWindowBeHidden(false)
+	, bAllowToBeContacted(true)
 {
 
 	if (!ErrorReport.TryReadDiagnosticsFile(DiagnosticText) && !FParse::Param(FCommandLine::Get(), TEXT("no-local-diagnosis")))
@@ -52,6 +53,16 @@ FCrashReportClient::FCrashReportClient(const FPlatformErrorReport& InErrorReport
 	}
 }
 
+
+FCrashReportClient::~FCrashReportClient()
+{
+	if( DiagnoseReportTask )
+	{
+		DiagnoseReportTask->EnsureCompletion();
+		delete DiagnoseReportTask;
+	}
+}
+
 FReply FCrashReportClient::Submit()
 {
 	StoreCommentAndUpload();
@@ -88,6 +99,21 @@ void FCrashReportClient::RequestCloseWindow(const TSharedRef<SWindow>& Window)
 	bShouldWindowBeHidden = true;
 }
 
+bool FCrashReportClient::AreCallstackWidgetsEnabled() const
+{
+	return !IsProcessingCallstack();
+}
+
+EVisibility FCrashReportClient::IsThrobberVisible() const
+{
+	return IsProcessingCallstack() ? EVisibility::Visible : EVisibility::Hidden;
+}
+
+void FCrashReportClient::SCrashReportClient_OnCheckStateChanged( ECheckBoxState NewRadioState )
+{
+	bAllowToBeContacted = NewRadioState == ECheckBoxState::Checked;
+}
+
 void FCrashReportClient::StartTicker()
 {
 	FTicker::GetCoreTicker().AddTicker(FTickerDelegate::CreateSP(this, &FCrashReportClient::Tick), 1.f);
@@ -96,14 +122,14 @@ void FCrashReportClient::StartTicker()
 void FCrashReportClient::StoreCommentAndUpload()
 {
 	// Call upload even if the report is empty: pending reports will be sent if any
-	ErrorReport.SetUserComment(UserComment);
+	ErrorReport.SetUserComment(UserComment, bAllowToBeContacted);
 	StartTicker();
 }
 
 bool FCrashReportClient::Tick(float UnusedDeltaTime)
 {
 	// We are waiting for diagnose report task to complete.
-	if( DiagnoseReportTask && !DiagnoseReportTask->IsWorkDone() )
+	if( IsProcessingCallstack() )
 	{
 		return true;
 	}
@@ -139,6 +165,12 @@ void FCrashReportClient::FinalizeDiagnoseReportWorker( FText ReportText )
 	Uploader.LocalDiagnosisComplete(FPaths::FileExists(DiagnosticsFilePath) ? DiagnosticsFilePath : TEXT(""));
 }
 
+
+bool FCrashReportClient::IsProcessingCallstack() const
+{
+	return DiagnoseReportTask && !DiagnoseReportTask->IsWorkDone();
+}
+
 FDiagnoseReportWorker::FDiagnoseReportWorker( FCrashReportClient* InCrashReportClient ) 
 	: CrashReportClient( InCrashReportClient )
 {}
@@ -164,7 +196,7 @@ FText FCrashReportUtil::FormatDiagnosticText( const FText& DiagnosticText, const
 	}
 	else
 	{
-		return FText::Format( LOCTEXT( "CrashReportClientCallstackPattern", "MachineId:{0}\nUserName:{1}\n\n{2}" ), FText::FromString( MachineId ), FText::FromString( UserNameNoDot ), DiagnosticText );
+		return FText::Format( LOCTEXT( "CrashReportClientCallstackPattern", "MachineId:{0}\n\n{1}" ), FText::FromString( MachineId ), DiagnosticText );
 	}
 
 }
diff --git a/Engine/Source/Programs/CrashReporter/CrashReportClient/Private/CrashReportClient.h b/Engine/Source/Programs/CrashReporter/CrashReportClient/Private/CrashReportClient.h
--- a/Engine/Source/Programs/CrashReporter/CrashReportClient/Private/CrashReportClient.h
+++ b/Engine/Source/Programs/CrashReporter/CrashReportClient/Private/CrashReportClient.h
@@ -50,6 +50,9 @@ public:
 	 */
 	FCrashReportClient( const FPlatformErrorReport& InErrorReport );
 
+	/** Destructor. */
+	virtual ~FCrashReportClient();
+
 	/**
 	 * Respond to the user pressing Submit
 	 * @return Whether the request was handled
@@ -87,11 +90,20 @@ public:
 	 */
 	void RequestCloseWindow(const TSharedRef<SWindow>& Window);
 
+	/** Whether the main window should be hidden. */
 	bool ShouldWindowBeHidden() const
 	{
 		return bShouldWindowBeHidden;
 	}
 
+	/** Whether the app should enable widgets related to the displayed callstack. */
+	bool AreCallstackWidgetsEnabled() const;
+
+	/** Whether the throbber should be visible while processing the callstack. */
+	EVisibility IsThrobberVisible() const;
+
+	void SCrashReportClient_OnCheckStateChanged( ECheckBoxState NewRadioState );
+
 private:
 	/**
 	 * Write the user's comment to the report and begin uploading the entire report 
@@ -113,6 +125,11 @@ private:
 	/** Enqueued from the diagnose report worker thread to be executed on the game thread. */
 	void FinalizeDiagnoseReportWorker( FText ReportText );
 
+	/**
+	 * @return true if we are still processing a callstack
+	 */
+	bool IsProcessingCallstack() const;
+
 	/** Comment provided by the user */
 	FText UserComment;
 
@@ -131,9 +148,12 @@ private:
 	/** Whether BeginUpload has been called. */
 	bool bBeginUploadCalled;
 
-	/** Whether the main windows should be hidden. */
+	/** Whether the main window should be hidden. */
 	bool bShouldWindowBeHidden;
 
+	/** Whether the user allowed us to be contacted. */
+	bool bAllowToBeContacted;
+
 };
 
 #endif // !CRASH_REPORT_UNATTENDED_ONLY
diff --git a/Engine/Source/Programs/CrashReporter/CrashReportClient/Private/CrashReportClientApp.cpp b/Engine/Source/Programs/CrashReporter/CrashReportClient/Private/CrashReportClientApp.cpp
--- a/Engine/Source/Programs/CrashReporter/CrashReportClient/Private/CrashReportClientApp.cpp
+++ b/Engine/Source/Programs/CrashReporter/CrashReportClient/Private/CrashReportClientApp.cpp
@@ -125,7 +125,7 @@ void RunCrashReportClient(const TCHAR* CommandLine)
 
 	if (bUnattended)
 	{
-		ErrorReport.SetUserComment( NSLOCTEXT( "CrashReportClient", "UnattendedMode", "Sent in the unattended mode" ) );
+		ErrorReport.SetUserComment( NSLOCTEXT( "CrashReportClient", "UnattendedMode", "Sent in the unattended mode" ), false );
 		FCrashReportClientUnattended CrashReportClient( ErrorReport );
 
 		// loop until the app is ready to quit
@@ -160,7 +160,7 @@ void RunCrashReportClient(const TCHAR* CommandLine)
 		Window->SetRequestDestroyWindowOverride(FRequestDestroyWindowOverride::CreateSP(CrashReportClient, &FCrashReportClient::RequestCloseWindow));
 
 		// Setting focus seems to have to happen after the Window has been added
-		ClientControl->SetDefaultFocus();
+		FSlateApplication::Get().ClearKeyboardFocus(EFocusCause::Cleared);
 
 		// Debugging code
 		if (RunWidgetReflector)
diff --git a/Engine/Source/Programs/CrashReporter/CrashReportClient/Private/GenericErrorReport.cpp b/Engine/Source/Programs/CrashReporter/CrashReportClient/Private/GenericErrorReport.cpp
--- a/Engine/Source/Programs/CrashReporter/CrashReportClient/Private/GenericErrorReport.cpp
+++ b/Engine/Source/Programs/CrashReporter/CrashReportClient/Private/GenericErrorReport.cpp
@@ -41,7 +41,7 @@ FGenericErrorReport::FGenericErrorReport(const FString& Directory)
 	FPlatformFileManager::Get().GetPlatformFile().IterateDirectory(*ReportDirectory, FilenamesVisitor);
 }
 
-bool FGenericErrorReport::SetUserComment(const FText& UserComment)
+bool FGenericErrorReport::SetUserComment(const FText& UserComment, bool bAllowToBeContacted)
 {
 	// Find .xml file
 	FString XmlFilename;
@@ -53,7 +53,7 @@ bool FGenericErrorReport::SetUserComment(const FText& UserComment)
 	FString XmlFilePath = ReportDirectory / XmlFilename;
 	// FXmlFile's constructor loads the file to memory, closes the file and parses the data
 	FXmlFile XmlFile(XmlFilePath);
-	auto DynamicSignaturesNode = XmlFile.IsValid() ?
+	FXmlNode* DynamicSignaturesNode = XmlFile.IsValid() ?
 		XmlFile.GetRootNode()->FindChildNode(TEXT("DynamicSignatures")) :
 		nullptr;
 
@@ -62,8 +62,17 @@ bool FGenericErrorReport::SetUserComment(const FText& UserComment)
 		return false;
 	}
 
-	DynamicSignaturesNode->AppendChildNode(TEXT("Parameter3"), UserComment.ToString());
-
+	// Add or update the user comment.
+	FXmlNode* Parameter3Node = DynamicSignaturesNode->FindChildNode(TEXT("Parameter3"));
+	if( Parameter3Node )
+	{
+		Parameter3Node->SetContent(UserComment.ToString());
+	}
+	else
+	{
+		DynamicSignaturesNode->AppendChildNode(TEXT("Parameter3"), UserComment.ToString());
+	}
+	
 	FString MachineIDandUserID;
 	// Set global user name ID: will be added to the report
 	extern FCrashDescription& GetCrashDescription();
@@ -76,8 +85,28 @@ bool FGenericErrorReport::SetUserComment(const FText& UserComment)
 		MachineIDandUserID = FString::Printf( TEXT( "!MachineId:%s!Name:%s" ), *GetCrashDescription().MachineId, *GetCrashDescription().UserName );
 	}
 
-	// Also write a user ID into the report
-	DynamicSignaturesNode->AppendChildNode(TEXT("Parameter4"), MachineIDandUserID);
+	// Add or update a user ID.
+	FXmlNode* Parameter4Node = DynamicSignaturesNode->FindChildNode(TEXT("Parameter4"));
+	if( Parameter4Node )
+	{
+		Parameter4Node->SetContent(MachineIDandUserID);
+	}
+	else
+	{
+		DynamicSignaturesNode->AppendChildNode(TEXT("Parameter4"), MachineIDandUserID);
+	}
+
+	// Add or update bAllowToBeContacted
+	const FString AllowToBeContacted = bAllowToBeContacted ? TEXT("true") : TEXT("false");
+	FXmlNode* AllowToBeContactedNode = DynamicSignaturesNode->FindChildNode(TEXT("bAllowToBeContacted"));
+	if( AllowToBeContactedNode )
+	{
+		AllowToBeContactedNode->SetContent(AllowToBeContacted);
+	}
+	else
+	{
+		DynamicSignaturesNode->AppendChildNode(TEXT("bAllowToBeContacted"), AllowToBeContacted);
+	}
 
 	// Re-save over the top
 	return XmlFile.Save(XmlFilePath);
diff --git a/Engine/Source/Programs/CrashReporter/CrashReportClient/Private/GenericErrorReport.h b/Engine/Source/Programs/CrashReporter/CrashReportClient/Private/GenericErrorReport.h
--- a/Engine/Source/Programs/CrashReporter/CrashReportClient/Private/GenericErrorReport.h
+++ b/Engine/Source/Programs/CrashReporter/CrashReportClient/Private/GenericErrorReport.h
@@ -42,7 +42,7 @@ public:
 	 * @param UserComment Information provided by the user to add to the report
 	 * @return Whether the comment was successfully written to the report
 	 */
-	bool SetUserComment(const FText& UserComment);
+	bool SetUserComment(const FText& UserComment, bool bAllowToBeContacted);
 
 	/**
 	 * Provide full paths to all the report files
diff --git a/Engine/Source/Programs/CrashReporter/CrashReportClient/Private/SCrashReportClient.cpp b/Engine/Source/Programs/CrashReporter/CrashReportClient/Private/SCrashReportClient.cpp
--- a/Engine/Source/Programs/CrashReporter/CrashReportClient/Private/SCrashReportClient.cpp
+++ b/Engine/Source/Programs/CrashReporter/CrashReportClient/Private/SCrashReportClient.cpp
@@ -7,64 +7,117 @@
 #include "SCrashReportClient.h"
 #include "CrashReportClientStyle.h"
 #include "SlateStyle.h"
+#include "SThrobber.h"
 
 #define LOCTEXT_NAMESPACE "CrashReportClient"
 
+static void OnBrowserLinkClicked(const FSlateHyperlinkRun::FMetadata& Metadata, TSharedRef<SWidget> ParentWidget)
+{
+	const FString* UrlPtr = Metadata.Find(TEXT("href"));
+	if(UrlPtr)
+	{
+		FPlatformProcess::LaunchURL(**UrlPtr, nullptr, nullptr);
+	}
+}
+
 void SCrashReportClient::Construct(const FArguments& InArgs, TSharedRef<FCrashReportClient> Client)
 {
 	CrashReportClient = Client;
 
-	TSharedPtr<SVerticalBox> CrashedAppBox;
+	auto CrashedAppName = CrashReportClient->GetCrashedAppName();
+
+	// Set the text displaying the name of the crashed app, if available
+	FText CrashedAppText = CrashedAppName.IsEmpty() ?
+		LOCTEXT("CrashedAppNotFound", "An Unreal process has crashed") :
+		LOCTEXT("CrashedApp", "The following process has crashed: ");
+
 	ChildSlot
 	[
 		SNew(SBorder)
 		.BorderImage(FCoreStyle::Get().GetBrush("ToolPanel.GroupBorder"))
-		.Padding(FMargin(3, 0))
 		[
 			SNew(SVerticalBox)
 
 			// Stuff anchored to the top
 			+SVerticalBox::Slot()
 			.AutoHeight()
-			.Padding(5)
+			.Padding(3)
 			[
-				SNew(SVerticalBox)
+				SNew(SHorizontalBox)
+				+SHorizontalBox::Slot()
+				.AutoWidth()
+				[
+					SNew(STextBlock)
+					.TextStyle(FCrashReportClientStyle::Get(), "Title")
+					.Text(CrashedAppText)
+				]
 
-				+SVerticalBox::Slot()
-				.AutoHeight()
-				.Padding(10)
+				+SHorizontalBox::Slot()
+				.AutoWidth()
 				[
-					SAssignNew(CrashedAppBox, SVerticalBox)
+					SNew(STextBlock)
+					.TextStyle(FCrashReportClientStyle::Get(), "Title")
+					.Text(CrashedAppName)
 				]
+			]
 
-				+SVerticalBox::Slot()
+			+SVerticalBox::Slot()
+			.AutoHeight()
+			.Padding(3)
+			[
+				SNew( SRichTextBlock )
+				.Text( LOCTEXT("CrashDetailed", "We are very sorry that this crash occurred. Our goal is to prevent crashes like this from occurring in the future. Please help us track down and fix this crash by providing detailed information about what you were doing so that we may reproduce the crash and fix it quickly. You can also log a Bug Report with us at <a id=\"browser\" href=\"https://answers.unrealengine.com\" style=\"Hyperlink\">AnswerHub</> and work directly with support staff to report this issue.")
+				)
+				.AutoWrapText(true)
+				.DecoratorStyleSet( &FCoreStyle::Get() )
+				+ SRichTextBlock::HyperlinkDecorator( TEXT("browser"), FSlateHyperlinkRun::FOnClick::CreateStatic( &OnBrowserLinkClicked, AsShared() ) )
+			]
+
+			+SVerticalBox::Slot()
 				.AutoHeight()
-				.Padding(5)
+				.Padding(3)
 				[
-					SNew(STextBlock)
+					SNew( STextBlock )
+					.Text( LOCTEXT("CrashThanks", "Thanks for your help in improving the Unreal Engine.")
+					)
 					.AutoWrapText(true)
-					.Text(LOCTEXT("Blurb", "Please describe what you were doing immediately before the crash:"))
 				]
 
-				+SVerticalBox::Slot()
+			+SVerticalBox::Slot()
 				.AutoHeight()
-				.Padding(5)
+				.Padding(3)
 				[
-					SAssignNew(UserCommentBox, SEditableTextBox)
-					.OnTextCommitted(CrashReportClient.ToSharedRef(), &FCrashReportClient::UserCommentChanged)
-					.HintText(LOCTEXT("CircumstancesOfCrash", "Circumstances of crash"))
+					SNew( STextBlock )
+					.Text( LOCTEXT("CrashPRovide", "Please provide detailed information about what you were doing when the crash occurred.")
+					)
+					.AutoWrapText(true)
 				]
-			]
 
 			+SVerticalBox::Slot()
+			.FillHeight(0.3f)
+			.Padding(3)
 			[
-				SNew(SVerticalBox)
-				+SVerticalBox::Slot()
-				.MaxHeight(550)
+				SNew(SMultiLineEditableTextBox)
+				.OnTextCommitted(CrashReportClient.ToSharedRef(), &FCrashReportClient::UserCommentChanged)
+				.Font( FSlateFontInfo( FPaths::EngineContentDir() / TEXT( "Slate/Fonts/Roboto-Regular.ttf" ), 9 ) )
+				.AutoWrapText( true )
+				.BackgroundColor(FSlateColor(FLinearColor::Black))
+				.ForegroundColor(FSlateColor(FLinearColor::White))
+				.HintText(LOCTEXT("CircumstancesOfCrash", "Circumstances of crash"))
+			]
+			
+			+SVerticalBox::Slot()
+			.FillHeight(0.7f)
+			.Padding(3)
+			[
+				SNew(SOverlay)
+
+				+SOverlay::Slot()
 				[
 					SNew(SBorder)
 					.BorderImage(new FSlateBoxBrush(TEXT("Common/TextBlockHighlightShape"), FMargin(.5f)))
 					.BorderBackgroundColor(FSlateColor(FLinearColor::Black))
+					.IsEnabled(CrashReportClient.ToSharedRef(), &FCrashReportClient::AreCallstackWidgetsEnabled)
 					[
 						SNew(SScrollBox)
 						+SScrollBox::Slot()
@@ -76,95 +129,95 @@ void SCrashReportClient::Construct(const FArguments& InArgs, TSharedRef<FCrashRe
 					]
 				]
 
-				+SVerticalBox::Slot()
-				.AutoHeight()
-				.HAlign(HAlign_Right)
-				.Padding(2)
+				+SOverlay::Slot()
+				.HAlign(HAlign_Center)
+				.VAlign(VAlign_Center)
 				[
-					SNew(SButton)
-					.Text(LOCTEXT("CopyCallstack", "Copy Callstack to Clipboard"))
-					.OnClicked(CrashReportClient.ToSharedRef(), &FCrashReportClient::CopyCallstack)
+					SNew(SThrobber)
+					.Visibility(CrashReportClient.ToSharedRef(), &FCrashReportClient::IsThrobberVisible)
+					.NumPieces(5)
 				]
 			]
 
-			// Stuff anchored to the bottom
 			+SVerticalBox::Slot()
 			.AutoHeight()
-			.Padding(10)
+			.Padding(3)
 			[
 				SNew(SHorizontalBox)
 
 				+SHorizontalBox::Slot()
-				.Padding(3)
+				.AutoWidth()
+				.VAlign(VAlign_Center)
 				[
-					SNew(STextBlock)
-					.Text(LOCTEXT("KeyboardShortcuts", "Shift+Enter: submit\n"))
+					SNew(SCheckBox)
+					.IsChecked(ECheckBoxState::Checked)
+					.OnCheckStateChanged(CrashReportClient.ToSharedRef(), &FCrashReportClient::SCrashReportClient_OnCheckStateChanged)
 				]
 
 				+SHorizontalBox::Slot()
-				.VAlign(VAlign_Bottom)
-				.HAlign(HAlign_Right)
-				.AutoWidth()
+				.FillWidth(1.0f)
+				.VAlign(VAlign_Center)
 				[
-					SNew(SUniformGridPanel)
-					.SlotPadding(FCoreStyle::Get().GetMargin("StandardDialog.SlotPadding"))
-					.MinDesiredSlotWidth(FCoreStyle::Get().GetFloat("StandardDialog.MinDesiredSlotWidth"))
-					.MinDesiredSlotHeight(FCoreStyle::Get().GetFloat("StandardDialog.MinDesiredSlotHeight"))
-					+SUniformGridPanel::Slot(0, 0)
-					[
-						SNew(SButton)
-						.HAlign(HAlign_Center)
-						.Text(LOCTEXT("Send", "Send"))
-						.OnClicked(Client, &FCrashReportClient::Submit)
-					]
+					SNew(STextBlock)
+					.AutoWrapText(true)
+					.Text(LOCTEXT("IAgree", "I agree to be contacted by Epic Games via email if additional information about this crash would help fix it."))
 				]
 			]
-		]
-	];
 
-	auto CrashedAppName = CrashReportClient->GetCrashedAppName();
+			// Stuff anchored to the bottom
+			+SVerticalBox::Slot()
+			.AutoHeight()
+			.Padding( FMargin(3, 3+16, 3, 3) )
+			[
+				SNew(SHorizontalBox)
 
-	// Set the text displaying the name of the crashed app, if available
-	FText CrashedAppText = CrashedAppName.IsEmpty() ?
-		LOCTEXT("CrashedAppNotFound", "An Unreal process has crashed") :
-		LOCTEXT("CrashedApp", "The following process has crashed:");
+				+SHorizontalBox::Slot()
+				.FillWidth(1.0f)
+				.HAlign(HAlign_Left)
+				.VAlign(VAlign_Center)
+				.Padding(0)
+				[			
+					SNew(SSpacer)
+				]
 
-	CrashedAppBox->AddSlot()
-	[
-		SNew(STextBlock)
-		.TextStyle(FCrashReportClientStyle::Get(), "Title")
-		.Text(CrashedAppText)
-	];
+				+SHorizontalBox::Slot()
+				.AutoWidth()
+				.HAlign(HAlign_Center)
+				.VAlign(VAlign_Center)
+				.Padding( FMargin(0,0,8,0) )
+				[
+					SNew(SButton)
+					.ContentPadding( FMargin(8,2) )
+					.Text(LOCTEXT("CopyCallstack", "Copy to Clipboard"))
+					.OnClicked(CrashReportClient.ToSharedRef(), &FCrashReportClient::CopyCallstack)
+					.IsEnabled(CrashReportClient.ToSharedRef(), &FCrashReportClient::AreCallstackWidgetsEnabled)
+				]
 
-	if (!CrashedAppName.IsEmpty())
-	{
-		CrashedAppBox->AddSlot()
-		.Padding(5)
-		.HAlign(HAlign_Center)
-		[
-			SNew(STextBlock)
-			.Text(FText::FromString(CrashedAppName))
-		];
-	}
+				+SHorizontalBox::Slot()
+				.HAlign(HAlign_Center)
+				.VAlign(VAlign_Center)
+				.AutoWidth()
+				.Padding( FMargin(0) )
+				[
+					SNew(SButton)
+					.ContentPadding( FMargin(8,2) )
+					.Text(LOCTEXT("Send", "Send"))
+					.OnClicked(Client, &FCrashReportClient::Submit)
+				]
+			]
+		]
+	];
 
 	FSlateApplication::Get().SetUnhandledKeyDownEventHandler(FOnKeyEvent::CreateSP(this, &SCrashReportClient::OnUnhandledKeyDown));
 }
 
-void SCrashReportClient::SetDefaultFocus()
-{
-	FSlateApplication::Get().SetKeyboardFocus(UserCommentBox.ToSharedRef());
-}
-
 FReply SCrashReportClient::OnUnhandledKeyDown(const FKeyEvent& InKeyEvent)
 {
 	const FKey Key = InKeyEvent.GetKey();
 	if (Key == EKeys::Enter)
 	{
-		if (InKeyEvent.IsShiftDown())
-		{
-			CrashReportClient->Submit();
-			return FReply::Handled();
-		}
+		CrashReportClient->Submit();
+		return FReply::Handled();
 	}
 
 	return FReply::Unhandled();
diff --git a/Engine/Source/Programs/CrashReporter/CrashReportClient/Private/SCrashReportClient.h b/Engine/Source/Programs/CrashReporter/CrashReportClient/Private/SCrashReportClient.h
--- a/Engine/Source/Programs/CrashReporter/CrashReportClient/Private/SCrashReportClient.h
+++ b/Engine/Source/Programs/CrashReporter/CrashReportClient/Private/SCrashReportClient.h
@@ -28,11 +28,6 @@ public:
 	 */
 	void Construct(const FArguments& InArgs, TSharedRef<FCrashReportClient> Client);
 
-	/**
-	 * Give the edit box focus (seems this has to be done after the window has been created)
-	 */
-	void SetDefaultFocus();
-
 private:
 	/**
 	 * Keyboard short-cut handler
@@ -43,9 +38,6 @@ private:
 
 	/** Crash report client implementation object */
 	TSharedPtr<FCrashReportClient> CrashReportClient;
-
-	/** The edit box that has keyboard focus to receive the user's comment */
-	TSharedPtr<SEditableTextBox> UserCommentBox;
 };
 
 #endif // !CRASH_REPORT_UNATTENDED_ONLY
