CrashReportWebsite - Added support for bAllowToBeContacted from the crash report client (UECORE-77-Crash Reporter: Option for users to opt out of being contacted)

diff --git a/Engine/Source/Programs/CrashReporter/CrashReportWebSite/Models/CrashReport.cs b/Engine/Source/Programs/CrashReporter/CrashReportWebSite/Models/CrashReport.cs
--- a/Engine/Source/Programs/CrashReporter/CrashReportWebSite/Models/CrashReport.cs
+++ b/Engine/Source/Programs/CrashReporter/CrashReportWebSite/Models/CrashReport.cs
@@ -239,6 +239,15 @@ namespace Tools.CrashReporter.CrashReportWebSite.Models
 			return Properties.Settings.Default.CrashReporterVideos + Id + "_CrashVideo.avi";
 		}
 
+		/// <summary> Whether the user allowed us to be contacted. </summary>
+		public bool AllowToBeContacted
+		{
+			get
+			{
+				return bAllowToBeContacted.GetValueOrDefault( false );
+			}
+		}
+			
 		/// <summary>
 		/// Return lines of processed callstack entries.
 		/// </summary>
diff --git a/Engine/Source/Programs/CrashReporter/CrashReportWebSite/Models/CrashReport.dbml b/Engine/Source/Programs/CrashReporter/CrashReportWebSite/Models/CrashReport.dbml
--- a/Engine/Source/Programs/CrashReporter/CrashReportWebSite/Models/CrashReport.dbml
+++ b/Engine/Source/Programs/CrashReporter/CrashReportWebSite/Models/CrashReport.dbml
@@ -61,7 +61,6 @@
   <Table Name="dbo.Crashes" Member="Crashes">
     <Type Name="Crash">
       <Column Name="Id" Type="System.Int32" DbType="Int NOT NULL IDENTITY" IsPrimaryKey="true" IsDbGenerated="true" CanBeNull="false" />
-      <Column Name="Title" Type="System.String" DbType="NChar(20)" CanBeNull="true" />
       <Column Name="Summary" Type="System.String" DbType="VarChar(512)" CanBeNull="true" />
       <Column Name="GameName" Type="System.String" DbType="VarChar(64)" CanBeNull="true" />
       <Column Name="Status" Type="System.String" DbType="VarChar(64)" CanBeNull="true" />
@@ -74,22 +73,18 @@
       <Column Name="Pattern" Type="System.String" DbType="VarChar(800)" CanBeNull="true" />
       <Column Name="CommandLine" Type="System.String" DbType="VarChar(512)" CanBeNull="true" />
       <Column Name="ComputerName" Type="System.String" DbType="VarChar(64)" CanBeNull="true" />
-      <Column Name="Selected" Type="System.Boolean" DbType="Bit" CanBeNull="true" />
       <Column Name="FixedChangeList" Type="System.String" DbType="VarChar(256)" CanBeNull="true" />
       <Column Name="LanguageExt" Type="System.String" DbType="VarChar(64)" CanBeNull="true" />
       <Column Name="Module" Type="System.String" DbType="VarChar(128)" CanBeNull="true" />
       <Column Name="BuildVersion" Type="System.String" DbType="VarChar(128)" CanBeNull="true" />
       <Column Name="BaseDir" Type="System.String" DbType="VarChar(512)" CanBeNull="true" />
-      <Column Name="Version" Type="System.Int32" DbType="Int" CanBeNull="true" />
       <Column Name="UserName" Type="System.String" DbType="VarChar(64)" CanBeNull="true" />
       <Column Name="TTPID" Type="System.String" DbType="VarChar(64)" CanBeNull="true" />
-      <Column Name="AutoReporterID" Type="System.Int32" DbType="Int" CanBeNull="true" />
-      <Column Name="Processed" Type="System.Boolean" DbType="Bit" CanBeNull="true" />
+      <Column Name="Processed" Member="bAllowToBeContacted" Storage="_Processed" Type="System.Boolean" DbType="Bit" CanBeNull="true" />
       <Column Name="HasLogFile" Type="System.Boolean" DbType="Bit" CanBeNull="true" />
       <Column Name="HasMiniDumpFile" Type="System.Boolean" DbType="Bit" CanBeNull="true" />
       <Column Name="HasVideoFile" Type="System.Boolean" DbType="Bit" CanBeNull="true" />
       <Column Name="HasDiagnosticsFile" Type="System.Boolean" DbType="Bit" CanBeNull="true" />
-      <Column Name="HasNewLogFile" Type="System.Boolean" DbType="Bit" CanBeNull="true" />
       <Column Name="Branch" Type="System.String" DbType="VarChar(32)" CanBeNull="true" />
       <Column Name="UserNameId" Type="System.Int32" DbType="Int" CanBeNull="true" />
       <Column Name="CrashType" Type="System.Int16" DbType="SmallInt" CanBeNull="true" />
diff --git a/Engine/Source/Programs/CrashReporter/CrashReportWebSite/Models/CrashReport.dbml.layout b/Engine/Source/Programs/CrashReporter/CrashReportWebSite/Models/CrashReport.dbml.layout
--- a/Engine/Source/Programs/CrashReporter/CrashReportWebSite/Models/CrashReport.dbml.layout
+++ b/Engine/Source/Programs/CrashReporter/CrashReportWebSite/Models/CrashReport.dbml.layout
@@ -32,16 +32,10 @@
         <elementListCompartment Id="77c436bb-7563-449c-a08a-8c417bcde3fa" absoluteBounds="3.3900000000000006, 6.3350000000000009, 1.9700000000000002, 0.8262939453125" name="DataPropertiesCompartment" titleTextColor="Black" itemTextColor="Black" />
       </nestedChildShapes>
     </classShape>
-    <classShape Id="f1204ef8-6660-48ce-8c1d-131c49e2dd25" absoluteBounds="6.125, 1.125, 2, 7.9245426432291648">
+    <classShape Id="f1204ef8-6660-48ce-8c1d-131c49e2dd25" absoluteBounds="6.125, 1.125, 2, 6.9630354817708326">
       <DataClassMoniker Name="/CrashReportDataContext/Crash" />
       <nestedChildShapes>
-        <elementListCompartment Id="53c11757-f0be-47da-bcae-7d20a571002e" absoluteBounds="6.1400000000000006, 1.5850000000000009, 1.9700000000000002, 7.3645426432291661" name="DataPropertiesCompartment" titleTextColor="Black" itemTextColor="Black" />
-      </nestedChildShapes>
-    </classShape>
-    <classShape Id="a188623d-4f2c-4a7c-b729-d354a2df0d5a" absoluteBounds="9, 6.25, 2, 1.1939925130208344">
-      <DataClassMoniker Name="/CrashReportDataContext/Crash_FunctionCall" />
-      <nestedChildShapes>
-        <elementListCompartment Id="ecfce9b6-c6ab-4440-91c4-a9cc7bb73ecf" absoluteBounds="9.015, 6.7100000000000009, 1.9700000000000002, 0.63399251302083326" name="DataPropertiesCompartment" titleTextColor="Black" itemTextColor="Black" />
+        <elementListCompartment Id="53c11757-f0be-47da-bcae-7d20a571002e" absoluteBounds="6.1400000000000006, 1.5850000000000009, 1.9700000000000002, 6.403035481770833" name="DataPropertiesCompartment" titleTextColor="Black" itemTextColor="Black" />
       </nestedChildShapes>
     </classShape>
     <classShape Id="01d6823d-5d7e-435f-99a2-7d6fac1e698d" absoluteBounds="8.75, 4, 2, 1.193992513020838">
@@ -63,7 +57,7 @@
         <classShapeMoniker Id="b4a7205e-4223-43f1-837f-ac4a1390bf89" />
       </nodes>
     </associationConnector>
-    <associationConnector edgePoints="[(5.375 : 4.82381790161133); (6.125 : 4.82381790161133)]" fixedFrom="NotFixed" fixedTo="NotFixed">
+    <associationConnector edgePoints="[(5.375 : 4.84084598496556); (6.125 : 4.84084598496556)]" fixedFrom="Algorithm" fixedTo="Algorithm">
       <AssociationMoniker Name="/CrashReportDataContext/User/User_Crash" />
       <nodes>
         <classShapeMoniker Id="01dddf4f-058f-43b4-b6c0-b90e7acfe0b9" />
@@ -91,28 +85,14 @@
         <classShapeMoniker Id="01dddf4f-058f-43b4-b6c0-b90e7acfe0b9" />
       </nodes>
     </associationConnector>
-    <associationConnector edgePoints="[(8.125 : 2.47199625651042); (8.875 : 2.47199625651042)]" fixedFrom="NotFixed" fixedTo="NotFixed">
+    <associationConnector edgePoints="[(8.125 : 2.47199625651042); (8.875 : 2.47199625651042)]" fixedFrom="Algorithm" fixedTo="Algorithm">
       <AssociationMoniker Name="/CrashReportDataContext/Crash/Crash_Buggs_Crash" />
       <nodes>
         <classShapeMoniker Id="f1204ef8-6660-48ce-8c1d-131c49e2dd25" />
         <classShapeMoniker Id="66e20693-2188-41ba-b050-7d08c2691d12" />
       </nodes>
     </associationConnector>
-    <associationConnector edgePoints="[(8.125 : 6.84699625651042); (9 : 6.84699625651042)]" fixedFrom="NotFixed" fixedTo="NotFixed">
-      <AssociationMoniker Name="/CrashReportDataContext/Crash/Crash_Crash_FunctionCall" />
-      <nodes>
-        <classShapeMoniker Id="f1204ef8-6660-48ce-8c1d-131c49e2dd25" />
-        <classShapeMoniker Id="a188623d-4f2c-4a7c-b729-d354a2df0d5a" />
-      </nodes>
-    </associationConnector>
-    <associationConnector edgePoints="[(9.875 : 5.19399251302084); (9.875 : 6.25)]" fixedFrom="NotFixed" fixedTo="NotFixed">
-      <AssociationMoniker Name="/CrashReportDataContext/FunctionCall/FunctionCall_Crash_FunctionCall" />
-      <nodes>
-        <classShapeMoniker Id="01d6823d-5d7e-435f-99a2-7d6fac1e698d" />
-        <classShapeMoniker Id="a188623d-4f2c-4a7c-b729-d354a2df0d5a" />
-      </nodes>
-    </associationConnector>
-    <associationConnector edgePoints="[(5.375 : 4.28690895080566); (6.125 : 4.28690895080566)]" fixedFrom="NotFixed" fixedTo="NotFixed">
+    <associationConnector edgePoints="[(5.375 : 4.29542299248278); (6.125 : 4.29542299248278)]" fixedFrom="Algorithm" fixedTo="Algorithm">
       <AssociationMoniker Name="/CrashReportDataContext/User/User_Crash1" />
       <nodes>
         <classShapeMoniker Id="01dddf4f-058f-43b4-b6c0-b90e7acfe0b9" />
diff --git a/Engine/Source/Programs/CrashReporter/CrashReportWebSite/Models/CrashRepository.cs b/Engine/Source/Programs/CrashReporter/CrashReportWebSite/Models/CrashRepository.cs
--- a/Engine/Source/Programs/CrashReporter/CrashReportWebSite/Models/CrashRepository.cs
+++ b/Engine/Source/Programs/CrashReporter/CrashReportWebSite/Models/CrashRepository.cs
@@ -904,6 +904,8 @@ namespace Tools.CrashReporter.CrashReportWebSite.Models
 			NewCrash.HasVideoFile = NewCrashInfo.bHasVideo;
 			NewCrash.HasMetaData = NewCrashInfo.bHasWERData;
 
+			NewCrash.bAllowToBeContacted = NewCrashInfo.bAllowToBeContacted;
+
 			NewCrash.TTPID = "";
 			NewCrash.FixedChangeList = "";
 
@@ -931,19 +933,19 @@ namespace Tools.CrashReporter.CrashReportWebSite.Models
 				Selected			bit				
 				Version				int				
 				AutoReporterID		int				
-				Processed			bit				
-				HasDiagnosticsFile	bit				
+				Processed			bit	-> renamed to AllowToBeContacted			
+				HasDiagnosticsFile	bit	always true		
 				HasNewLogFile		bit				
-				HasMetaData			bit	
+				HasMetaData			bit	always true
 			*/
 			// Set the unused fields to the default values.
-			NewCrash.Title = "";
-			NewCrash.Selected = false;
-			NewCrash.Version = 4;
-			NewCrash.AutoReporterID = 0;
-			NewCrash.Processed = true;
-			//NewCrash.HasDiagnosticsFile = true;
-			NewCrash.HasNewLogFile = false;
+			//NewCrash.Title = "";			removed from dbml
+			//NewCrash.Selected = false;	removed from dbml
+			//NewCrash.Version = 4;			removed from dbml
+			//NewCrash.AutoReporterID = 0;	removed from dbml
+			//NewCrash.HasNewLogFile = false;removed from dbml
+			// 
+			//NewCrash.HasDiagnosticsFile = true;		
 			//NewCrash.HasMetaData = true;
 
 			// Add the crash to the database
diff --git a/Engine/Source/Programs/CrashReporter/CrashReportWebSite/Views/Crashes/Show.aspx b/Engine/Source/Programs/CrashReporter/CrashReportWebSite/Views/Crashes/Show.aspx
--- a/Engine/Source/Programs/CrashReporter/CrashReportWebSite/Views/Crashes/Show.aspx
+++ b/Engine/Source/Programs/CrashReporter/CrashReportWebSite/Views/Crashes/Show.aspx
@@ -193,6 +193,9 @@
 			<dt>Epic Account Id</dt> 
 				<dd class='even'><%=Html.DisplayFor(m => Model.Crash.EpicAccountId) %></dd>
 
+			<dt>Allowed to contact</dt> 
+				<dd class='even'><%=Html.DisplayFor(m => Model.Crash.AllowToBeContacted) %></dd>
+
 			<dt>Build Version</dt>
 				<dd ><%=Html.DisplayFor(m => Model.Crash.BuildVersion) %></dd>
 
