CrashReportClient - Tweaks to the UI

diff --git a/Engine/Source/Programs/CrashReporter/CrashReportClient/Private/SCrashReportClient.cpp b/Engine/Source/Programs/CrashReporter/CrashReportClient/Private/SCrashReportClient.cpp
--- a/Engine/Source/Programs/CrashReporter/CrashReportClient/Private/SCrashReportClient.cpp
+++ b/Engine/Source/Programs/CrashReporter/CrashReportClient/Private/SCrashReportClient.cpp
@@ -74,36 +74,41 @@ void SCrashReportClient::Construct(const FArguments& InArgs, TSharedRef<FCrashRe
 			]
 
 			+SVerticalBox::Slot()
-				.AutoHeight()
-				.Padding(3)
-				[
-					SNew( STextBlock )
-					.Text( LOCTEXT("CrashThanks", "Thanks for your help in improving the Unreal Engine.")
-					)
-					.AutoWrapText(true)
-				]
-
-			+SVerticalBox::Slot()
-				.AutoHeight()
-				.Padding(3)
-				[
-					SNew( STextBlock )
-					.Text( LOCTEXT("CrashPRovide", "Please provide detailed information about what you were doing when the crash occurred.")
-					)
-					.AutoWrapText(true)
-				]
+			.AutoHeight()
+			.Padding( FMargin(3,6) )
+			[
+				SNew( STextBlock )
+				.Text( LOCTEXT("CrashThanks", "Thanks for your help in improving the Unreal Engine."))
+				.AutoWrapText(true)
+			]
 
 			+SVerticalBox::Slot()
 			.FillHeight(0.3f)
 			.Padding(3)
 			[
-				SNew(SMultiLineEditableTextBox)
-				.OnTextCommitted(CrashReportClient.ToSharedRef(), &FCrashReportClient::UserCommentChanged)
-				.Font( FSlateFontInfo( FPaths::EngineContentDir() / TEXT( "Slate/Fonts/Roboto-Regular.ttf" ), 9 ) )
-				.AutoWrapText( true )
-				.BackgroundColor(FSlateColor(FLinearColor::Black))
-				.ForegroundColor(FSlateColor(FLinearColor::White))
-				.HintText(LOCTEXT("CircumstancesOfCrash", "Circumstances of crash"))
+				SNew( SOverlay )
+
+				+ SOverlay::Slot()
+				[
+					SAssignNew( CrashDetailsInformation, SMultiLineEditableTextBox )
+					.OnTextCommitted( CrashReportClient.ToSharedRef(), &FCrashReportClient::UserCommentChanged )
+					.Font( FSlateFontInfo( FPaths::EngineContentDir() / TEXT( "Slate/Fonts/Roboto-Regular.ttf" ), 9 ) )
+					.AutoWrapText( true )
+					.BackgroundColor( FSlateColor( FLinearColor::Black ) )
+					.ForegroundColor( FSlateColor( FLinearColor::White * 0.8f ) )
+				]
+
+				// HintText is not implemented in SMultiLineEditableTextBox, so this is a workaround.
+				+ SOverlay::Slot()
+				
+				[
+					SNew(STextBlock)
+					.Margin( FMargin(4,2,0,0) )
+					.Font( FSlateFontInfo( FPaths::EngineContentDir() / TEXT( "Slate/Fonts/Roboto-Regular.ttf" ), 9 ) )
+					.ColorAndOpacity( FSlateColor( FLinearColor::White * 0.8f ) )
+					.Text( LOCTEXT( "CrashProvide", "Please provide detailed information about what you were doing when the crash occurred." ) )
+					.Visibility( this, &SCrashReportClient::IsHintTextVisible )
+				]	
 			]
 			
 			+SVerticalBox::Slot()
@@ -111,24 +116,19 @@ void SCrashReportClient::Construct(const FArguments& InArgs, TSharedRef<FCrashRe
 			.Padding(3)
 			[
 				SNew(SOverlay)
-
+		
 				+SOverlay::Slot()
 				[
-					SNew(SBorder)
-					.BorderImage(new FSlateBoxBrush(TEXT("Common/TextBlockHighlightShape"), FMargin(.5f)))
-					.BorderBackgroundColor(FSlateColor(FLinearColor::Black))
-					.IsEnabled(CrashReportClient.ToSharedRef(), &FCrashReportClient::AreCallstackWidgetsEnabled)
-					[
-						SNew(SScrollBox)
-						+SScrollBox::Slot()
-						[
-							SNew(STextBlock)
-							.TextStyle(FCrashReportClientStyle::Get(), "Code")
-							.Text(Client, &FCrashReportClient::GetDiagnosticText)
-						]
-					]
+					SNew( SMultiLineEditableTextBox )
+					.Font( FSlateFontInfo( FPaths::EngineContentDir() / TEXT( "Slate/Fonts/Roboto-Regular.ttf" ), 8 ) )
+					.AutoWrapText( false )
+					.IsReadOnly( true )
+					.BackgroundColor( FSlateColor( FLinearColor::Black ) )
+					.ForegroundColor( FSlateColor( FLinearColor::White * 0.8f ) )
+					.Text( Client, &FCrashReportClient::GetDiagnosticText )
 				]
 
+
 				+SOverlay::Slot()
 				.HAlign(HAlign_Center)
 				.VAlign(VAlign_Center)
@@ -141,7 +141,7 @@ void SCrashReportClient::Construct(const FArguments& InArgs, TSharedRef<FCrashRe
 
 			+SVerticalBox::Slot()
 			.AutoHeight()
-			.Padding(3)
+			.Padding( FMargin( 3, 6 ) )
 			[
 				SNew(SHorizontalBox)
 
@@ -180,19 +180,6 @@ void SCrashReportClient::Construct(const FArguments& InArgs, TSharedRef<FCrashRe
 					SNew(SSpacer)
 				]
 
-				+SHorizontalBox::Slot()
-				.AutoWidth()
-				.HAlign(HAlign_Center)
-				.VAlign(VAlign_Center)
-				.Padding( FMargin(0,0,8,0) )
-				[
-					SNew(SButton)
-					.ContentPadding( FMargin(8,2) )
-					.Text(LOCTEXT("CopyCallstack", "Copy to Clipboard"))
-					.OnClicked(CrashReportClient.ToSharedRef(), &FCrashReportClient::CopyCallstack)
-					.IsEnabled(CrashReportClient.ToSharedRef(), &FCrashReportClient::AreCallstackWidgetsEnabled)
-				]
-
 				+SHorizontalBox::Slot()
 				.HAlign(HAlign_Center)
 				.VAlign(VAlign_Center)
@@ -223,6 +210,11 @@ FReply SCrashReportClient::OnUnhandledKeyDown(const FKeyEvent& InKeyEvent)
 	return FReply::Unhandled();
 }
 
+EVisibility SCrashReportClient::IsHintTextVisible() const
+{
+	return CrashDetailsInformation->GetText().IsEmpty() ? EVisibility::HitTestInvisible : EVisibility::Hidden;
+}
+
 #undef LOCTEXT_NAMESPACE
 
 #endif // !CRASH_REPORT_UNATTENDED_ONLY
diff --git a/Engine/Source/Programs/CrashReporter/CrashReportClient/Private/SCrashReportClient.h b/Engine/Source/Programs/CrashReporter/CrashReportClient/Private/SCrashReportClient.h
--- a/Engine/Source/Programs/CrashReporter/CrashReportClient/Private/SCrashReportClient.h
+++ b/Engine/Source/Programs/CrashReporter/CrashReportClient/Private/SCrashReportClient.h
@@ -36,8 +36,13 @@ private:
 	 */
 	FReply OnUnhandledKeyDown(const FKeyEvent& InKeyEvent);
 
+	/** Whether the hint text should be visible. */
+	EVisibility IsHintTextVisible() const;
+
 	/** Crash report client implementation object */
 	TSharedPtr<FCrashReportClient> CrashReportClient;
+
+	TSharedPtr<SMultiLineEditableTextBox> CrashDetailsInformation;
 };
 
 #endif // !CRASH_REPORT_UNATTENDED_ONLY
