CrashReportClient - Tweaks to the UI

diff --git a/Engine/Source/Programs/CrashReporter/CrashReportClient/Private/CrashReportClientStyle.cpp b/Engine/Source/Programs/CrashReporter/CrashReportClient/Private/CrashReportClientStyle.cpp
--- a/Engine/Source/Programs/CrashReporter/CrashReportClient/Private/CrashReportClientStyle.cpp
+++ b/Engine/Source/Programs/CrashReporter/CrashReportClient/Private/CrashReportClientStyle.cpp
@@ -25,6 +25,7 @@ void FCrashReportClientStyle::Shutdown()
 }
 
 #define TTF_FONT(RelativePath, ...) FSlateFontInfo(ContentFromEngine(TEXT(RelativePath), TEXT(".ttf")), __VA_ARGS__)
+#define IMAGE_BRUSH( RelativePath, ... ) FSlateImageBrush(Style.RootToContentDir(RelativePath, TEXT(".png")), __VA_ARGS__)
 
 namespace
 {
@@ -60,6 +61,19 @@ TSharedRef< FSlateStyleSet > FCrashReportClientStyle::Create()
 		.SetColorAndOpacity(FSlateColor::UseSubduedForeground())
 	);
 
+	const FVector2D Icon16x16( 16.0f, 16.0f );
+	FSlateBrush* GenericWhiteBox = new IMAGE_BRUSH( "Old/White", Icon16x16 );
+
+	// SEditableTextBox defaults...
+	const FEditableTextBoxStyle NormalEditableTextBoxStyle = FEditableTextBoxStyle()
+		.SetBackgroundImageNormal( *GenericWhiteBox )
+		.SetBackgroundImageHovered( *GenericWhiteBox )
+		.SetBackgroundImageFocused( *GenericWhiteBox )
+		.SetBackgroundImageReadOnly( *GenericWhiteBox );
+	{
+		Style.Set( "NormalEditableTextBox", NormalEditableTextBoxStyle );
+	}
+
 	return StyleRef;
 }
 
diff --git a/Engine/Source/Programs/CrashReporter/CrashReportClient/Private/SCrashReportClient.cpp b/Engine/Source/Programs/CrashReporter/CrashReportClient/Private/SCrashReportClient.cpp
--- a/Engine/Source/Programs/CrashReporter/CrashReportClient/Private/SCrashReportClient.cpp
+++ b/Engine/Source/Programs/CrashReporter/CrashReportClient/Private/SCrashReportClient.cpp
@@ -41,7 +41,7 @@ void SCrashReportClient::Construct(const FArguments& InArgs, TSharedRef<FCrashRe
 			// Stuff anchored to the top
 			+SVerticalBox::Slot()
 			.AutoHeight()
-			.Padding(3)
+			.Padding(4)
 			[
 				SNew(SHorizontalBox)
 				+SHorizontalBox::Slot()
@@ -63,7 +63,7 @@ void SCrashReportClient::Construct(const FArguments& InArgs, TSharedRef<FCrashRe
 
 			+SVerticalBox::Slot()
 			.AutoHeight()
-			.Padding(3)
+			.Padding( FMargin( 4, 10 ) )
 			[
 				SNew( SRichTextBlock )
 				.Text( LOCTEXT("CrashDetailed", "We are very sorry that this crash occurred. Our goal is to prevent crashes like this from occurring in the future. Please help us track down and fix this crash by providing detailed information about what you were doing so that we may reproduce the crash and fix it quickly. You can also log a Bug Report with us at <a id=\"browser\" href=\"https://answers.unrealengine.com\" style=\"Hyperlink\">AnswerHub</> and work directly with support staff to report this issue.")
@@ -75,7 +75,7 @@ void SCrashReportClient::Construct(const FArguments& InArgs, TSharedRef<FCrashRe
 
 			+SVerticalBox::Slot()
 			.AutoHeight()
-			.Padding( FMargin(3,6) )
+			.Padding( FMargin( 4, 10 ) )
 			[
 				SNew( STextBlock )
 				.Text( LOCTEXT("CrashThanks", "Thanks for your help in improving the Unreal Engine."))
@@ -84,13 +84,14 @@ void SCrashReportClient::Construct(const FArguments& InArgs, TSharedRef<FCrashRe
 
 			+SVerticalBox::Slot()
 			.FillHeight(0.3f)
-			.Padding(3)
+			.Padding( FMargin( 4, 10, 4, 4 ) )
 			[
 				SNew( SOverlay )
 
 				+ SOverlay::Slot()
 				[
 					SAssignNew( CrashDetailsInformation, SMultiLineEditableTextBox )
+					.Style( &FCrashReportClientStyle::Get().GetWidgetStyle<FEditableTextBoxStyle>( "NormalEditableTextBox" ) )
 					.OnTextCommitted( CrashReportClient.ToSharedRef(), &FCrashReportClient::UserCommentChanged )
 					.Font( FSlateFontInfo( FPaths::EngineContentDir() / TEXT( "Slate/Fonts/Roboto-Regular.ttf" ), 9 ) )
 					.AutoWrapText( true )
@@ -104,25 +105,50 @@ void SCrashReportClient::Construct(const FArguments& InArgs, TSharedRef<FCrashRe
 				[
 					SNew(STextBlock)
 					.Margin( FMargin(4,2,0,0) )
-					.Font( FSlateFontInfo( FPaths::EngineContentDir() / TEXT( "Slate/Fonts/Roboto-Regular.ttf" ), 9 ) )
-					.ColorAndOpacity( FSlateColor( FLinearColor::White * 0.8f ) )
+					.Font( FSlateFontInfo( FPaths::EngineContentDir() / TEXT( "Slate/Testing/Fonts/Roboto-Italic.ttf" ), 9 ) )
+					.ColorAndOpacity( FSlateColor( FLinearColor::White * 0.5f ) )
 					.Text( LOCTEXT( "CrashProvide", "Please provide detailed information about what you were doing when the crash occurred." ) )
 					.Visibility( this, &SCrashReportClient::IsHintTextVisible )
 				]	
 			]
-			
+
+			+ SVerticalBox::Slot()
+			.AutoHeight()
+			.Padding( FMargin( 4, 10, 4, 0 ) )
+			[
+				SNew( SOverlay )
+
+				+ SOverlay::Slot()			
+				[
+					SNew( SColorBlock )
+					.Color( FLinearColor::Black )
+				]
+
+				+ SOverlay::Slot()
+				[
+					SNew( STextBlock )
+					.Margin( FMargin( 4, 2, 0, 8 ) )
+					.Font( FSlateFontInfo( FPaths::EngineContentDir() / TEXT( "Slate/Testing/Fonts/Roboto-Italic.ttf" ), 9 ) )
+					.Text( LOCTEXT( "CrashReportData", "Crash report data:" ) )
+					.ColorAndOpacity( FSlateColor( FLinearColor::White * 0.5f ) )
+					.AutoWrapText( false )
+				]
+			]
+
 			+SVerticalBox::Slot()
 			.FillHeight(0.7f)
-			.Padding(3)
+			.Padding( FMargin( 4, 0 ) )
 			[
 				SNew(SOverlay)
 		
 				+SOverlay::Slot()
 				[
 					SNew( SMultiLineEditableTextBox )
+					.Style( &FCrashReportClientStyle::Get().GetWidgetStyle<FEditableTextBoxStyle>( "NormalEditableTextBox" ) )
 					.Font( FSlateFontInfo( FPaths::EngineContentDir() / TEXT( "Slate/Fonts/Roboto-Regular.ttf" ), 8 ) )
 					.AutoWrapText( false )
 					.IsReadOnly( true )
+					.ReadOnlyForegroundColor( FSlateColor( FLinearColor::Black ) )
 					.BackgroundColor( FSlateColor( FLinearColor::Black ) )
 					.ForegroundColor( FSlateColor( FLinearColor::White * 0.8f ) )
 					.Text( Client, &FCrashReportClient::GetDiagnosticText )
@@ -141,7 +167,7 @@ void SCrashReportClient::Construct(const FArguments& InArgs, TSharedRef<FCrashRe
 
 			+SVerticalBox::Slot()
 			.AutoHeight()
-			.Padding( FMargin( 3, 6 ) )
+			.Padding( FMargin( 4, 12 ) )
 			[
 				SNew(SHorizontalBox)
 
@@ -167,7 +193,7 @@ void SCrashReportClient::Construct(const FArguments& InArgs, TSharedRef<FCrashRe
 			// Stuff anchored to the bottom
 			+SVerticalBox::Slot()
 			.AutoHeight()
-			.Padding( FMargin(3, 3+16, 3, 3) )
+			.Padding( FMargin(4, 4+16, 4, 4) )
 			[
 				SNew(SHorizontalBox)
 
