CrashReportWebsite - Improved performance of generating data for the dashboard stats, added anonymous crashes to graphs

diff --git a/Engine/Source/Programs/CrashReporter/CrashReportWebSite/Controllers/DashboardController.cs b/Engine/Source/Programs/CrashReporter/CrashReportWebSite/Controllers/DashboardController.cs
--- a/Engine/Source/Programs/CrashReporter/CrashReportWebSite/Controllers/DashboardController.cs
+++ b/Engine/Source/Programs/CrashReporter/CrashReportWebSite/Controllers/DashboardController.cs
@@ -10,169 +10,259 @@ using Tools.CrashReporter.CrashReportWebSite.Models;
 
 namespace Tools.CrashReporter.CrashReportWebSite.Controllers
 {
+	/// <summary>
+	/// 
+	/// </summary>
+	public struct FCrashMinimal
+	{
+		/// <summary>
+		/// 
+		/// </summary>
+		public readonly DateTime TimeOfCrash;
+
+		/// <summary>
+		///	
+		/// </summary>
+		public readonly int UserId;
+
+
+		/// <summary>
+		/// Constructor
+		/// </summary>
+		/// <param name="Crash">Crash</param>
+		public FCrashMinimal( Crash Crash )
+		{
+			TimeOfCrash = Crash.TimeOfCrash.Value;
+			UserId = Crash.UserNameId.Value;
+		}
+
+		public FCrashMinimal( DateTime InTimeOfCrash, int InUserId )
+		{
+			TimeOfCrash = InTimeOfCrash;
+			UserId = InUserId;
+		}
+	}
+
 	/// <summary>
 	/// The controller to handle graphing of crashes per user group over time.
 	/// </summary>
 	public class DashboardController : Controller
 	{
+		/// <summary>Fake id for all user groups</summary>
+		public static readonly int AllUserGroupId = -1;
+
+		CrashRepository CrashRepository = new CrashRepository();
+
+		/// <summary></summary>
+		public int GetIdFromUserGroup( string UserGroup )
+		{
+			var Group = CrashRepository.Context.UserGroups.Where( X => X.Name.Contains( UserGroup ) ).FirstOrDefault();
+			return Group.Id;
+		}
+
+		/// <summary></summary>
+		public HashSet<int> GetUserIdsFromUserGroup( string UserGroup )
+		{
+			int UserGroupId = GetIdFromUserGroup( UserGroup );
+			return GetUserIdsFromUserGroupId( UserGroupId );
+		}
+
+		/// <summary></summary>
+		public HashSet<int> GetUserIdsFromUserGroupId( int UserGroupId )
+		{
+			var UserIds = CrashRepository.Context.Users.Where( X => X.UserGroupId == UserGroupId ).Select( X => X.Id );
+			return new HashSet<int>( UserIds );
+		}
+
 		/// <summary>
-		/// The main view of the dash board.
+		/// Return a dictionary of crashes per group grouped by week.
 		/// </summary>
-		/// <returns>A view showing two charts of crashes over time.</returns>
-		public ActionResult Index()
+		/// <param name="Crashes">A set of crashes to interrogate.</param>
+		/// <param name="UserGroupId">The id of the user group to interrogate.</param>
+		/// <returns>A dictionary of week vs. crash count.</returns>
+		public Dictionary<DateTime, int> GetWeeklyCountsByGroup( List<FCrashMinimal> Crashes, int UserGroupId )
 		{
-			CrashRepository LocalCrashRepository = new CrashRepository();
-			IQueryable<Crash> Crashes = LocalCrashRepository.ListAll();
-
-			int UndefinedUserGroupId = LocalCrashRepository.FindOrAddUserGroup( "Undefined" );
-			int GeneralUserGroupId = LocalCrashRepository.FindOrAddUserGroup( "General" );
-			int CoderUserGroupId = LocalCrashRepository.FindOrAddUserGroup( "Coder" );
-			int EngineQAUserGroupId = LocalCrashRepository.FindOrAddUserGroup( "EngineQA" );
-			int GameQAUserGroupId = LocalCrashRepository.FindOrAddUserGroup( "GameQA" );
-			int AutomatedUserGroupId = LocalCrashRepository.FindOrAddUserGroup( "Automated" );
-
-			Dictionary<DateTime, int> GeneralResults = LocalCrashRepository.GetWeeklyCountsByGroup( Crashes, GeneralUserGroupId, UndefinedUserGroupId );
-			Dictionary<DateTime, int> CoderResults = LocalCrashRepository.GetWeeklyCountsByGroup( Crashes, CoderUserGroupId, UndefinedUserGroupId );
-			Dictionary<DateTime, int> EngineQAResults = LocalCrashRepository.GetWeeklyCountsByGroup( Crashes, EngineQAUserGroupId, UndefinedUserGroupId );
-			Dictionary<DateTime, int> GameQAResults = LocalCrashRepository.GetWeeklyCountsByGroup( Crashes, GameQAUserGroupId, UndefinedUserGroupId );
-			Dictionary<DateTime, int> AutomatedResults = LocalCrashRepository.GetWeeklyCountsByGroup( Crashes, AutomatedUserGroupId, UndefinedUserGroupId );
-			Dictionary<DateTime, int> AllResults = LocalCrashRepository.GetWeeklyCountsByGroup( Crashes, -1, UndefinedUserGroupId );
-
-			Dictionary<DateTime, int> DailyGeneralResults = LocalCrashRepository.GetDailyCountsByGroup( Crashes, GeneralUserGroupId, UndefinedUserGroupId );
-			Dictionary<DateTime, int> DailyCoderResults = LocalCrashRepository.GetDailyCountsByGroup( Crashes, CoderUserGroupId, UndefinedUserGroupId );
-			Dictionary<DateTime, int> DailyEngineQAResults = LocalCrashRepository.GetDailyCountsByGroup( Crashes, EngineQAUserGroupId, UndefinedUserGroupId );
-			Dictionary<DateTime, int> DailyGameQAResults = LocalCrashRepository.GetDailyCountsByGroup( Crashes, GameQAUserGroupId, UndefinedUserGroupId );
-			Dictionary<DateTime, int> DailyAutomatedResults = LocalCrashRepository.GetDailyCountsByGroup( Crashes, AutomatedUserGroupId, UndefinedUserGroupId );
-			Dictionary<DateTime, int> DailyAllResults = LocalCrashRepository.GetDailyCountsByGroup( Crashes, -1, UndefinedUserGroupId );
-
-			string CrashesByWeek = "";
-
-			foreach( KeyValuePair<DateTime, int> Result in AllResults )
+			using( FAutoScopedLogTimer LogTimer = new FAutoScopedLogTimer( this.GetType().ToString() + "(UserGroupId=" + UserGroupId + ")" ) )
 			{
-				int GeneralCrashes = 0;
-				try
-				{
-					GeneralCrashes = GeneralResults[Result.Key];
-				}
-				catch
-				{
-				}
+				Dictionary<DateTime, int> Results = new Dictionary<DateTime, int>();
 
-				int CoderCrashes = 0;
-				try
-				{
-					CoderCrashes = CoderResults[Result.Key];
-				}
-				catch
+				var UsersIds = GetUserIdsFromUserGroupId( UserGroupId );
+
+				// Trim crashes to user group.
+				if( UserGroupId != DashboardController.AllUserGroupId )
 				{
+					Crashes = Crashes.Where( Crash => UsersIds.Contains( Crash.UserId ) ).ToList();
 				}
 
-				int EngineQACrashes = 0;
 				try
 				{
-					EngineQACrashes = EngineQAResults[Result.Key];
+					Results =
+					(
+						from CrashDetail in Crashes
+						group CrashDetail by CrashDetail.TimeOfCrash.AddDays( -(int)CrashDetail.TimeOfCrash.DayOfWeek ).Date into GroupCount
+						orderby GroupCount.Key
+						select new { Count = GroupCount.Count(), Date = GroupCount.Key }
+					).ToDictionary( x => x.Date, y => y.Count );
 				}
-				catch
+				catch( Exception Ex )
 				{
+					Debug.WriteLine( "Exception in GetWeeklyCountsByGroup: " + Ex.ToString() );
 				}
 
-				int GameQACrashes = 0;
-				try
-				{
-					GameQACrashes = GameQAResults[Result.Key];
-				}
-				catch
+				return Results;
+			}
+		}
+
+		/// <summary>
+		/// Return a dictionary of crashes per group grouped by day.
+		/// </summary>
+		/// <param name="Crashes">A set of crashes to interrogate.</param>
+		/// <param name="UserGroupId">The id of the user group to interrogate.</param>
+		/// <returns>A dictionary of day vs. crash count.</returns>
+		public Dictionary<DateTime, int> GetDailyCountsByGroup( List<FCrashMinimal> Crashes, int UserGroupId )
+		{
+			using( FAutoScopedLogTimer LogTimer = new FAutoScopedLogTimer( this.GetType().ToString() + "(UserGroupId=" + UserGroupId + ")" ) )
+			{
+				Dictionary<DateTime, int> Results = new Dictionary<DateTime, int>();
+
+				var UsersIds = GetUserIdsFromUserGroupId( UserGroupId );
+
+				// Trim crashes to user group.
+				if( UserGroupId != DashboardController.AllUserGroupId )
 				{
+					Crashes = Crashes.Where( Crash => UsersIds.Contains( Crash.UserId ) ).ToList();
 				}
 
-				int AutomatedCrashes = 0;
 				try
 				{
-					AutomatedCrashes = AutomatedResults[Result.Key];
+					Results =
+					(
+						from CrashDetail in Crashes
+						group CrashDetail by CrashDetail.TimeOfCrash.Date into GroupCount
+						orderby GroupCount.Key
+						select new { Count = GroupCount.Count(), Date = GroupCount.Key }
+					).ToDictionary( x => x.Date, y => y.Count );
 				}
-				catch
+				catch( Exception Ex )
 				{
+					Debug.WriteLine( "Exception in GetDailyCountsByGroup: " + Ex.ToString() );
 				}
 
-				int Year = Result.Key.Year;
-				int Month = Result.Key.AddMonths( -1 ).Month;
-				if( Result.Key.Month == 13 || Result.Key.Month == 1 )
+				return Results;
+			}
+		}
+
+		/// <summary>
+		/// The main view of the dash board.
+		/// </summary>
+		/// <returns>A view showing two charts of crashes over time.</returns>
+		public ActionResult Index()
+		{
+			using( FAutoScopedLogTimer LogTimer = new FAutoScopedLogTimer( this.GetType().ToString() ) )
+			{
+				DateTime Today = DateTime.UtcNow;
+				DateTime AfewMonthsAgo = Today.AddMonths( -6 );
+
+				FAutoScopedLogTimer LogTimerSQL = new FAutoScopedLogTimer( "CrashesFilterByDate", "", "" );
+				IQueryable<Crash> Crashes = CrashRepository.FilterByDate( CrashRepository.ListAll(), AfewMonthsAgo, Today );
+				var VMinimalCrashes = Crashes.Select( Crash => new { TimeOfCrash = Crash.TimeOfCrash.Value, UserID = Crash.UserNameId.Value } ).ToList();
+
+				List<FCrashMinimal> MinimalCrashes = new List<FCrashMinimal>( VMinimalCrashes.Count );
+				foreach( var Anotype in VMinimalCrashes )
 				{
-					Month = 0;
+					MinimalCrashes.Add( new FCrashMinimal( Anotype.TimeOfCrash, Anotype.UserID ) );
 				}
+				LogTimerSQL.Dispose();
 
-				int Day = Result.Key.Day + 6;
+				int GeneralUserGroupId = CrashRepository.FindOrAddUserGroup( "General" );
+				int CoderUserGroupId = CrashRepository.FindOrAddUserGroup( "Coder" );
+				int EngineQAUserGroupId = CrashRepository.FindOrAddUserGroup( "EngineQA" );
+				int GameQAUserGroupId = CrashRepository.FindOrAddUserGroup( "GameQA" );
+				int AnonymousUserGroupId = CrashRepository.FindOrAddUserGroup( "Anonymous" );
 
-				string Line = "[new Date(" + Year + ", " + Month + ", " + Day + "), " + GeneralCrashes + ", " + CoderCrashes + ", " + EngineQACrashes + ", " + GameQACrashes + ", " + AutomatedCrashes + ", " + Result.Value + "], ";
-				CrashesByWeek = CrashesByWeek + Line;
-			}
+				Dictionary<DateTime, int> GeneralResults = GetWeeklyCountsByGroup( MinimalCrashes, GeneralUserGroupId );
+				Dictionary<DateTime, int> CoderResults = GetWeeklyCountsByGroup( MinimalCrashes, CoderUserGroupId );
+				Dictionary<DateTime, int> EngineQAResults = GetWeeklyCountsByGroup( MinimalCrashes, EngineQAUserGroupId );
+				Dictionary<DateTime, int> GameQAResults = GetWeeklyCountsByGroup( MinimalCrashes, GameQAUserGroupId );
+				Dictionary<DateTime, int> AnonymousResults = GetWeeklyCountsByGroup( MinimalCrashes, AnonymousUserGroupId );
+				Dictionary<DateTime, int> AllResults = GetWeeklyCountsByGroup( MinimalCrashes, AllUserGroupId );
 
-			CrashesByWeek = CrashesByWeek.TrimEnd( ", ".ToCharArray() );
+				Dictionary<DateTime, int> DailyGeneralResults = GetDailyCountsByGroup( MinimalCrashes, GeneralUserGroupId );
+				Dictionary<DateTime, int> DailyCoderResults = GetDailyCountsByGroup( MinimalCrashes, CoderUserGroupId );
+				Dictionary<DateTime, int> DailyEngineQAResults = GetDailyCountsByGroup( MinimalCrashes, EngineQAUserGroupId );
+				Dictionary<DateTime, int> DailyGameQAResults = GetDailyCountsByGroup( MinimalCrashes, GameQAUserGroupId );
+				Dictionary<DateTime, int> DailyAnonymousResults = GetDailyCountsByGroup( MinimalCrashes, AnonymousUserGroupId );
+				Dictionary<DateTime, int> DailyAllResults = GetDailyCountsByGroup( MinimalCrashes, AllUserGroupId );
 
-			string CrashesByDay = "";
-			foreach( KeyValuePair<DateTime, int> DailyResult in DailyAllResults )
-			{
-				int GeneralCrashes = 0;
-				try
-				{
-					GeneralCrashes = DailyGeneralResults[DailyResult.Key];
-				}
-				catch
-				{
-				}
+				string CrashesByWeek = "";
 
-				int CoderCrashes = 0;
-				try
-				{
-					CoderCrashes = DailyCoderResults[DailyResult.Key];
-				}
-				catch
+				foreach( KeyValuePair<DateTime, int> Result in AllResults )
 				{
-				}
+					int GeneralCrashes = 0;
+					GeneralResults.TryGetValue( Result.Key, out GeneralCrashes );
 
-				int EngineQACrashes = 0;
-				try
-				{
-					EngineQACrashes = DailyEngineQAResults[DailyResult.Key];
-				}
-				catch
-				{
-				}
+					int CoderCrashes = 0;
+					CoderResults.TryGetValue( Result.Key, out CoderCrashes );
 
-				int GameQACrashes = 0;
-				try
-				{
-					GameQACrashes = DailyGameQAResults[DailyResult.Key];
-				}
-				catch
-				{
-				}
+					int EngineQACrashes = 0;
+					EngineQAResults.TryGetValue( Result.Key, out EngineQACrashes );
 
-				int AutomatedCrashes = 0;
-				try
-				{
-					AutomatedCrashes = DailyAutomatedResults[DailyResult.Key];
-				}
-				catch
-				{
+					int GameQACrashes = 0;
+					GameQAResults.TryGetValue( Result.Key, out GameQACrashes );
+
+					int AnonymousCrashes = 0;
+					AnonymousResults.TryGetValue( Result.Key, out AnonymousCrashes );
+
+					int Year = Result.Key.Year;
+					int Month = Result.Key.AddMonths( -1 ).Month;
+					if( Result.Key.Month == 13 || Result.Key.Month == 1 )
+					{
+						Month = 0;
+					}
+
+					int Day = Result.Key.Day + 6;
+
+					string Line = "[new Date(" + Year + ", " + Month + ", " + Day + "), " + GeneralCrashes + ", " + CoderCrashes + ", " + EngineQACrashes + ", " + GameQACrashes + ", " + AnonymousCrashes + ", " + Result.Value + "], ";
+					CrashesByWeek += Line;
 				}
 
-				int Year = DailyResult.Key.Year;
-				int Month = DailyResult.Key.AddMonths( -1 ).Month;
-				if( DailyResult.Key.Month == 13 || DailyResult.Key.Month == 1 )
+				CrashesByWeek = CrashesByWeek.TrimEnd( ", ".ToCharArray() );
+
+				string CrashesByDay = "";
+				foreach( KeyValuePair<DateTime, int> DailyResult in DailyAllResults )
 				{
-					Month = 0;
-				}
+					int GeneralCrashes = 0;
+					DailyGeneralResults.TryGetValue( DailyResult.Key, out GeneralCrashes );
 
-				int Day = DailyResult.Key.Day;
+					int CoderCrashes = 0;
+					DailyCoderResults.TryGetValue( DailyResult.Key, out CoderCrashes );
 
-				string Line = "[new Date(" + Year + ", " + Month + ", " + Day + "), " + GeneralCrashes + ", " + CoderCrashes + ", " + EngineQACrashes + ", " + GameQACrashes + ", " + AutomatedCrashes + ", " + DailyResult.Value + "], ";
-				CrashesByDay = CrashesByDay + Line;
-			}
+					int EngineQACrashes = 0;
+					DailyEngineQAResults.TryGetValue( DailyResult.Key, out EngineQACrashes );
+
+					int GameQACrashes = 0;
+					DailyGameQAResults.TryGetValue( DailyResult.Key, out GameQACrashes );
 
-			CrashesByDay = CrashesByDay.TrimEnd( ", ".ToCharArray() );
+					int AnonymousCrashes = 0;
+					DailyAnonymousResults.TryGetValue( DailyResult.Key, out AnonymousCrashes );
 
-			return View( "Index", new DashboardViewModel { CrashesByWeek = CrashesByWeek, CrashesByDay = CrashesByDay } );
+					int Year = DailyResult.Key.Year;
+					int Month = DailyResult.Key.AddMonths( -1 ).Month;
+					if( DailyResult.Key.Month == 13 || DailyResult.Key.Month == 1 )
+					{
+						Month = 0;
+					}
+
+					int Day = DailyResult.Key.Day;
+
+					string Line = "[new Date(" + Year + ", " + Month + ", " + Day + "), " + GeneralCrashes + ", " + CoderCrashes + ", " + EngineQACrashes + ", " + GameQACrashes + ", " + AnonymousCrashes + ", " + DailyResult.Value + "], ";
+					CrashesByDay += Line;
+				}
+
+				CrashesByDay = CrashesByDay.TrimEnd( ", ".ToCharArray() );
+
+				return View( "Index", new DashboardViewModel { CrashesByWeek = CrashesByWeek, CrashesByDay = CrashesByDay } );
+			}
 		}
 	}
 }
\ No newline at end of file
diff --git a/Engine/Source/Programs/CrashReporter/CrashReportWebSite/Views/Dashboard/Index.aspx b/Engine/Source/Programs/CrashReporter/CrashReportWebSite/Views/Dashboard/Index.aspx
--- a/Engine/Source/Programs/CrashReporter/CrashReportWebSite/Views/Dashboard/Index.aspx
+++ b/Engine/Source/Programs/CrashReporter/CrashReportWebSite/Views/Dashboard/Index.aspx
@@ -14,11 +14,11 @@
 <asp:Content ID="ScriptContent" ContentPlaceHolderID="ScriptContent" runat="server" >
 	<script type="text/javascript">
 
-	    $(document).ready(function ()
-	    {
-	    });
-             
-    </script>
+		$(document).ready(function ()
+		{
+		});
+			 
+	</script>
 </asp:Content>
 
 <asp:Content ID="AboveMainContent" ContentPlaceHolderID="AboveMainContent" runat="server" >
@@ -26,54 +26,54 @@
 
 <asp:Content ID="MainContent" ContentPlaceHolderID="MainContent" runat="server">
 	<div id='CrashesTableContainer'>
-    <div id="DashBoard">
+	<div id="DashBoard">
 
-        <script type='text/javascript' src='http://www.google.com/jsapi'></script>
-        <script type='text/javascript'>
-            google.load('visualization', '1', { 'packages': ['annotatedtimeline'] });
-            google.setOnLoadCallback(drawChart);
-            function drawChart()
-            {
-                var data = new google.visualization.DataTable();
-                data.addColumn('date', 'Date');
-                data.addColumn('number', 'General Crashes');
-                data.addColumn('number', 'Coder Crashes');
-                data.addColumn('number', 'EngineQA Crashes');
-                data.addColumn('number', 'GameQA Crashes');
-                data.addColumn('number', 'Automated Crashes');
-                data.addColumn('number', 'All Crashes');
-                data.addRows([<%=Model.CrashesByWeek%>]);
-                var chart = new google.visualization.AnnotatedTimeLine(document.getElementById('weekly_chart'));
-                chart.draw(data, { displayAnnotations: true });
-            }
-        </script>
+		<script type='text/javascript' src='http://www.google.com/jsapi'></script>
+		<script type='text/javascript'>
+			google.load('visualization', '1', { 'packages': ['annotatedtimeline'] });
+			google.setOnLoadCallback(drawChart);
+			function drawChart()
+			{
+				var data = new google.visualization.DataTable();
+				data.addColumn('date', 'Date');
+				data.addColumn('number', 'General Crashes');
+				data.addColumn('number', 'Coder Crashes');
+				data.addColumn('number', 'EngineQA Crashes');
+				data.addColumn('number', 'GameQA Crashes');
+				data.addColumn('number', 'Anonymous Crashes');
+				data.addColumn('number', 'All Crashes');
+				data.addRows([<%=Model.CrashesByWeek%>]);
+				var chart = new google.visualization.AnnotatedTimeLine(document.getElementById('weekly_chart'));
+				chart.draw(data, { displayAnnotations: true });
+			}
+		</script>
 
-        <script type='text/javascript'>
-            google.load('visualization', '1', { 'packages': ['annotatedtimeline'] });
-            google.setOnLoadCallback(drawChart);
-            function drawChart()
-            {
-                var data = new google.visualization.DataTable();
-                data.addColumn('date', 'Date');
-                data.addColumn('number', 'General Crashes');
-                data.addColumn('number', 'Coder Crashes');
-                data.addColumn('number', 'EngineQA Crashes');
-                data.addColumn('number', 'GameQA Crashes');
-                data.addColumn('number', 'Automated Crashes');
-                data.addColumn('number', 'All Crashes');
-                data.addRows([<%=Model.CrashesByDay%>]);
-                var chart = new google.visualization.AnnotatedTimeLine(document.getElementById('daily_chart'));
-                chart.draw(data, { displayAnnotations: true });
-            }
-        </script>
+		<script type='text/javascript'>
+			google.load('visualization', '1', { 'packages': ['annotatedtimeline'] });
+			google.setOnLoadCallback(drawChart);
+			function drawChart()
+			{
+				var data = new google.visualization.DataTable();
+				data.addColumn('date', 'Date');
+				data.addColumn('number', 'General Crashes');
+				data.addColumn('number', 'Coder Crashes');
+				data.addColumn('number', 'EngineQA Crashes');
+				data.addColumn('number', 'GameQA Crashes');
+				data.addColumn('number', 'Anonymous Crashes');
+				data.addColumn('number', 'All Crashes');
+				data.addRows([<%=Model.CrashesByDay%>]);
+				var chart = new google.visualization.AnnotatedTimeLine(document.getElementById('daily_chart'));
+				chart.draw(data, { displayAnnotations: true });
+			}
+		</script>
 
-        <h2>Crashes By Week</h2>
-        <div id='weekly_chart'></div>
-        <br />
-        
-        <h2>Crashes By Day</h2>
-        <div id='daily_chart'></div>
-    </div>
+		<h2>Crashes By Week</h2>
+		<div id='weekly_chart'></div>
+		<br />
+		
+		<h2>Crashes By Day</h2>
+		<div id='daily_chart'></div>
+	</div>
  </div>
 
 </asp:Content>
