CrashReportWebsite -
UECORE-50 Crash Reporter: Anonymous users count as one single user making "Number of affected users" misleading
UECORE-80 CRASH REPORTER: New CrashGroups column for "Number of Affected Users in Time Frame"
UECORE-63 Crash Reporter Anonymous Crash Groups cannot be ordered by [REDACTED_PERSON]
UECORE-29 CRASH REPORTER: Cosmetic change to replace "TTPID" with "JIRA"
Optimized BuggView, minimized data transfered from the database, now it's quite usable, entitysets are ridiculously slow

diff --git a/Engine/Source/Programs/CrashReporter/CrashReportWebSite/Controllers/BuggsController.cs b/Engine/Source/Programs/CrashReporter/CrashReportWebSite/Controllers/BuggsController.cs
--- a/Engine/Source/Programs/CrashReporter/CrashReportWebSite/Controllers/BuggsController.cs
+++ b/Engine/Source/Programs/CrashReporter/CrashReportWebSite/Controllers/BuggsController.cs
@@ -36,7 +36,7 @@ namespace Tools.CrashReporter.CrashReportWebSite.Controllers
 		{
 			using( FAutoScopedLogTimer LogTimer = new FAutoScopedLogTimer( this.GetType().ToString() ) )
 			{
-				FormHelper FormData = new FormHelper( Request, BuggsForm, "CrashesInTimeFrame" );
+				FormHelper FormData = new FormHelper( Request, BuggsForm, "CrashesInTimeFrameGroup" );
 				BuggsViewModel Results = LocalBuggRepository.GetResults( FormData );
 				return View( "Index", Results );
 			}
@@ -95,7 +95,17 @@ namespace Tools.CrashReporter.CrashReportWebSite.Controllers
 					return RedirectToAction( "" );
 				}
 
-				Crashes = Bugg.GetCrashes().ToList();
+				using( FAutoScopedLogTimer GetCrashesTimer = new FAutoScopedLogTimer( "Bugg.GetCrashes().ToList" ) )
+				{
+					Crashes = Bugg.GetCrashes();
+
+					HashSet<string> MachineIds = new HashSet<string>();
+					foreach( Crash Crash in Crashes )
+					{
+						MachineIds.Add( Crash.ComputerName );
+					}
+					Bugg.NumberOfUniqueMachines = MachineIds.Count;
+				}
 
 				// Apply any user settings
 				if( BuggsForm.Count > 0 )
@@ -133,30 +143,39 @@ namespace Tools.CrashReporter.CrashReportWebSite.Controllers
 				Crash NewCrash = Model.Crashes.FirstOrDefault();
 				if( NewCrash != null )
 				{
-					CallStackContainer CallStack = new CallStackContainer( NewCrash );
+					using( FScopedLogTimer LogTimer2 = new FScopedLogTimer( "CallstackTrimming" ) )
+					{
+						CallStackContainer CallStack = new CallStackContainer( NewCrash );
+
+						// Set callstack properties
+						CallStack.bDisplayModuleNames = DisplayModuleNames;
+						CallStack.bDisplayFunctionNames = DisplayFunctionNames;
+						CallStack.bDisplayFileNames = DisplayFileNames;
+						CallStack.bDisplayFilePathNames = DisplayFilePathNames;
+						CallStack.bDisplayUnformattedCallStack = DisplayUnformattedCallStack;
 
-					// Set callstack properties
-					CallStack.bDisplayModuleNames = DisplayModuleNames;
-					CallStack.bDisplayFunctionNames = DisplayFunctionNames;
-					CallStack.bDisplayFileNames = DisplayFileNames;
-					CallStack.bDisplayFilePathNames = DisplayFilePathNames;
-					CallStack.bDisplayUnformattedCallStack = DisplayUnformattedCallStack;
+						Model.CallStack = CallStack;
 
-					Model.CallStack = CallStack;
+						// Shorten very long function names.
+						foreach( CallStackEntry Entry in Model.CallStack.CallStackEntries )
+						{
+							Entry.FunctionName = Entry.GetTrimmedFunctionName( 128 );
+						}
 
-					NewCrash.CallStackContainer = NewCrash.GetCallStack();
+						Model.SourceContext = NewCrash.SourceContext;
+					}
 				}
 
-				using( FScopedLogTimer LogTimer2 = new FScopedLogTimer( "BuggsController.Show.PopulateUserInfo" + "(id=" + id + ")" ) )
+				/*using( FScopedLogTimer LogTimer2 = new FScopedLogTimer( "BuggsController.Show.PopulateUserInfo" + "(id=" + id + ")" ) )
 				{
 					// Add in the users for each crash in the Bugg
 					foreach( Crash CrashInstance in Model.Crashes )
 					{
 						LocalCrashRepository.PopulateUserInfo( CrashInstance );
 					}
-				}
+				}*/
 				return View( "Show", Model );
-			}	
+			}
 		}
 	}
 }
diff --git a/Engine/Source/Programs/CrashReporter/CrashReportWebSite/Models/BuggViewModel.cs b/Engine/Source/Programs/CrashReporter/CrashReportWebSite/Models/BuggViewModel.cs
--- a/Engine/Source/Programs/CrashReporter/CrashReportWebSite/Models/BuggViewModel.cs
+++ b/Engine/Source/Programs/CrashReporter/CrashReportWebSite/Models/BuggViewModel.cs
@@ -12,9 +12,14 @@ namespace Tools.CrashReporter.CrashReportWebSite.Models
 	{
 		/// <summary>The current Bugg to display details of.</summary>
 		public Bugg Bugg { get; set; }
+
 		/// <summary>A container for all the crashes associated with this Bugg.</summary>
 		public IEnumerable<Crash> Crashes { get; set; }
+
 		/// <summary>The callstack common to all crashes in this Bugg.</summary>
 		public CallStackContainer CallStack { get; set; }
+
+		/// <summary></summary>
+		public string SourceContext { get; set; }
 	}
 }
\ No newline at end of file
diff --git a/Engine/Source/Programs/CrashReporter/CrashReportWebSite/Models/BuggsViewModel.cs b/Engine/Source/Programs/CrashReporter/CrashReportWebSite/Models/BuggsViewModel.cs
--- a/Engine/Source/Programs/CrashReporter/CrashReportWebSite/Models/BuggsViewModel.cs
+++ b/Engine/Source/Programs/CrashReporter/CrashReportWebSite/Models/BuggsViewModel.cs
@@ -47,7 +47,7 @@ namespace Tools.CrashReporter.CrashReportWebSite.Models
 		public string BuildVersion { get; set; }
 
 		/// <summary>A dictionary of the number of Buggs per user group.</summary>
-		public Dictionary<string, int> GroupCounts { get; set; }
+		public SortedDictionary<string, int> GroupCounts { get; set; }
 
 		/// <summary>The set of statuses a Bugg could have its status set to.</summary>
 		public IEnumerable<string> SetStatus
diff --git a/Engine/Source/Programs/CrashReporter/CrashReportWebSite/Models/CrashReport.cs b/Engine/Source/Programs/CrashReporter/CrashReportWebSite/Models/CrashReport.cs
--- a/Engine/Source/Programs/CrashReporter/CrashReportWebSite/Models/CrashReport.cs
+++ b/Engine/Source/Programs/CrashReporter/CrashReportWebSite/Models/CrashReport.cs
@@ -49,15 +49,30 @@ namespace Tools.CrashReporter.CrashReportWebSite.Models
 		}
 	}
 
+	/// <summary></summary>
+	public partial class Buggs_Crash
+	{
+		/// <summary></summary>
+		public override string ToString()
+		{
+			return string.Format( "{0} <-> {1}", BuggId, CrashId );
+		}
+	}
+
 	/// <summary>
 	/// Derived information to the default Bugg information from the database.
 	/// </summary>
 	public partial class Bugg
 	{
 		/// <summary></summary>
-		public int CrashesInTimeFrame { get; set; }
+		public int CrashesInTimeFrameGroup { get; set; }
+
+		/// <summary></summary>
+		public int CrashesInTimeFrameAll { get; set; }
+
 		/// <summary></summary>
-		public string SourceContext { get; set; }
+		public int NumberOfUniqueMachines { get; set; }
+
 
 		/// <summary> Helper method, display this Bugg as a human readable string. Debugging purpose. </summary>
 		public override string ToString()
@@ -69,38 +84,24 @@ namespace Tools.CrashReporter.CrashReportWebSite.Models
 		/// 
 		/// </summary>
 		/// <returns></returns>
-		public EntitySet<Crash> GetCrashes()
+		public List<Crash> GetCrashes()
 		{
-			using( FAutoScopedLogTimer LogTimer = new FAutoScopedLogTimer( this.GetType().ToString() ) )
-			{
-				int CrashCount = Crashes.Count;
-				if( NumberOfCrashes != CrashCount )
-				{
-					NumberOfCrashes = CrashCount;
-
-					if( NumberOfCrashes > 0 )
-					{
-						BuggRepository LocalBuggRepository = new BuggRepository();
-						LocalBuggRepository.UpdateBuggData( this, Crashes );
-					}
-				}
-
-				// Just fill the CallStackContainers
-				foreach( Crash CurrentCrash in Crashes )
-				{
-					if( CurrentCrash.CallStackContainer == null )
-					{
-						CurrentCrash.CallStackContainer = CurrentCrash.GetCallStack();
-					}
-
-					if( SourceContext == null )
-					{
-						SourceContext = CurrentCrash.SourceContext;
-					}
-				}
-
-				return Crashes;
-			}
+			CrashRepository CrashRepo = new CrashRepository();
+			var CrashList =
+							(
+								from BuggCrash in CrashRepo.Context.Buggs_Crashes
+								where BuggCrash.BuggId == Id
+								select BuggCrash.Crash/*new
+								{
+									Id = BuggCrash.Crash.Id,
+									BuildVersion = BuggCrash.Crash.BuildVersion,
+									TimeofCrash = BuggCrash.Crash.TimeOfCrash.Value,
+									CL = BuggCrash.Crash.ChangeListVersion,
+									UserName = BuggCrash.Crash.UserName,
+									Desc = BuggCrash.Crash.Description
+								}*/
+							).AsEnumerable().ToList();
+			return CrashList;
 		}
 
 		/// <summary></summary>
@@ -134,44 +135,6 @@ namespace Tools.CrashReporter.CrashReportWebSite.Models
 				return Results;
 			}
 		}
-
-		/// <summary>
-		/// http://www.codeproject.com/KB/linq/linq-to-sql-many-to-many.aspx
-		/// </summary>
-		private EntitySet<Crash> CrashesCache;
-
-		/// <summary>
-		/// 
-		/// </summary>
-		public EntitySet<Crash> Crashes
-		{
-			get
-			{
-				if (CrashesCache == null)
-				{
-					CrashesCache = new EntitySet<Crash>(OnCrashesAdd, OnCrashesRemove);
-					CrashesCache.SetSource(Buggs_Crashes.Select(CrashInstance => CrashInstance.Crash));
-				}
-				return CrashesCache;
-			}
-			set
-			{
-				CrashesCache.Assign( value );
-			}
-		}
-
-		private void OnCrashesAdd( Crash CrashInstance )
-		{
-			Buggs_Crashes.Add( new Buggs_Crash { Bugg = this, Crash = CrashInstance } );
-			SendPropertyChanged( null );
-		}
-
-		private void OnCrashesRemove( Crash CrashInstance )
-		{
-			Buggs_Crash BuggCrash = Buggs_Crashes.FirstOrDefault( BuggCrashInstance => BuggCrashInstance.BuggId == Id && BuggCrashInstance.CrashId == CrashInstance.Id );
-			Buggs_Crashes.Remove( BuggCrash );
-			SendPropertyChanged( null );
-		}
 	}
 
 	/// <summary>
@@ -187,7 +150,7 @@ namespace Tools.CrashReporter.CrashReportWebSite.Models
 		{
 			get
 			{
-				return UserById ?? UserByName;
+				return UserById;
 			}
 		}
 
diff --git a/Engine/Source/Programs/CrashReporter/CrashReportWebSite/Models/CrashReport.dbml b/Engine/Source/Programs/CrashReporter/CrashReportWebSite/Models/CrashReport.dbml
--- a/Engine/Source/Programs/CrashReporter/CrashReportWebSite/Models/CrashReport.dbml
+++ b/Engine/Source/Programs/CrashReporter/CrashReportWebSite/Models/CrashReport.dbml
@@ -28,9 +28,7 @@
       <Column Name="UserName" Type="System.String" DbType="VarChar(64) NOT NULL" CanBeNull="false" />
       <Column Name="UserGroupId" Type="System.Int32" DbType="Int NOT NULL" CanBeNull="false" />
       <Association Name="User_Buggs_User" Member="Buggs_Users" ThisKey="Id" OtherKey="UserNameId" Type="Buggs_User" />
-      <Association Name="User_Buggs_User1" Member="Buggs_Users1" ThisKey="Id" OtherKey="UserNameId" Type="Buggs_User" />
       <Association Name="User_Crash" Member="CrashesFromUserId" Storage="_Crashes" ThisKey="Id" OtherKey="UserNameId" Type="Crash" />
-      <Association Name="User_Crash1" Member="CrashesFromName" ThisKey="UserName" OtherKey="UserName" Type="Crash" />
       <Association Name="UserGroup_User" Member="UserGroup" ThisKey="UserGroupId" OtherKey="Id" Type="UserGroup" IsForeignKey="true" />
     </Type>
   </Table>
@@ -55,7 +53,6 @@
       <Column Name="UserName" Type="System.String" DbType="VarChar(64) NOT NULL" IsPrimaryKey="true" CanBeNull="false" />
       <Column Name="UserNameId" Type="System.Int32" DbType="Int" CanBeNull="true" />
       <Association Name="User_Buggs_User" Member="User" ThisKey="UserNameId" OtherKey="Id" Type="User" IsForeignKey="true" />
-      <Association Name="User_Buggs_User1" Member="User1" ThisKey="UserNameId" OtherKey="Id" Type="User" IsForeignKey="true" />
     </Type>
   </Table>
   <Table Name="dbo.Crashes" Member="Crashes">
@@ -94,7 +91,6 @@
       <Column Name="EngineVersion" Type="System.String" DbType="varchar(64)" CanBeNull="true" />
       <Association Name="Crash_Buggs_Crash" Member="Buggs_Crashes" ThisKey="Id" OtherKey="CrashId" Type="Buggs_Crash" />
       <Association Name="User_Crash" Member="UserById" Storage="_User" ThisKey="UserNameId" OtherKey="Id" Type="User" IsForeignKey="true" />
-      <Association Name="User_Crash1" Member="UserByName" ThisKey="UserName" OtherKey="UserName" Type="User" IsForeignKey="true" />
     </Type>
   </Table>
   <Table Name="dbo.FunctionCalls" Member="FunctionCalls">
diff --git a/Engine/Source/Programs/CrashReporter/CrashReportWebSite/Models/CrashReport.dbml.layout b/Engine/Source/Programs/CrashReporter/CrashReportWebSite/Models/CrashReport.dbml.layout
--- a/Engine/Source/Programs/CrashReporter/CrashReportWebSite/Models/CrashReport.dbml.layout
+++ b/Engine/Source/Programs/CrashReporter/CrashReportWebSite/Models/CrashReport.dbml.layout
@@ -57,7 +57,7 @@
         <classShapeMoniker Id="b4a7205e-4223-43f1-837f-ac4a1390bf89" />
       </nodes>
     </associationConnector>
-    <associationConnector edgePoints="[(5.375 : 4.84084598496556); (6.125 : 4.84084598496556)]" fixedFrom="Algorithm" fixedTo="Algorithm">
+    <associationConnector edgePoints="[(5.375 : 4.84084598496556); (6.125 : 4.84084598496556)]" fixedFrom="NotFixed" fixedTo="NotFixed">
       <AssociationMoniker Name="/CrashReportDataContext/User/User_Crash" />
       <nodes>
         <classShapeMoniker Id="01dddf4f-058f-43b4-b6c0-b90e7acfe0b9" />
@@ -85,27 +85,13 @@
         <classShapeMoniker Id="01dddf4f-058f-43b4-b6c0-b90e7acfe0b9" />
       </nodes>
     </associationConnector>
-    <associationConnector edgePoints="[(8.125 : 2.47199625651042); (8.875 : 2.47199625651042)]" fixedFrom="Algorithm" fixedTo="Algorithm">
+    <associationConnector edgePoints="[(8.125 : 2.47199625651042); (8.875 : 2.47199625651042)]" fixedFrom="NotFixed" fixedTo="NotFixed">
       <AssociationMoniker Name="/CrashReportDataContext/Crash/Crash_Buggs_Crash" />
       <nodes>
         <classShapeMoniker Id="f1204ef8-6660-48ce-8c1d-131c49e2dd25" />
         <classShapeMoniker Id="66e20693-2188-41ba-b050-7d08c2691d12" />
       </nodes>
     </associationConnector>
-    <associationConnector edgePoints="[(5.375 : 4.29542299248278); (6.125 : 4.29542299248278)]" fixedFrom="Algorithm" fixedTo="Algorithm">
-      <AssociationMoniker Name="/CrashReportDataContext/User/User_Crash1" />
-      <nodes>
-        <classShapeMoniker Id="01dddf4f-058f-43b4-b6c0-b90e7acfe0b9" />
-        <classShapeMoniker Id="f1204ef8-6660-48ce-8c1d-131c49e2dd25" />
-      </nodes>
-    </associationConnector>
-    <associationConnector edgePoints="[(3.8125 : 5.2612939453125); (3.8125 : 5.875)]" fixedFrom="NotFixed" fixedTo="NotFixed">
-      <AssociationMoniker Name="/CrashReportDataContext/User/User_Buggs_User1" />
-      <nodes>
-        <classShapeMoniker Id="01dddf4f-058f-43b4-b6c0-b90e7acfe0b9" />
-        <classShapeMoniker Id="8a60fc4a-a797-417f-87ed-25b22bd5e420" />
-      </nodes>
-    </associationConnector>
     <classShape Id="d3fe456a-12f1-46c2-a283-3d1ccff66eff" absoluteBounds="2.75, 7.75, 2, 1.7708968098958344">
       <DataClassMoniker Name="/CrashReportDataContext/UsersMapping" />
       <nestedChildShapes>
diff --git a/Engine/Source/Programs/CrashReporter/CrashReportWebSite/Models/CrashRepository.cs b/Engine/Source/Programs/CrashReporter/CrashReportWebSite/Models/CrashRepository.cs
--- a/Engine/Source/Programs/CrashReporter/CrashReportWebSite/Models/CrashRepository.cs
+++ b/Engine/Source/Programs/CrashReporter/CrashReportWebSite/Models/CrashRepository.cs
@@ -12,6 +12,7 @@ using System.Web;
 using Naspinski.IQueryableSearch;
 
 using Tools.CrashReporter.CrashReportCommon;
+using Tools.CrashReporter.CrashReportWebSite.Controllers;
 
 namespace Tools.CrashReporter.CrashReportWebSite.Models
 {
@@ -22,20 +23,20 @@ namespace Tools.CrashReporter.CrashReportWebSite.Models
 	/// </summary>
 	public class CrashRepository : IDisposable
 	{
-		private CrashReportDataContext CrashRepositoryDataContext;
+		public CrashReportDataContext Context;
 
 		/// <summary>
 		/// The default constructor.
 		/// </summary>
 		public CrashRepository()
 		{
-			CrashRepositoryDataContext = new CrashReportDataContext();
+			Context = new CrashReportDataContext();
 		}
 
 		/// <summary> Submits enqueue changes to the database. </summary>
 		public void SubmitChanges()
 		{
-			CrashRepositoryDataContext.SubmitChanges();
+			Context.SubmitChanges();
 		}
 
 		/// <summary>
@@ -53,7 +54,7 @@ namespace Tools.CrashReporter.CrashReportWebSite.Models
 		/// <param name="Disposing">true if the Dispose call is from user code, and not system code.</param>
 		protected virtual void Dispose( bool Disposing )
 		{
-			CrashRepositoryDataContext.Dispose();
+			Context.Dispose();
 		}
 
 		/// <summary>
@@ -71,7 +72,7 @@ namespace Tools.CrashReporter.CrashReportWebSite.Models
 					int UserGroupId = FindOrAddUserGroup( UserGroupName );
 					Users =
 					(
-						from UserDetail in CrashRepositoryDataContext.Users
+						from UserDetail in Context.Users
 						where UserDetail.UserGroupId == UserGroupId
 						orderby UserDetail.UserName
 						select UserDetail.UserName
@@ -99,7 +100,7 @@ namespace Tools.CrashReporter.CrashReportWebSite.Models
 				{
 					IQueryable<Crash> Crashes =
 					(
-						from CrashDetail in CrashRepositoryDataContext.Crashes
+						from CrashDetail in Context.Crashes
 						where CrashDetail.Id == Id
 						select CrashDetail
 					);
@@ -125,10 +126,10 @@ namespace Tools.CrashReporter.CrashReportWebSite.Models
 			try
 			{
 				string Query = "UPDATE Crashes SET Status = {0} WHERE Id IN ( SELECT CrashId FROM Buggs_Crashes WHERE BuggId = {1} )";
-				CrashRepositoryDataContext.ExecuteCommand( Query, Status, BuggId );
+				Context.ExecuteCommand( Query, Status, BuggId );
 
 				Query = "UPDATE Buggs SET Status = {0} WHERE id = {1}";
-				CrashRepositoryDataContext.ExecuteCommand( Query, Status, BuggId );
+				Context.ExecuteCommand( Query, Status, BuggId );
 			}
 			catch( Exception Ex )
 			{
@@ -146,10 +147,10 @@ namespace Tools.CrashReporter.CrashReportWebSite.Models
 			try
 			{
 				string Query = "UPDATE Crashes SET FixedChangeList = {0} WHERE Id IN ( SELECT CrashId FROM Buggs_Crashes WHERE BuggId = {1} )";
-				CrashRepositoryDataContext.ExecuteCommand( Query, FixedChangeList, BuggId );
+				Context.ExecuteCommand( Query, FixedChangeList, BuggId );
 
 				Query = "UPDATE Buggs SET FixedChangeList = {0} WHERE id = {1}";
-				CrashRepositoryDataContext.ExecuteCommand( Query, FixedChangeList, BuggId );
+				Context.ExecuteCommand( Query, FixedChangeList, BuggId );
 			}
 			catch( Exception Ex )
 			{
@@ -167,10 +168,10 @@ namespace Tools.CrashReporter.CrashReportWebSite.Models
 			try
 			{
 				string Query = "UPDATE Crashes SET TTPID = {0} WHERE Id IN ( SELECT CrashId FROM Buggs_Crashes WHERE BuggId = {1} )";
-				CrashRepositoryDataContext.ExecuteCommand( Query, TTPID, BuggId );
+				Context.ExecuteCommand( Query, TTPID, BuggId );
 
 				Query = "UPDATE Buggs SET TTPID = {0} WHERE id = {1}";
-				CrashRepositoryDataContext.ExecuteCommand( Query, TTPID, BuggId );
+				Context.ExecuteCommand( Query, TTPID, BuggId );
 			}
 			catch( Exception Ex )
 			{
@@ -189,9 +190,8 @@ namespace Tools.CrashReporter.CrashReportWebSite.Models
 				try
 				{
 					int UserGroupId = CrashInstance.User.UserGroupId;
-					var Result = CrashRepositoryDataContext.UserGroups.Where( i => i.Id == UserGroupId ).First();
+					var Result = Context.UserGroups.Where( i => i.Id == UserGroupId ).First();
 					CrashInstance.UserGroupName = Result.Name;
-
 				}
 				catch( Exception Ex )
 				{
@@ -209,7 +209,7 @@ namespace Tools.CrashReporter.CrashReportWebSite.Models
 			IQueryable<Crash> Crashes = null;
 			try
 			{
-				Crashes = CrashRepositoryDataContext.Crashes.OrderByDescending( CrashDetail => CrashDetail.TimeOfCrash );
+				Crashes = Context.Crashes.OrderByDescending( CrashDetail => CrashDetail.TimeOfCrash );
 			}
 			catch( Exception Ex )
 			{
@@ -232,7 +232,7 @@ namespace Tools.CrashReporter.CrashReportWebSite.Models
 			{
 				Crashes =
 				(
-					from CrashDetail in CrashRepositoryDataContext.Crashes
+					from CrashDetail in Context.Crashes
 					select CrashDetail
 				).OrderByDescending( CrashDetail => CrashDetail.TimeOfCrash )
 				.Skip( Skip )
@@ -323,7 +323,7 @@ namespace Tools.CrashReporter.CrashReportWebSite.Models
 							{
 								/*IQueryable<UsersMapping>*/
 								// Make sure that type of [dbo].[UsersMapping] is the same as [analyticsdb-01.dmz.epicgames.net].[CrashReport].[dbo].[UsersMapping]
-								IEnumerable<UsersMapping> FoundUsers = CrashRepositoryDataContext.ExecuteQuery<UsersMapping>
+								IEnumerable<UsersMapping> FoundUsers = Context.ExecuteQuery<UsersMapping>
 								( @"SELECT * FROM [analyticsdb-01.dmz.epicgames.net].[CrashReport].[dbo].[UsersMapping] WHERE lower(UserName) = {0} OR lower(UserEmail) = {0}", Params[1] );
 
 								foreach( UsersMapping TheUser in FoundUsers )
@@ -432,7 +432,7 @@ namespace Tools.CrashReporter.CrashReportWebSite.Models
 				Results =
 				(
 					from CrashDetail in Results
-					from UserDetail in CrashRepositoryDataContext.Users
+					from UserDetail in Context.Users
 					where UserDetail.UserGroupId == UserGroupId &&
 						( CrashDetail.UserNameId == UserDetail.Id || CrashDetail.UserName == UserDetail.UserName )
 					select CrashDetail
@@ -483,7 +483,7 @@ namespace Tools.CrashReporter.CrashReportWebSite.Models
 		public Dictionary<string, int> GetCountsByGroup()
 		{
 			// @TODO yrx 2014-11-06 Optimize?
-			using( FAutoScopedLogTimer LogTimer = new FAutoScopedLogTimer( this.GetType().ToString() ) )
+			using( FAutoScopedLogTimer LogTimer = new FAutoScopedLogTimer( this.GetType().ToString() + " SQL OPT" ) )
 			{
 				Dictionary<string, int> Results = new Dictionary<string, int>();
 
@@ -491,8 +491,8 @@ namespace Tools.CrashReporter.CrashReportWebSite.Models
 				{
 					var GroupCounts =
 					(
-						from UserDetail in CrashRepositoryDataContext.Users
-						join UserGroupDetail in CrashRepositoryDataContext.UserGroups on UserDetail.UserGroupId equals UserGroupDetail.Id
+						from UserDetail in Context.Users
+						join UserGroupDetail in Context.UserGroups on UserDetail.UserGroupId equals UserGroupDetail.Id
 						group UserDetail by UserGroupDetail.Name into GroupCount
 						select new { Key = GroupCount.Key, Count = GroupCount.Count() }
 					);
@@ -503,7 +503,7 @@ namespace Tools.CrashReporter.CrashReportWebSite.Models
 					}
 
 					// Add in all groups, even though there are no crashes associated
-					IEnumerable<string> UserGroups = ( from UserGroupDetail in CrashRepositoryDataContext.UserGroups select UserGroupDetail.Name );
+					IEnumerable<string> UserGroups = ( from UserGroupDetail in Context.UserGroups select UserGroupDetail.Name );
 					foreach( string UserGroupName in UserGroups )
 					{
 						if( !Results.Keys.Contains( UserGroupName ) )
@@ -534,8 +534,8 @@ namespace Tools.CrashReporter.CrashReportWebSite.Models
 
 				try
 				{
-					var UsersIDsAndGroupIDs = CrashRepositoryDataContext.Users.Select( User => new { UserId = User.Id, UserGroupId = User.UserGroupId } ).ToList();
-					var UserGroupArray = CrashRepositoryDataContext.UserGroups.ToList();
+					var UsersIDsAndGroupIDs = Context.Users.Select( User => new { UserId = User.Id, UserGroupId = User.UserGroupId } ).ToList();
+					var UserGroupArray = Context.UserGroups.ToList();
 					UserGroupArray.Sort( ( UG1, UG2 ) => UG1.Name.CompareTo( UG2.Name ) ); 
 
 					// Initialize all groups to 0.
@@ -568,74 +568,6 @@ namespace Tools.CrashReporter.CrashReportWebSite.Models
 			}
 		}
 
-		/// <summary>
-		/// Return a dictionary of crashes per group grouped by week.
-		/// </summary>
-		/// <param name="Crashes">A set of crashes to interrogate.</param>
-		/// <param name="UserGroupId">The id of the user group to interrogate.</param>
-		/// <param name="UndefinedUserGroupId">Id of the undefined user group.</param>
-		/// <returns>A dictionary of week vs. crash count.</returns>
-		public Dictionary<DateTime, int> GetWeeklyCountsByGroup( IQueryable<Crash> Crashes, int UserGroupId, int UndefinedUserGroupId )
-		{
-			using( FAutoScopedLogTimer LogTimer = new FAutoScopedLogTimer( this.GetType().ToString() + "(UserGroupId=" + UserGroupId + ")" ) )
-			{
-				Dictionary<DateTime, int> Results = new Dictionary<DateTime, int>();
-
-				try
-				{
-					Results =
-					(
-						from CrashDetail in Crashes
-						join UserDetail in CrashRepositoryDataContext.Users on CrashDetail.UserNameId equals UserDetail.Id
-						where ( UserGroupId < 0 && UserDetail.UserGroupId != UndefinedUserGroupId ) || UserDetail.UserGroupId == UserGroupId
-						group CrashDetail by CrashDetail.TimeOfCrash.Value.AddDays( -(int)CrashDetail.TimeOfCrash.Value.DayOfWeek ).Date into GroupCount
-						orderby GroupCount.Key
-						select new { Count = GroupCount.Count(), Date = GroupCount.Key }
-					).ToDictionary( x => x.Date, y => y.Count );
-				}
-				catch( Exception Ex )
-				{
-					Debug.WriteLine( "Exception in GetWeeklyCountsByGroup: " + Ex.ToString() );
-				}
-
-				return Results;
-			}
-		}
-
-		/// <summary>
-		/// Return a dictionary of crashes per group grouped by day.
-		/// </summary>
-		/// <param name="Crashes">A set of crashes to interrogate.</param>
-		/// <param name="UserGroupId">The id of the user group to interrogate.</param>
-		/// <param name="UndefinedUserGroupId">Id of the undefined user group.</param>
-		/// <returns>A dictionary of day vs. crash count.</returns>
-		public Dictionary<DateTime, int> GetDailyCountsByGroup( IQueryable<Crash> Crashes, int UserGroupId, int UndefinedUserGroupId )
-		{
-			using( FAutoScopedLogTimer LogTimer = new FAutoScopedLogTimer( this.GetType().ToString() + "(UserGroupId=" + UserGroupId + ")" ) )
-			{
-				Dictionary<DateTime, int> Results = new Dictionary<DateTime, int>();
-
-				try
-				{
-					Results =
-					(
-						from CrashDetail in Crashes
-						join UserDetail in CrashRepositoryDataContext.Users on CrashDetail.UserNameId equals UserDetail.Id
-						where ( UserGroupId < 0 && UserDetail.UserGroupId != UndefinedUserGroupId ) || UserDetail.UserGroupId == UserGroupId
-						group CrashDetail by CrashDetail.TimeOfCrash.Value.Date into GroupCount
-						orderby GroupCount.Key
-						select new { Count = GroupCount.Count(), Date = GroupCount.Key }
-					).ToDictionary( x => x.Date, y => y.Count );
-				}
-				catch( Exception Ex )
-				{
-					Debug.WriteLine( "Exception in GetDailyCountsByGroup: " + Ex.ToString() );
-				}
-
-				return Results;
-			}
-		}
-
 		/// <summary>
 		/// Filter a set of crashes to a date range.
 		/// </summary>
@@ -686,7 +618,7 @@ namespace Tools.CrashReporter.CrashReportWebSite.Models
 
 					case "UserName":
 						// Note: only works where user is stored by Id
-						Results = OrderBy( Results.Join( CrashRepositoryDataContext.Users, CrashInstance => CrashInstance.UserNameId, UserInstance => UserInstance.Id,
+						Results = OrderBy( Results.Join( Context.Users, CrashInstance => CrashInstance.UserNameId, UserInstance => UserInstance.Id,
 														( CrashInstance, UserInstance ) => new { CrashInstance, UserInstance.UserName } ),
 														Joined => Joined.UserName, bSortByDescending ).Select( Joined => Joined.CrashInstance );
 						break;
@@ -760,18 +692,18 @@ namespace Tools.CrashReporter.CrashReportWebSite.Models
 		{
 			try
 			{
-				IQueryable<int> UserGroups = ( from UserGroupDetail in CrashRepositoryDataContext.UserGroups
+				IQueryable<int> UserGroups = ( from UserGroupDetail in Context.UserGroups
 											   where UserGroupDetail.Name.ToLower() == UserGroupName.ToLower()
 											   select UserGroupDetail.Id );
 
-				IQueryable<int> UserNames = ( from UserDetail in CrashRepositoryDataContext.Users
+				IQueryable<int> UserNames = ( from UserDetail in Context.Users
 										  where UserDetail.UserName.ToLower() == UserName.ToLower()
 										  select UserDetail.Id );
 
 				if( UserGroups.Count() == 1 && UserNames.Count() == 1 )
 				{
 					string Query = "UPDATE Users SET UserGroupId = {0} WHERE Id = {1}";
-					CrashRepositoryDataContext.ExecuteCommand( Query, UserGroups.First(), UserNames.First() );
+					Context.ExecuteCommand( Query, UserGroups.First(), UserNames.First() );
 				}
 			}
 			catch( Exception Ex )
@@ -793,7 +725,7 @@ namespace Tools.CrashReporter.CrashReportWebSite.Models
 				int UserGroupNameId = 0;
 				try
 				{
-					IQueryable<int> UserGroups = ( from UserGroupDetail in CrashRepositoryDataContext.UserGroups
+					IQueryable<int> UserGroups = ( from UserGroupDetail in Context.UserGroups
 												   where UserGroupDetail.Name.ToLower() == UserGroupName.ToLower()
 												   select UserGroupDetail.Id );
 
@@ -802,9 +734,9 @@ namespace Tools.CrashReporter.CrashReportWebSite.Models
 					{
 						UserGroup NewUserGroup = new UserGroup();
 						NewUserGroup.Name = UserGroupName;
-						CrashRepositoryDataContext.UserGroups.InsertOnSubmit( NewUserGroup );
+						Context.UserGroups.InsertOnSubmit( NewUserGroup );
 
-						CrashRepositoryDataContext.SubmitChanges();
+						Context.SubmitChanges();
 						UserGroupNameId = NewUserGroup.Id;
 					}
 					else
@@ -851,20 +783,20 @@ namespace Tools.CrashReporter.CrashReportWebSite.Models
 			// Valid MachineID and UserName, updated crash from non-UE4 release
 			if(!string.IsNullOrEmpty(NewCrashInfo.UserName))
 			{
-				NewCrash.UserNameId = CrashRepositoryDataContext.FindOrAddUser( NewCrashInfo.UserName );
+				NewCrash.UserNameId = Context.FindOrAddUser( NewCrashInfo.UserName );
 			}
 			// Valid MachineID and EpicAccountId, updated crash from UE4 release
 			else if(!string.IsNullOrEmpty( NewCrashInfo.EpicAccountId ))
 			{
 				NewCrash.EpicAccountId = NewCrashInfo.EpicAccountId;
-				NewCrash.UserNameId = CrashRepositoryDataContext.FindOrAddUser( UserNameAnonymous );
+				NewCrash.UserNameId = Context.FindOrAddUser( UserNameAnonymous );
 			}
 			// Crash from an older version.
 			else
 			{
 				// MachineGuid for older crashes is obsolete, so ignore it.
 				//NewCrash.ComputerName = NewCrashInfo.MachineGuid;
-				NewCrash.UserNameId = CrashRepositoryDataContext.FindOrAddUser
+				NewCrash.UserNameId = Context.FindOrAddUser
 				(
 					!string.IsNullOrEmpty( NewCrashInfo.UserName ) ? NewCrashInfo.UserName : UserNameAnonymous
 				);
@@ -949,13 +881,13 @@ namespace Tools.CrashReporter.CrashReportWebSite.Models
 			//NewCrash.HasMetaData = true;
 
 			// Add the crash to the database
-			CrashRepositoryDataContext.Crashes.InsertOnSubmit( NewCrash );
-			CrashRepositoryDataContext.SubmitChanges();
+			Context.Crashes.InsertOnSubmit( NewCrash );
+			Context.SubmitChanges();
 
 			NewID = NewCrash.Id;
 
 			// Build a callstack pattern for crash bucketing
-			CrashRepositoryDataContext.BuildPattern(NewCrash);
+			Context.BuildPattern(NewCrash);
 
 			return NewID;
 		}
diff --git a/Engine/Source/Programs/CrashReporter/CrashReportWebSite/Views/Buggs/Show.aspx b/Engine/Source/Programs/CrashReporter/CrashReportWebSite/Views/Buggs/Show.aspx
--- a/Engine/Source/Programs/CrashReporter/CrashReportWebSite/Views/Buggs/Show.aspx
+++ b/Engine/Source/Programs/CrashReporter/CrashReportWebSite/Views/Buggs/Show.aspx
@@ -80,7 +80,7 @@
 			</select>
 			<input type="submit" name="SetStatusSubmit" value="Set" class="SetButtonGreen" />
 
-			<span id="set-ttp" style="">TTP</span>
+			<span id="set-ttp" style="">JIRA</span>
 			<input name="SetTTP" type="text" id="ttp-text-box" />
 			<input type="submit" name="SetTTPSubmit" value="Set" class="SetButtonGreen" />
 	
@@ -112,11 +112,15 @@
 		<dt>Number of Users</dt>
 			<dd class='even'><%=Html.DisplayFor( m => Model.Bugg.NumberOfUsers )%></dd>
 
+		<dt>Number of Unique Machines</dt>
+			<dd class='even'><%=Html.DisplayFor( m => Model.Bugg.NumberOfUniqueMachines )%></dd>
+
 		<dt>Number of Crashes</dt>
 			<dd ><%=Html.DisplayFor( m => Model.Bugg.NumberOfCrashes )%></dd>
 
-		<dt>TTP</dt>
-			<dd class='even'><%=Html.DisplayFor( m => Model.Bugg.TTPID )%></dd>
+		<dt>JIRA</dt>
+			<%--<dd class='even'><%=Html.DisplayFor( m => Model.Bugg.TTPID )%></dd>--%>
+			<dd class='even'><a href="https://jira.ol.epicgames.net/browse/<%=Model.Bugg.TTPID%>" target="_blank"><%=Model.Bugg.TTPID%></a></dd>
 
 		<dt>Status</dt>
 			<dd ><%=Html.DisplayFor( m => Model.Bugg.Status )%></dd>
@@ -143,7 +147,7 @@
 
 
 		<h3>Call Stack of Most Recent Crash</h3>
-		<div id='RawCallStackContainer' style='display:none' >
+		<div id='RawCallStackContainer' style='display:none;'>
 			<div>
 				<% foreach( CallStackEntry CallStackLine in Model.CallStack.CallStackEntries )
 					{%>
@@ -186,7 +190,7 @@
 		<div id='SourceContextContainer'>
 			<h3>Source Context</h3>
 			<div id='SourceContextDisplay'>
-				<pre><%=Html.DisplayFor( x => Model.Bugg.SourceContext )%></pre>
+				<pre><%=Html.DisplayFor( x => Model.SourceContext )%></pre>
 			</div>
 		</div>
 		<%using( Html.BeginForm( "Show", "Buggs", FormMethod.Post, new { id = "EditCrashDescriptionForm"} ) )
@@ -223,7 +227,7 @@
 							<td class="Id"><%=Html.ActionLink( CrashInstance.Id.ToString(), "Show", new { controller = "crashes", id = CrashInstance.Id }, null )%></td>
 							<td class="BuildVersion"><%=CrashInstance.BuildVersion%></td>
 							<td class="TimeOfCrash">
-								<%=CrashInstance.GetTimeOfCrash()[0]%><br />
+								<%=CrashInstance.GetTimeOfCrash()[0]%> &nbsp;
 								<%=CrashInstance.GetTimeOfCrash()[1]%><br />
 								Change: <%= CrashInstance.ChangeListVersion %><br />
 							</td>
@@ -233,12 +237,15 @@
 								<%=CrashInstance.EngineMode %><br />
 								<%=CrashInstance.PlatformName %><br />
 							</td><td class="CallStack" >
+								<%--Display callstack only for the first crash--%>
+								<% if(IterationCount==1)
+								{ %>
 								<div style="clip : auto; ">
 									<div id='<%=CrashInstance.Id %>-TrimmedCallStackBox' class='TrimmedCallStackBox'>
-										<% foreach( CallStackEntry Entry in CrashInstance.GetCallStackEntries( 0, 4 ) )
+										<% foreach( CallStackEntry Entry in Model.CallStack.CallStackEntries.Take(4) )
 										{%>
 											<span class="function-name">
-												<%=Url.CallStackSearchLink( Html.Encode( Entry.GetTrimmedFunctionName( 40 ) ), Model )%>
+												<%=Url.CallStackSearchLink( Html.Encode( Entry.GetTrimmedFunctionName( 60 ) ), Model )%>
 											</span>
 											<br />
 										<%} %>
@@ -246,16 +253,17 @@
 						
 									<a class='FullCallStackTrigger' ><span class='FullCallStackTriggerText'>Full Callstack</span>
 										<div id='<%=CrashInstance.Id %>-FullCallStackBox' class='FullCallStackBox'>
-											<%foreach( CallStackEntry Entry in CrashInstance.GetCallStackEntries( 0, 40 ) )
+											<%foreach( CallStackEntry Entry in Model.CallStack.CallStackEntries.Take( 40 ) )
 												{%>
 												<span class="FunctionName">
-													<%=Html.Encode( Entry.GetTrimmedFunctionName( 60 ) )%>
+													<%=Html.Encode( Entry.GetTrimmedFunctionName( 80 ) )%>
 												</span>
 												<br />
 											<%} %>
 										</div>
 									</a>
 								</div>
+								<%} %>
 							</td>
 							<td class="Description"><span><%=CrashInstance.Description%>&nbsp;</span></td> 
 						</tr>
diff --git a/Engine/Source/Programs/CrashReporter/CrashReportWebSite/Views/Buggs/ViewBuggs.ascx b/Engine/Source/Programs/CrashReporter/CrashReportWebSite/Views/Buggs/ViewBuggs.ascx
--- a/Engine/Source/Programs/CrashReporter/CrashReportWebSite/Views/Buggs/ViewBuggs.ascx
+++ b/Engine/Source/Programs/CrashReporter/CrashReportWebSite/Views/Buggs/ViewBuggs.ascx
@@ -40,16 +40,17 @@
 		<th>&nbsp;</th>
 		<!-- if you add a new column be sure to update the switch statement in the repository SortBy function -->
 		<th style='width:  6em;'><%=Url.TableHeader( "Id", "Id", Model )%></th>
-		<th style='width: 15em;'><%=Url.TableHeader( "Build Version", "BuildVersion", Model )%></th>
-		<th style='width: 15em;'><%=Url.TableHeader( "Crashes In Time Frame", "CrashesInTimeFrame", Model )%></th>
+		<th style='width: 8em;'><%=Url.TableHeader( "Build Version", "BuildVersion", Model )%></th>
+		<th style='width: 8em;'><%=Url.TableHeader( "Crashes In Time Frame (Group)", "CrashesInTimeFrameGroup", Model )%></th>
+		<th style='width: 8em;'><%=Url.TableHeader( "Crashes In Time Frame (All)", "CrashesInTimeFrameAll", Model )%></th>
 		<th style='width: 12em;'><%=Url.TableHeader( "Latest Crash", "LatestCrash", Model )%></th>
 		<th style='width: 12em;'><%=Url.TableHeader( "First Crash", "FirstCrash", Model )%></th> 
-		<th style='width: 11em;'><%=Url.TableHeader( "# of Crashes", "NumberOfCrashes", Model )%></th> 
-		<th style='width: 11em;'><%=Url.TableHeader( "Users Affected", "NumberOfUsers", Model )%></th> 
-		<th style='width: 14em;'><%=Url.TableHeader( "Pattern", "Pattern", Model )%></th> 
-		<th style='width: 12em;'><%=Url.TableHeader( "Status", "Status", Model )%></th>
-		<th style='width: 12em;'><%=Url.TableHeader( "FixedCL#", "FixedChangeList", Model )%></th>
-		<th style='width: 12em;'><%=Url.TableHeader( "TTPID", "TTPID", Model )%></th>
+		<th style='width: 8em;'><%=Url.TableHeader( "# of Crashes", "NumberOfCrashes", Model )%></th> 
+		<th style='width: 8em;'><%=Url.TableHeader( "Users Affected", "NumberOfUsers", Model )%></th> 
+		<th style='width: 24em;'><%=Url.TableHeader( "Pattern", "Pattern", Model )%></th> 
+		<th style='width: 8em;'><%=Url.TableHeader( "Status", "Status", Model )%></th>
+		<th style='width: 8em;'><%=Url.TableHeader( "FixedCL#", "FixedChangeList", Model )%></th>
+		<th style='width: 8em;'><%=Url.TableHeader( "JIRA", "TTPID", Model )%></th>
 	</tr>
 	<%
 		using( FAutoScopedLogTimer LogTimer = new FAutoScopedLogTimer( this.GetType().ToString() ) )
@@ -66,13 +67,13 @@
 						if( string.IsNullOrWhiteSpace( CurrentBugg.FixedChangeList ) && string.IsNullOrWhiteSpace( CurrentBugg.TTPID ) )
 						{
 							BuggRowColor = "#FFFF88"; // yellow
-							BuggColorDescription = "This CrashGroup has not been fixed or assigned a TTP";
+							BuggColorDescription = "This CrashGroup has not been fixed or assigned a JIRA";
 						}
 
 						if( !string.IsNullOrWhiteSpace( CurrentBugg.TTPID ) && string.IsNullOrWhiteSpace( CurrentBugg.FixedChangeList ) )
 						{
 							BuggRowColor = "#D01F3C"; // red
-							BuggColorDescription = "This CrashGroup has  been assigned a TTP: " + CurrentBugg.TTPID + " but has not been fixed.";
+							BuggColorDescription = "This CrashGroup has  been assigned a JIRA: " + CurrentBugg.TTPID + " but has not been fixed.";
 						}
 
 						if( CurrentBugg.Status == "Coder" )
@@ -97,11 +98,15 @@
 			<td class="BuggTd" style="background-color: <%=BuggRowColor %>;" title="<%=BuggColorDescription %>"></td>
 			<td class="Id"><%=Html.ActionLink( CurrentBugg.Id.ToString(), "Show", new { controller = "Buggs", id = CurrentBugg.Id }, null )%></td>
 			<td><%=CurrentBugg.BuildVersion%></td>
-			<td><%=CurrentBugg.CrashesInTimeFrame%></td>
+			<td><%=CurrentBugg.CrashesInTimeFrameGroup%></td>
+			<td><%=CurrentBugg.CrashesInTimeFrameAll%></td>
 			<td><%=CurrentBugg.TimeOfLastCrash%></td>
 			<td><%=CurrentBugg.TimeOfFirstCrash%></td>
 			<td><%=CurrentBugg.NumberOfCrashes%> </td>
-			<td><%=CurrentBugg.NumberOfUsers%></td>
+
+			
+			<td><%=CurrentBugg.NumberOfUniqueMachines%></td>
+			
 			<td class="CallStack">
 				<div>
 					<div id='Div1' class='TrimmedCallStackBox'>
@@ -111,9 +116,9 @@
 						foreach( string FunctionCall in FunctionCalls )
 						{
 							string FunctionCallDisplay = FunctionCall;
-							if( FunctionCall.Length > 45 )
+							if( FunctionCall.Length > 60 )
 							{
-								FunctionCallDisplay = FunctionCall.Substring( 0, 45 );
+								FunctionCallDisplay = FunctionCall.Substring( 0, 60 );
 							}%>
 								<span class="FunctionName">
 									<%=Url.CallStackSearchLink( Html.Encode( FunctionCallDisplay ), Model )%>
@@ -142,7 +147,7 @@
 			</td>
 			<td><%=CurrentBugg.Status%></td>
 			<td><%=CurrentBugg.FixedChangeList%></td>
-			<td><%=CurrentBugg.TTPID%></td>
+			<td> <span><a href="https://jira.ol.epicgames.net/browse/<%=CurrentBugg.TTPID%>" target="_blank"><%=CurrentBugg.TTPID%></a></span>  </td>
 		</tr>
 	<%				}
 				}
diff --git a/Engine/Source/Programs/CrashReporter/CrashReportWebSite/Views/Crashes/Show.aspx b/Engine/Source/Programs/CrashReporter/CrashReportWebSite/Views/Crashes/Show.aspx
--- a/Engine/Source/Programs/CrashReporter/CrashReportWebSite/Views/Crashes/Show.aspx
+++ b/Engine/Source/Programs/CrashReporter/CrashReportWebSite/Views/Crashes/Show.aspx
@@ -188,10 +188,10 @@
 				<dd ><%=Html.DisplayFor(m => Model.Crash.PlatformName) %></dd>
 
 			<dt>Machine Id</dt> 
-				<dd class='even'><%=Html.DisplayFor(m => Model.Crash.ComputerName) %></dd>
+				<dd class='even'><small><%=Html.DisplayFor(m => Model.Crash.ComputerName) %></small></dd>
 
 			<dt>Epic Account Id</dt> 
-				<dd class='even'><%=Html.DisplayFor(m => Model.Crash.EpicAccountId) %></dd>
+				<dd class='even'><small><%=Html.DisplayFor(m => Model.Crash.EpicAccountId) %></small></dd>
 
 			<dt>Allowed to contact</dt> 
 				<dd class='even'><%=Html.DisplayFor(m => Model.Crash.AllowToBeContacted) %></dd>
@@ -205,8 +205,9 @@
 			<dt>Changelist #</dt>
 				<dd class='even'><%=Html.DisplayFor(m => Model.Crash.ChangeListVersion) %></dd>
 
-			<dt>TTP</dt>
-				<dd ><%=Html.DisplayFor(m => Model.Crash.TTPID) %></dd>
+			<dt>JIRA</dt>
+				<%--<dd ><%=Html.DisplayFor(m => Model.Crash.TTPID) %></dd>--%>
+				<dd><a href="https://jira.ol.epicgames.net/browse/<%=Model.Crash.TTPID%>" target="_blank"><%=Model.Crash.TTPID%></a></dd>
 
 			<dt>Status</dt>
 				<dd ><%=Html.DisplayFor(m => Model.Crash.Status) %></dd>
diff --git a/Engine/Source/Programs/CrashReporter/CrashReportWebSite/Views/Crashes/ViewCrash.ascx b/Engine/Source/Programs/CrashReporter/CrashReportWebSite/Views/Crashes/ViewCrash.ascx
--- a/Engine/Source/Programs/CrashReporter/CrashReportWebSite/Views/Crashes/ViewCrash.ascx
+++ b/Engine/Source/Programs/CrashReporter/CrashReportWebSite/Views/Crashes/ViewCrash.ascx
@@ -49,7 +49,7 @@
 
 					<input type="submit" name="SetStatusSubmit" value="Set" class="SetButton" />
 
-					<span id="set-ttp" style="">TTP</span>
+					<span id="set-ttp" style="">JIRA</span>
 					<input name="SetTTP" type="text" id="ttp-text-box" />
 					<input type="submit" name="SetTTPSubmit" value="Set" class="SetButton" />
 
@@ -71,7 +71,7 @@
 		<th style='width: 11em;'><%=Url.TableHeader( "Game", "GameName", Model )%></th> 
 		<th style='width: 11em;'><%=Url.TableHeader( "Mode", "EngineMode", Model )%></th> 
 		<th style='width: 16em;'><%=Url.TableHeader( "FixedCL#", "FixedChangeList", Model )%></th>
-		<th style='width: 9em;'><%=Url.TableHeader( "Bug", "TTPID", Model) %></th>
+		<th style='width: 9em;'><%=Url.TableHeader( "JIRA", "TTPID", Model) %></th>
 		<th style='width: 9em;'><%=Url.TableHeader( "Branch", "Branch", Model) %></th>
 		<th>Description</th>
 		<th style='width: 15em'><%=Url.TableHeader( "Message", "Summary", Model )%></th>
@@ -109,7 +109,7 @@
 							if( string.IsNullOrWhiteSpace( CurrentCrash.FixedChangeList ) && string.IsNullOrWhiteSpace( CurrentCrash.TTPID ) )
 							{
 								CrashRowColor = "#FFFF88"; // yellow
-								CrashColorDescription = "This crash has not been fixed or assigned a TTP";
+								CrashColorDescription = "This crash has not been fixed or assigned a JIRA";
 							}
 
 							if( !string.IsNullOrWhiteSpace( CurrentCrash.FixedChangeList ) )
@@ -122,7 +122,7 @@
 							if( ( BuggCrash != null ) && !string.IsNullOrWhiteSpace( CurrentCrash.TTPID ) && string.IsNullOrWhiteSpace( CurrentCrash.FixedChangeList ) )
 							{
 								CrashRowColor = "#D01F3C"; // red
-								CrashColorDescription = "This crash has occurred more than once and been assigned a TTP: " + CurrentCrash.TTPID + " but has not been fixed.";
+								CrashColorDescription = "This crash has occurred more than once and been assigned a JIRA: " + CurrentCrash.TTPID + " but has not been fixed.";
 							}
 
 							if( CurrentCrash.Status == "Tester" )
@@ -175,7 +175,7 @@
 									<td class="Game"><%=CurrentCrash.GameName%></td>
 									<td class="Mode"><%=CurrentCrash.EngineMode%></td>
 									<td class="FixedChangeList"><%=CurrentCrash.FixedChangeList%></td>
-									<td class="Ttp"><%=CurrentCrash.TTPID%>&nbsp;</td>
+									<td class="Jira"> <span><a href="https://jira.ol.epicgames.net/browse/<%=CurrentCrash.TTPID%>" target="_blank"><%=CurrentCrash.TTPID%></a></span>  </td>
 									<td class="Branch"><%=CurrentCrash.Branch%>&nbsp;</td>
 									<td class="Description"><span class="TableData"><%=CurrentCrash.Description%>&nbsp;</span></td>
 									<td class="Summary"><%=CurrentCrash.Summary%></td>
