CrashReportWebsite - Missing file

diff --git a/Engine/Source/Programs/CrashReporter/CrashReportWebSite/Models/BuggRepository.cs b/Engine/Source/Programs/CrashReporter/CrashReportWebSite/Models/BuggRepository.cs
--- a/Engine/Source/Programs/CrashReporter/CrashReportWebSite/Models/BuggRepository.cs
+++ b/Engine/Source/Programs/CrashReporter/CrashReportWebSite/Models/BuggRepository.cs
@@ -21,14 +21,14 @@ namespace Tools.CrashReporter.CrashReportWebSite.Models
 	{
 		private const string DefaultUserGroup = "General";
 
-		private readonly CrashReportDataContext BuggsDataContext;
+		private readonly CrashReportDataContext Context;
 
 		/// <summary>
 		/// The default constructor to create a data context to access the database.
 		/// </summary>
 		public BuggRepository()
 		{
-			BuggsDataContext = new CrashReportDataContext();
+			Context = new CrashReportDataContext();
 		}
 
 		/// <summary>
@@ -46,7 +46,7 @@ namespace Tools.CrashReporter.CrashReportWebSite.Models
 		/// <param name="Disposing">true if the Dispose call is from user code, and not system code.</param>
 		protected virtual void Dispose( bool Disposing )
 		{
-			BuggsDataContext.Dispose();
+			Context.Dispose();
 		}
 
 		/// <summary>
@@ -55,7 +55,7 @@ namespace Tools.CrashReporter.CrashReportWebSite.Models
 		/// <returns>The data context used to interface with the database.</returns>
 		public CrashReportDataContext GetDataContext()
 		{
-			return BuggsDataContext;
+			return Context;
 		}
 
 		/// <summary>
@@ -73,7 +73,7 @@ namespace Tools.CrashReporter.CrashReportWebSite.Models
 				{
 					Result =
 					(
-						from BuggDetail in BuggsDataContext.Buggs
+						from BuggDetail in Context.Buggs
 						where BuggDetail.Id == Id
 						select BuggDetail
 					).FirstOrDefault();
@@ -98,7 +98,7 @@ namespace Tools.CrashReporter.CrashReportWebSite.Models
 			{
 				try
 				{
-					FunctionCall FunctionCall = BuggsDataContext.FunctionCalls.Where( FunctionCallInstance => FunctionCallInstance.Call.Contains( FunctionCallName ) ).First();
+					FunctionCall FunctionCall = Context.FunctionCalls.Where( FunctionCallInstance => FunctionCallInstance.Call.Contains( FunctionCallName ) ).First();
 					return FunctionCall.Id;
 				}
 				catch( Exception Ex )
@@ -136,7 +136,7 @@ namespace Tools.CrashReporter.CrashReportWebSite.Models
 				List<string> FunctionCalls = new List<string>();
 				try
 				{
-					List<FunctionCall> Funcs = BuggsDataContext.FunctionCalls.Where( FuncCall => Ids.Contains( FuncCall.Id ) ).ToList();
+					List<FunctionCall> Funcs = Context.FunctionCalls.Where( FuncCall => Ids.Contains( FuncCall.Id ) ).ToList();
 					// Order by Ids
 					foreach( int Id in Ids )
 					{
@@ -160,7 +160,7 @@ namespace Tools.CrashReporter.CrashReportWebSite.Models
 		/// <param name="Crashes">The set of crashes to add to the Bugg.</param>
 		public void UpdateBuggData( Bugg Bugg, IEnumerable<Crash> Crashes )
 		{
-			using( FAutoScopedLogTimer LogTimer = new FAutoScopedLogTimer( this.GetType().ToString()+ "(" + Bugg.Id + ")" ) )
+			using( FAutoScopedLogTimer LogTimer = new FAutoScopedLogTimer( this.GetType().ToString() + "(" + Bugg.Id + ")" ) )
 			{
 				FLogger.WriteEvent( "UpdateBuggData Bugg.Id=" + Bugg.Id );
 
@@ -229,7 +229,7 @@ namespace Tools.CrashReporter.CrashReportWebSite.Models
 
 					if( bHasChanges )
 					{
-						BuggsDataContext.SubmitChanges();
+						Context.SubmitChanges();
 					}
 				}
 				catch( Exception Ex )
@@ -251,13 +251,13 @@ namespace Tools.CrashReporter.CrashReportWebSite.Models
 			try
 			{
 				// Make sure we don't already have this relationship
-				if( BuggsDataContext.Buggs_Crashes.Where( BuggInstance => BuggInstance.CrashId == CurrentCrash.Id && BuggInstance.BuggId == Bugg.Id ).Count() < 1 )
+				if( Context.Buggs_Crashes.Where( BuggInstance => BuggInstance.CrashId == CurrentCrash.Id && BuggInstance.BuggId == Bugg.Id ).Count() < 1 )
 				{
 					// We don't so create the relationship
 					Buggs_Crash NewBugg = new Buggs_Crash();
 					NewBugg.CrashId = CurrentCrash.Id;
 					NewBugg.BuggId = Bugg.Id;
-					BuggsDataContext.Buggs_Crashes.InsertOnSubmit( NewBugg );
+					Context.Buggs_Crashes.InsertOnSubmit( NewBugg );
 					return true;
 				}
 			}
@@ -279,14 +279,14 @@ namespace Tools.CrashReporter.CrashReportWebSite.Models
 		{
 			try
 			{
-				int BuggUserCount = BuggsDataContext.Buggs_Users.Where( BuggUserInstance => BuggUserInstance.BuggId == Bugg.Id && BuggUserInstance.UserNameId == UserNameId ).Count();
+				int BuggUserCount = Context.Buggs_Users.Where( BuggUserInstance => BuggUserInstance.BuggId == Bugg.Id && BuggUserInstance.UserNameId == UserNameId ).Count();
 				if( BuggUserCount < 1 )
 				{
 					Buggs_User NewBuggsUser = new Buggs_User();
 					NewBuggsUser.BuggId = Bugg.Id;
 					NewBuggsUser.UserNameId = UserNameId;
 
-					BuggsDataContext.Buggs_Users.InsertOnSubmit( NewBuggsUser );
+					Context.Buggs_Users.InsertOnSubmit( NewBuggsUser );
 					return true;
 				}
 			}
@@ -306,7 +306,7 @@ namespace Tools.CrashReporter.CrashReportWebSite.Models
 		{
 			IQueryable<Bugg> Buggs =
 			(
-				from BuggDetail in BuggsDataContext.Buggs
+				from BuggDetail in Context.Buggs
 				select BuggDetail
 			);
 
@@ -319,12 +319,12 @@ namespace Tools.CrashReporter.CrashReportWebSite.Models
 		/// <param name="Results">The unfiltered set of Buggs.</param>
 		/// <param name="Query">The query to use as a filter.</param>
 		/// <returns>A filtered set of Buggs.</returns>
-		public IQueryable<Bugg> Search( IQueryable<Bugg> Results, string Query )
+		public IEnumerable<Bugg> Search( IEnumerable<Bugg> Results, string Query )
 		{
 			using( FAutoScopedLogTimer LogTimer = new FAutoScopedLogTimer( this.GetType().ToString() + "(Query=" + Query + ")" ) )
 			{
 				// Also may want to revisit how we search since this could get inefficient for a big search set.
-				IQueryable<Bugg> Buggs;
+				IEnumerable<Bugg> Buggs;
 				try
 				{
 					string QueryString = HttpUtility.HtmlDecode( Query.ToString() );
@@ -345,7 +345,7 @@ namespace Tools.CrashReporter.CrashReportWebSite.Models
 						}
 					}
 
-					Buggs = (IQueryable<Bugg>)Results.Search( new string[] { "Pattern" }, TermsToUse.Split( "+".ToCharArray() ) );
+					Buggs = (IEnumerable<Bugg>)Results.AsQueryable().Search( new string[] { "Pattern" }, TermsToUse.Split( "+".ToCharArray() ) );
 				}
 				catch( Exception Ex )
 				{
@@ -371,15 +371,15 @@ namespace Tools.CrashReporter.CrashReportWebSite.Models
 				// The downside is that if we add more parameters each query may need to be updated.... Or we just add new case statements
 				// The other downside is that there's less code reuse, but that may be worth it.
 
-				IQueryable<Bugg> Results = null;
+				IEnumerable<Bugg> Results = null;
 				int Skip = ( FormData.Page - 1 ) * FormData.PageSize;
 				int Take = FormData.PageSize;
 
 				// Get every Bugg.
-				Results = ListAll();
+				var ResultsAll = ListAll();
 
 				// Look at all Buggs that are still 'open' i.e. the last crash occurred in our date range.
-				Results = FilterByDate( Results, FormData.DateFrom, FormData.DateTo );
+				Results = FilterByDate( ResultsAll, FormData.DateFrom, FormData.DateTo );
 
 				// Filter results by build version.
 				Results = FilterByBuildVersion( Results, FormData.BuildVersion );
@@ -411,22 +411,22 @@ namespace Tools.CrashReporter.CrashReportWebSite.Models
 				}
 
 				// Get UserGroup ResultCounts
-				Dictionary<string, int> GroupCounts = GetCountsByGroup( Results );
+				SortedDictionary<string, int> GroupCounts = GetCountsByGroup( Results );
 
 				// Filter by user group if present
 				Results = FilterByUserGroup( Results, FormData.UserGroup );
 
 				// Pass in the results and return them sorted properly
-				IEnumerable<Bugg> SortedResults = GetSortedResults( Results, FormData.SortTerm, ( FormData.SortOrder == "Descending" ), FormData.DateFrom, FormData.DateTo );
+				IEnumerable<Bugg> SortedResults = GetSortedResults( Results, FormData.SortTerm, ( FormData.SortOrder == "Descending" ), FormData.DateFrom, FormData.DateTo, FormData.UserGroup );
 
 				// Grab just the results we want to display on this page
 
-			    var SortedResultsList = SortedResults.ToList();
-                var TotalCountedRecords = SortedResultsList.Count();
+				var SortedResultsList = SortedResults.ToList();
+				var TotalCountedRecords = SortedResultsList.Count();
 
-			    SortedResultsList = SortedResultsList.GetRange(Skip, TotalCountedRecords >= Skip + Take ? Take : TotalCountedRecords);
+				SortedResultsList = SortedResultsList.GetRange( Skip, TotalCountedRecords >= Skip + Take ? Take : TotalCountedRecords );
 
-			    return new BuggsViewModel
+				return new BuggsViewModel
 				{
 					Results = SortedResultsList,
 					PagingInfo = new PagingInfo { CurrentPage = FormData.Page, PageSize = FormData.PageSize, TotalResults = Results.Count() },
@@ -448,10 +448,10 @@ namespace Tools.CrashReporter.CrashReportWebSite.Models
 		/// </summary>
 		/// <param name="Buggs">A list of Buggs to bucket.</param>
 		/// <returns>A dictionary of user group names, and the count of Buggs for each group.</returns>
-		public Dictionary<string, int> GetCountsByGroup( IQueryable<Bugg> Buggs )
+		public SortedDictionary<string, int> GetCountsByGroup( IEnumerable<Bugg> Buggs )
 		{
 			// @TODO yrx 2014-11-06 Optimize
-			using( FAutoScopedLogTimer LogTimer = new FAutoScopedLogTimer( this.GetType().ToString() ) )
+			using( FAutoScopedLogTimer LogTimer = new FAutoScopedLogTimer( this.GetType().ToString() + " SQL OPT" ) )
 			{
 				Dictionary<string, int> Results = new Dictionary<string, int>();
 
@@ -460,14 +460,14 @@ namespace Tools.CrashReporter.CrashReportWebSite.Models
 					Results =
 					(
 						from BuggDetail in Buggs
-						join BuggsUserGroupDetail in BuggsDataContext.Buggs_UserGroups on BuggDetail.Id equals BuggsUserGroupDetail.BuggId
-						join UserGroupDetail in BuggsDataContext.UserGroups on BuggsUserGroupDetail.UserGroupId equals UserGroupDetail.Id
+						join BuggsUserGroupDetail in Context.Buggs_UserGroups on BuggDetail.Id equals BuggsUserGroupDetail.BuggId
+						join UserGroupDetail in Context.UserGroups on BuggsUserGroupDetail.UserGroupId equals UserGroupDetail.Id
 						group 0 by UserGroupDetail.Name into GroupCount
 						select new { Key = GroupCount.Key, Count = GroupCount.Count() }
 					).ToDictionary( x => x.Key, y => y.Count );
 
 					// Add in all groups, even though there are no crashes associated
-					IEnumerable<string> UserGroups = ( from UserGroupDetail in BuggsDataContext.UserGroups select UserGroupDetail.Name );
+					IEnumerable<string> UserGroups = ( from UserGroupDetail in Context.UserGroups select UserGroupDetail.Name );
 					foreach( string UserGroupName in UserGroups )
 					{
 						if( !Results.Keys.Contains( UserGroupName ) )
@@ -481,8 +481,9 @@ namespace Tools.CrashReporter.CrashReportWebSite.Models
 					Debug.WriteLine( "Exception in GetCountsByGroup: " + Ex.ToString() );
 				}
 
-				return Results;
-		}
+				SortedDictionary<string, int> SortedResults = new SortedDictionary<string, int>( Results );
+				return SortedResults;
+			}
 		}
 
 		/// <summary>
@@ -492,10 +493,14 @@ namespace Tools.CrashReporter.CrashReportWebSite.Models
 		/// <param name="DateFrom">The earliest date to filter by.</param>
 		/// <param name="DateTo">The latest date to filter by.</param>
 		/// <returns>The set of Buggs between the earliest and latest date.</returns>
-		public IQueryable<Bugg> FilterByDate( IQueryable<Bugg> Results, DateTime DateFrom, DateTime DateTo )
+		public IEnumerable<Bugg> FilterByDate( IQueryable<Bugg> Results, DateTime DateFrom, DateTime DateTo )
 		{
-			IQueryable<Bugg> BuggsInTimeFrame = Results.Where( Bugg => Bugg.TimeOfLastCrash >= DateFrom && Bugg.TimeOfLastCrash <= AddOneDayToDate( DateTo ) );
-			return BuggsInTimeFrame;			
+			using( FAutoScopedLogTimer LogTimer = new FAutoScopedLogTimer( this.GetType().ToString() + " SQL" ) )
+			{
+				IQueryable<Bugg> BuggsInTimeFrame = Results.Where( Bugg => Bugg.TimeOfLastCrash >= DateFrom && Bugg.TimeOfLastCrash <= AddOneDayToDate( DateTo ) );
+				IEnumerable<Bugg> BuggsInTimeFrameEnumerable = BuggsInTimeFrame.ToList();
+				return BuggsInTimeFrameEnumerable;
+			}
 		}
 
 		/// <summary>
@@ -504,9 +509,9 @@ namespace Tools.CrashReporter.CrashReportWebSite.Models
 		/// <param name="Results">The unfiltered set of Buggs.</param>
 		/// <param name="BuildVersion">The build version to filter by.</param>
 		/// <returns>The set of Buggs that matches specified build version</returns>
-		public IQueryable<Bugg> FilterByBuildVersion( IQueryable<Bugg> Results, string BuildVersion )
+		public IEnumerable<Bugg> FilterByBuildVersion( IEnumerable<Bugg> Results, string BuildVersion )
 		{
-			IQueryable<Bugg> BuggsForBuildVersion = Results;
+			IEnumerable<Bugg> BuggsForBuildVersion = Results;
 
 			// Filter by BuildVersion
 			if( !string.IsNullOrEmpty( BuildVersion ) )
@@ -518,8 +523,8 @@ namespace Tools.CrashReporter.CrashReportWebSite.Models
 					select MyBugg
 				);
 			}
-			return BuggsForBuildVersion;	
-		} 
+			return BuggsForBuildVersion;
+		}
 
 		/// <summary>
 		/// Filter a set of Buggs by user group name.
@@ -527,19 +532,20 @@ namespace Tools.CrashReporter.CrashReportWebSite.Models
 		/// <param name="SetOfBuggs">The unfiltered set of Buggs.</param>
 		/// <param name="UserGroup">The user group name to filter by.</param>
 		/// <returns>The set of Buggs by users in the requested user group.</returns>
-		public IQueryable<Bugg> FilterByUserGroup( IQueryable<Bugg> SetOfBuggs, string UserGroup )
+		public IEnumerable<Bugg> FilterByUserGroup( IEnumerable<Bugg> SetOfBuggs, string UserGroup )
 		{
-			using( FAutoScopedLogTimer LogTimer = new FAutoScopedLogTimer( this.GetType().ToString() + "(UserGroup=" + UserGroup + ")" ) )
+			using( FAutoScopedLogTimer LogTimer = new FAutoScopedLogTimer( this.GetType().ToString() + "(UserGroup=" + UserGroup + ") SQL" ) )
 			{
 				IQueryable<Bugg> NewSetOfBuggs = null;
+				IQueryable<Bugg> SetOfBuggsQueryable = SetOfBuggs.AsQueryable();
 
 				try
 				{
 					NewSetOfBuggs =
 					(
-						from BuggDetail in SetOfBuggs
-						join BuggsUserGroupDetail in BuggsDataContext.Buggs_UserGroups on BuggDetail.Id equals BuggsUserGroupDetail.BuggId
-						join UserGroupDetail in BuggsDataContext.UserGroups on BuggsUserGroupDetail.UserGroupId equals UserGroupDetail.Id
+						from BuggDetail in SetOfBuggsQueryable
+						join BuggsUserGroupDetail in Context.Buggs_UserGroups on BuggDetail.Id equals BuggsUserGroupDetail.BuggId
+						join UserGroupDetail in Context.UserGroups on BuggsUserGroupDetail.UserGroupId equals UserGroupDetail.Id
 						where UserGroupDetail.Name.Contains( UserGroup )
 						select BuggDetail
 					);
@@ -551,52 +557,35 @@ namespace Tools.CrashReporter.CrashReportWebSite.Models
 				}
 
 
-				return NewSetOfBuggs;
+				return NewSetOfBuggs.AsEnumerable();
 			}
 		}
 
+		/// <summary></summary>
+		public int GetIdFromUserGroup( string UserGroup )
+		{
+			var Group = Context.UserGroups.Where( X => X.Name.Contains( UserGroup ) ).FirstOrDefault();
+			return Group.Id;
+		}
 
-        public DataTable LINQToDataTable<T>(IEnumerable<T> varlist)
-        {
-             DataTable dtReturn = new DataTable();
-
-             // column names 
-             PropertyInfo[] oProps = null;
-
-             if (varlist == null) return dtReturn;
-
-             foreach (T rec in varlist)
-             {
-                  // Use reflection to get property names, to create table, Only first time, others will follow 
-                  if (oProps == null)
-                  {
-                       oProps = ((Type)rec.GetType()).GetProperties();
-                       foreach (PropertyInfo pi in oProps)
-                       {
-                            Type colType = pi.PropertyType;
-
-                            if ((colType.IsGenericType) && (colType.GetGenericTypeDefinition()      
-                            ==typeof(Nullable<>)))
-                             {
-                                 colType = colType.GetGenericArguments()[0];
-                             }
-
-                            dtReturn.Columns.Add(new DataColumn(pi.Name, colType));
-                       }
-                  }
-
-                  DataRow dr = dtReturn.NewRow();
-
-                  foreach (PropertyInfo pi in oProps)
-                  {
-                       dr[pi.Name] = pi.GetValue(rec, null) ?? DBNull.Value;
-                  }
-
-                  dtReturn.Rows.Add(dr);
-             }
-             return dtReturn;
-        }
+		/// <summary></summary>
+		public HashSet<int> GetUserIdsFromUserGroup( string UserGroup )
+		{
+			int UserGroupId = GetIdFromUserGroup( UserGroup );
+			var UserIds = Context.Users.Where( X => X.UserGroupId == UserGroupId ).Select( X => X.Id );
+			return new HashSet<int>( UserIds );
+		}
 
+		/// <summary></summary>
+		public HashSet<string> GetUserNamesFromUserGroups( string UserGroup )
+		{
+			using( FAutoScopedLogTimer LogTimer = new FAutoScopedLogTimer( this.GetType().ToString() + " SQL" ) )
+			{
+				int UserGroupId = GetIdFromUserGroup( UserGroup );
+				var UserNames = Context.Users.Where( X => X.UserGroupId == UserGroupId ).Select( X => X.UserName );
+				return new HashSet<string>( UserNames );
+			}
+		}
 
 		private DateTime AddOneDayToDate( DateTime Date )
 		{
@@ -611,86 +600,198 @@ namespace Tools.CrashReporter.CrashReportWebSite.Models
 		/// <param name="bSortDescending">Whether to sort by descending or ascending.</param>
 		/// <param name="DateFrom">The date of the earliest Bugg to examine.</param>
 		/// <param name="DateTo">The date of the most recent Bugg to examine.</param>
+		/// <param name="UserGroup">The user group name to filter by.</param>
 		/// <returns>A sorted container of Buggs.</returns>
-		public IQueryable<Bugg> GetSortedResults( IQueryable<Bugg> Results, string SortTerm, bool bSortDescending, DateTime DateFrom, DateTime DateTo )
+		public IEnumerable<Bugg> GetSortedResults( IEnumerable<Bugg> Results, string SortTerm, bool bSortDescending, DateTime DateFrom, DateTime DateTo, string UserGroup )
 		{
 			using( FAutoScopedLogTimer LogTimer = new FAutoScopedLogTimer( this.GetType().ToString() ) )
 			{
 				try
 				{
+					// Get the group id and grab all buggs for the specified group.
+					int NumResults = Results.Count();
+					HashSet<int> ResultIds = new HashSet<int>( Results.Select( x => x.Id ) );
+					HashSet<string> UserNamesForUserGroup = GetUserNamesFromUserGroups( UserGroup );
+
+
+					// Simplified query.
+					var BuggIdToCountMapGroup = new Dictionary<int, int>();
+					var BuggIdToCountMapRest = new Dictionary<int, int>();
+					var BuggIdToMachineSet = new Dictionary<int, HashSet<string>>();
+					Dictionary<string, int> MachineIdToCountMap = new Dictionary<string, int>();
+					List<Buggs_Crash> BuggsFromDate = null;
+					Dictionary<int, string> CrashToUser = null;
+					Dictionary<int, string> CrashToMachine = null;
+
+
+					// Get all buggs from the date range.
+					using( FScopedLogTimer LogTimer1 = new FScopedLogTimer( "BuggRepository.GetSortedResults.BuggsFromDate SQL" ) )
+					{
+						BuggsFromDate =
+								(
+									from BuggCrash in Context.Buggs_Crashes
+									where BuggCrash.Crash.TimeOfCrash >= DateFrom && BuggCrash.Crash.TimeOfCrash <= AddOneDayToDate( DateTo )
+									select BuggCrash
+								).AsEnumerable().ToList();
+
+						var CrashesWithIdUserMachine =
+						(
+							from Crash in Context.Crashes
+							where Crash.TimeOfCrash >= DateFrom && Crash.TimeOfCrash <= AddOneDayToDate( DateTo )
+							select new { Id = Crash.Id, UserName = Crash.UserName, MachineId = Crash.ComputerName }
+						);
+
+						CrashToUser = CrashesWithIdUserMachine.ToDictionary( x => x.Id, y => y.UserName );
+						CrashToMachine = CrashesWithIdUserMachine.ToDictionary( x => x.Id, y => y.MachineId );
+					}
+
+
+					using( FScopedLogTimer LogTimer0 = new FScopedLogTimer( "BuggRepository.GetSortedResults.Filtering" ) )
+					{
+						// This calculates total crashes for selected group and all groups.
+						foreach( Buggs_Crash BuggCrash in BuggsFromDate )
+						{
+							string MachineId;
+							CrashToMachine.TryGetValue( BuggCrash.CrashId, out MachineId );
+
+							string UserName = CrashToUser[BuggCrash.CrashId];
+							bool bValidForGroup = UserNamesForUserGroup.Contains( UserName );
+							if( !bValidForGroup )
+							{
+								int CountRest = 0;
+								bool bFoundRestKey = BuggIdToCountMapRest.TryGetValue( BuggCrash.BuggId, out CountRest );
+								if( bFoundRestKey )
+								{
+									BuggIdToCountMapRest[BuggCrash.BuggId]++;
+								}
+								else
+								{
+									BuggIdToCountMapRest.Add( BuggCrash.BuggId, 1 );
+								}
+
+								continue;
+							}
+
+							int Count = 0;
+							bool bFoundKey = BuggIdToCountMapGroup.TryGetValue( BuggCrash.BuggId, out Count );
+							if( bFoundKey )
+							{
+								BuggIdToCountMapGroup[BuggCrash.BuggId]++;
+							}
+							else
+							{
+								BuggIdToCountMapGroup.Add( BuggCrash.BuggId, 1 );
+							}
+
+							if( MachineId != null && MachineId.Length > 0 )
+							{
+								HashSet<string> MachineSet = null;
+								bool bFoundMachineKey = BuggIdToMachineSet.TryGetValue( BuggCrash.BuggId, out MachineSet );
+								if( !bFoundMachineKey )
+								{
+									BuggIdToMachineSet.Add( BuggCrash.BuggId, new HashSet<string>() );
+								}
+
+								BuggIdToMachineSet[BuggCrash.BuggId].Add( MachineId );
+							}
+						}
+					}
 
-                    var IntermediateResults =
-                    (
-                        from BuggCrashDetail in BuggsDataContext.Buggs_Crashes
-                        where BuggCrashDetail.Crash.TimeOfCrash >= DateFrom && BuggCrashDetail.Crash.TimeOfCrash <= AddOneDayToDate(DateTo)
-                        group BuggCrashDetail by BuggCrashDetail.BuggId into CrashesGrouped
-                        join BuggDetail in Results on CrashesGrouped.Key equals BuggDetail.Id
-                        select new { Bugg = BuggDetail, Count = CrashesGrouped.Count() }
-                    );
-
-                    using (FScopedLogTimer LogTimer2 = new FScopedLogTimer("BuggRepository.GetSortedResults.CrashesInTimeFrame"))
-                    {
-                        foreach (var Result in IntermediateResults)
-                        {
-                            Result.Bugg.CrashesInTimeFrame = Result.Count;
-                        }
-                    }
-                    
-                    switch (SortTerm)
-                    {
-                        case "CrashesInTimeFrame":
-                            IntermediateResults = CrashRepository.OrderBy(IntermediateResults, BuggCrashInstance => BuggCrashInstance.Count, bSortDescending);
-                            break;
-
-                        case "Id":
-                            IntermediateResults = CrashRepository.OrderBy(IntermediateResults, BuggCrashInstance => BuggCrashInstance.Bugg.Id, bSortDescending);
-                            break;
-
-                        case "BuildVersion":
-                            IntermediateResults = CrashRepository.OrderBy(IntermediateResults, BuggCrashInstance => BuggCrashInstance.Bugg.BuildVersion, bSortDescending);
-                            break;
-
-                        case "LatestCrash":
-                            IntermediateResults = CrashRepository.OrderBy(IntermediateResults, BuggCrashInstance => BuggCrashInstance.Bugg.TimeOfLastCrash, bSortDescending);
-                            break;
-
-                        case "FirstCrash":
-                            IntermediateResults = CrashRepository.OrderBy(IntermediateResults, BuggCrashInstance => BuggCrashInstance.Bugg.TimeOfFirstCrash, bSortDescending);
-                            break;
-
-                        case "NumberOfCrashes":
-                            IntermediateResults = CrashRepository.OrderBy(IntermediateResults, BuggCrashInstance => BuggCrashInstance.Bugg.NumberOfCrashes, bSortDescending);
-                            break;
-
-                        case "NumberOfUsers":
-                            IntermediateResults = CrashRepository.OrderBy(IntermediateResults, BuggCrashInstance => BuggCrashInstance.Bugg.NumberOfUsers, bSortDescending);
-                            break;
-
-                        case "Pattern":
-                            IntermediateResults = CrashRepository.OrderBy(IntermediateResults, BuggCrashInstance => BuggCrashInstance.Bugg.Pattern, bSortDescending);
-                            break;
-
-                        case "Status":
-                            IntermediateResults = CrashRepository.OrderBy(IntermediateResults, BuggCrashInstance => BuggCrashInstance.Bugg.Status, bSortDescending);
-                            break;
-
-                        case "FixedChangeList":
-                            IntermediateResults = CrashRepository.OrderBy(IntermediateResults, BuggCrashInstance => BuggCrashInstance.Bugg.FixedChangeList, bSortDescending);
-                            break;
-
-                        case "TTPID":
-                            IntermediateResults = CrashRepository.OrderBy(IntermediateResults, BuggCrashInstance => BuggCrashInstance.Bugg.TTPID, bSortDescending);
-                            break;
-
-                    }
-					return IntermediateResults.Select(x => x.Bugg);
+					using( FScopedLogTimer LogTimer2 = new FScopedLogTimer( "BuggRepository.GetSortedResults.CrashesInTimeFrame" ) )
+					{
+						foreach( var Result in Results )
+						{
+							int GroupCount = 0;
+							BuggIdToCountMapGroup.TryGetValue( Result.Id, out GroupCount );
+							Result.CrashesInTimeFrameGroup = GroupCount;
+
+							int GroupRest = 0;
+							BuggIdToCountMapRest.TryGetValue( Result.Id, out GroupRest );
+							Result.CrashesInTimeFrameAll = GroupCount + GroupRest;
+						}
+					}
+
+					using( FScopedLogTimer LogTimer2 = new FScopedLogTimer( "BuggRepository.GetSortedResults.UniqueMachineCrashesInTimeFrame" ) )
+					{
+						foreach( var Result in Results )
+						{
+							HashSet<string> MachineSet = null;
+							bool bFoundKey = BuggIdToMachineSet.TryGetValue( Result.Id, out MachineSet );
+							if( bFoundKey )
+							{
+								Result.NumberOfUniqueMachines = MachineSet.Count;
+							}
+							else
+							{
+								Result.NumberOfUniqueMachines = 0;
+							}
+						}
+					}
+
+					var IntermediateQueryable = Results.AsQueryable();
+
+					switch( SortTerm )
+					{
+						case "CrashesInTimeFrameGroup":
+							Results = CrashRepository.OrderBy( IntermediateQueryable, BuggCrashInstance => BuggCrashInstance.CrashesInTimeFrameGroup, bSortDescending );
+							break;
+
+						case "CrashesInTimeFrameAll":
+							Results = CrashRepository.OrderBy( IntermediateQueryable, BuggCrashInstance => BuggCrashInstance.CrashesInTimeFrameAll, bSortDescending );
+							break;
+
+						case "Id":
+							Results = CrashRepository.OrderBy( IntermediateQueryable, BuggCrashInstance => BuggCrashInstance.Id, bSortDescending );
+							break;
+
+						case "BuildVersion":
+							Results = CrashRepository.OrderBy( IntermediateQueryable, BuggCrashInstance => BuggCrashInstance.BuildVersion, bSortDescending );
+							break;
+
+						case "LatestCrash":
+							Results = CrashRepository.OrderBy( IntermediateQueryable, BuggCrashInstance => BuggCrashInstance.TimeOfLastCrash, bSortDescending );
+							break;
+
+						case "FirstCrash":
+							Results = CrashRepository.OrderBy( IntermediateQueryable, BuggCrashInstance => BuggCrashInstance.TimeOfFirstCrash, bSortDescending );
+							break;
+
+						case "NumberOfCrashes":
+							Results = CrashRepository.OrderBy( IntermediateQueryable, BuggCrashInstance => BuggCrashInstance.NumberOfCrashes, bSortDescending );
+							break;
+
+
+						case "NumberOfUsers":
+							{
+								Results = CrashRepository.OrderBy( IntermediateQueryable, BuggCrashInstance => BuggCrashInstance.NumberOfUniqueMachines, bSortDescending );
+							}
+							break;
+
+						case "Pattern":
+							Results = CrashRepository.OrderBy( IntermediateQueryable, BuggCrashInstance => BuggCrashInstance.Pattern, bSortDescending );
+							break;
+
+						case "Status":
+							Results = CrashRepository.OrderBy( IntermediateQueryable, BuggCrashInstance => BuggCrashInstance.Status, bSortDescending );
+							break;
+
+						case "FixedChangeList":
+							Results = CrashRepository.OrderBy( IntermediateQueryable, BuggCrashInstance => BuggCrashInstance.FixedChangeList, bSortDescending );
+							break;
+
+						case "TTPID":
+							Results = CrashRepository.OrderBy( IntermediateQueryable, BuggCrashInstance => BuggCrashInstance.TTPID, bSortDescending );
+							break;
+
+					}
+					return Results;
 				}
 				catch( Exception Ex )
 				{
 					Debug.WriteLine( "Exception in GetSortedResults: " + Ex.ToString() );
 				}
 				return Results;
+			}
 		}
 	}
 }
-}
