CrashReportWebsite - Reverted BuildPattern to previous version, for now

diff --git a/Engine/Source/Programs/CrashReporter/CrashReportWebSite/Models/CrashRepository.cs b/Engine/Source/Programs/CrashReporter/CrashReportWebSite/Models/CrashRepository.cs
--- a/Engine/Source/Programs/CrashReporter/CrashReportWebSite/Models/CrashRepository.cs
+++ b/Engine/Source/Programs/CrashReporter/CrashReportWebSite/Models/CrashRepository.cs
@@ -976,10 +976,10 @@ namespace Tools.CrashReporter.CrashReportWebSite.Models
 						Pattern.Add( CurrentFunctionCall.Id.ToString() );
 					}
 
-					CrashInstance.Pattern = "+";
-					CrashInstance.Pattern += string.Join( "+", Pattern );
+					//CrashInstance.Pattern = "+";
+					CrashInstance.Pattern = string.Join( "+", Pattern );
 					// We need something like this +1+2+3+5+ for searching for exact pattern like +5+
-					CrashInstance.Pattern += "+";
+					//CrashInstance.Pattern += "+";
 
 					SubmitChanges();
 				}
