CrashReportWebsite
MaxLinesToParse changed to 64, moved exceptions logging to the log file

diff --git a/Engine/Source/Programs/CrashReporter/CrashReportWebSite/Models/CrashDataFormater.cs b/Engine/Source/Programs/CrashReporter/CrashReportWebSite/Models/CrashDataFormater.cs
--- a/Engine/Source/Programs/CrashReporter/CrashReportWebSite/Models/CrashDataFormater.cs
+++ b/Engine/Source/Programs/CrashReporter/CrashReportWebSite/Models/CrashDataFormater.cs
@@ -105,7 +105,7 @@ namespace Tools.CrashReporter.CrashReportWebSite.Models
 	public class CallStackContainer
 	{
 		/// <summary>The maximum number of call stack lines to parse.</summary>
-		private readonly int MaxLinesToParse = 128;
+		public readonly int MaxLinesToParse = 64;
 
 		/// <summary>A list of parse call stack lines.</summary>
 		private List<CallStackEntry> LocalCallStackEntries = new List<CallStackEntry>();
diff --git a/Engine/Source/Programs/CrashReporter/CrashReportWebSite/Models/CrashRepository.cs b/Engine/Source/Programs/CrashReporter/CrashReportWebSite/Models/CrashRepository.cs
--- a/Engine/Source/Programs/CrashReporter/CrashReportWebSite/Models/CrashRepository.cs
+++ b/Engine/Source/Programs/CrashReporter/CrashReportWebSite/Models/CrashRepository.cs
@@ -136,7 +136,7 @@ namespace Tools.CrashReporter.CrashReportWebSite.Models
 			}
 			catch( Exception Ex )
 			{
-				Debug.WriteLine( "Exception in SetBuggStatus: " + Ex.ToString() );
+				FLogger.WriteException( "SetBuggStatus: " + Ex.ToString() );
 			}
 		}
 
@@ -157,7 +157,7 @@ namespace Tools.CrashReporter.CrashReportWebSite.Models
 			}
 			catch( Exception Ex )
 			{
-				Debug.WriteLine( "Exception in SetBuggFixedChangeList: " + Ex.ToString() );
+				FLogger.WriteException( "SetBuggFixedChangeList: " + Ex.ToString() );
 			}
 		}
 
@@ -178,7 +178,7 @@ namespace Tools.CrashReporter.CrashReportWebSite.Models
 			}
 			catch( Exception Ex )
 			{
-				Debug.WriteLine( "Exception in SetBuggTTPID: " + Ex.ToString() );
+				FLogger.WriteException( "SetBuggTTPID: " + Ex.ToString() );
 			}
 		}
 
@@ -198,7 +198,7 @@ namespace Tools.CrashReporter.CrashReportWebSite.Models
 				}
 				catch( Exception Ex )
 				{
-					Debug.WriteLine( "Exception in PopulateUserInfo: " + Ex.ToString() );
+					FLogger.WriteException( "PopulateUserInfo: " + Ex.ToString() );
 				}
 			}
 		}
@@ -711,7 +711,7 @@ namespace Tools.CrashReporter.CrashReportWebSite.Models
 			}
 			catch( Exception Ex )
 			{
-				Debug.WriteLine( "Exception in AdSetUserGroupdUser: " + Ex.ToString() );
+				FLogger.WriteException( "AdSetUserGroupdUser: " + Ex.ToString() );
 			}
 		}
 
@@ -749,7 +749,7 @@ namespace Tools.CrashReporter.CrashReportWebSite.Models
 				}
 				catch( Exception Ex )
 				{
-					Debug.WriteLine( "Exception in AddUserGroup: " + Ex.ToString() );
+					Debug.WriteLine( "FindOrAddUserGroup: " + Ex.ToString() );
 				}
 
 				return UserGroupNameId;
@@ -932,7 +932,7 @@ namespace Tools.CrashReporter.CrashReportWebSite.Models
 			}
 			catch( Exception Ex )
 			{
-				Debug.WriteLine( "Exception in AddUser: " + Ex.ToString() );
+				FLogger.WriteException( "FindOrAddUser: " + Ex.ToString() );
 			}
 
 			return UserNameId;
@@ -985,7 +985,7 @@ namespace Tools.CrashReporter.CrashReportWebSite.Models
 				}
 				catch( Exception Ex )
 				{
-					Debug.WriteLine( "Exception in BuildPattern: " + Ex.ToString() );
+					FLogger.WriteException( "BuildPattern: " + Ex.ToString() );
 				}
 			}
 		}
