CrashReportWebsite
UECORE-67 Crash report reporting (first iteration)

diff --git a/Engine/Source/Programs/CrashReporter/CrashReportWebSite/Controllers/ReportsController.cs b/Engine/Source/Programs/CrashReporter/CrashReportWebSite/Controllers/ReportsController.cs
new file mode 100644
--- /dev/null
+++ b/Engine/Source/Programs/CrashReporter/CrashReportWebSite/Controllers/ReportsController.cs
@@ -0,0 +1,377 @@
+﻿// Copyright 1998-2015 Epic Games, Inc. All Rights Reserved.
+
+using System;
+using System.Collections.Generic;
+using System.Diagnostics;
+using System.Linq;
+using System.Web.Mvc;
+
+using Tools.CrashReporter.CrashReportWebSite.Models;
+
+namespace Tools.CrashReporter.CrashReportWebSite.Controllers
+{
+
+	/// <summary>
+	/// The controller to handle graphing of crashes per user group over time.
+	/// </summary>
+	public class ReportsController : Controller
+	{
+		/// <summary>Fake id for all user groups</summary>
+		public static readonly int AllUserGroupId = -1;
+
+		/// <summary></summary>
+		private BuggRepository BuggRepository = new BuggRepository();
+		/// <summary></summary>
+		private CrashRepository CrashRepository = new CrashRepository();
+
+		/// <summary></summary>
+		public int GetIdFromUserGroup( string UserGroup )
+		{
+			var Group = CrashRepository.Context.UserGroups.Where( X => X.Name.Contains( UserGroup ) ).FirstOrDefault();
+			return Group.Id;
+		}
+
+		/// <summary></summary>
+		public HashSet<int> GetUserIdsFromUserGroup( string UserGroup )
+		{
+			int UserGroupId = GetIdFromUserGroup( UserGroup );
+			return GetUserIdsFromUserGroupId( UserGroupId );
+		}
+
+		/// <summary></summary>
+		public HashSet<int> GetUserIdsFromUserGroupId( int UserGroupId )
+		{
+			var UserIds = CrashRepository.Context.Users.Where( X => X.UserGroupId == UserGroupId ).Select( X => X.Id );
+			return new HashSet<int>( UserIds );
+		}
+
+		/// <summary>
+		/// Return a dictionary of crashes per group grouped by week.
+		/// </summary>
+		/// <param name="Crashes">A set of crashes to interrogate.</param>
+		/// <param name="UserGroupId">The id of the user group to interrogate.</param>
+		/// <returns>A dictionary of week vs. crash count.</returns>
+		public Dictionary<DateTime, int> GetWeeklyCountsByGroup( List<FCrashMinimal> Crashes, int UserGroupId )
+		{
+			using( FAutoScopedLogTimer LogTimer = new FAutoScopedLogTimer( this.GetType().ToString() + "(UserGroupId=" + UserGroupId + ")" ) )
+			{
+				Dictionary<DateTime, int> Results = new Dictionary<DateTime, int>();
+
+				var UsersIds = GetUserIdsFromUserGroupId( UserGroupId );
+
+				// Trim crashes to user group.
+				if( UserGroupId != DashboardController.AllUserGroupId )
+				{
+					Crashes = Crashes.Where( Crash => UsersIds.Contains( Crash.UserId ) ).ToList();
+				}
+
+				try
+				{
+					Results =
+					(
+						from CrashDetail in Crashes
+						group CrashDetail by CrashDetail.TimeOfCrash.AddDays( -(int)CrashDetail.TimeOfCrash.DayOfWeek ).Date into GroupCount
+						orderby GroupCount.Key
+						select new { Count = GroupCount.Count(), Date = GroupCount.Key }
+					).ToDictionary( x => x.Date, y => y.Count );
+				}
+				catch( Exception Ex )
+				{
+					Debug.WriteLine( "Exception in GetWeeklyCountsByGroup: " + Ex.ToString() );
+				}
+
+				return Results;
+			}
+		}
+
+		/// <summary>
+		/// Return a dictionary of crashes per group grouped by day.
+		/// </summary>
+		/// <param name="Crashes">A set of crashes to interrogate.</param>
+		/// <param name="UserGroupId">The id of the user group to interrogate.</param>
+		/// <returns>A dictionary of day vs. crash count.</returns>
+		public Dictionary<DateTime, int> GetDailyCountsByGroup( List<FCrashMinimal> Crashes, int UserGroupId )
+		{
+			using( FAutoScopedLogTimer LogTimer = new FAutoScopedLogTimer( this.GetType().ToString() + "(UserGroupId=" + UserGroupId + ")" ) )
+			{
+				Dictionary<DateTime, int> Results = new Dictionary<DateTime, int>();
+
+				var UsersIds = GetUserIdsFromUserGroupId( UserGroupId );
+
+				// Trim crashes to user group.
+				if( UserGroupId != DashboardController.AllUserGroupId )
+				{
+					Crashes = Crashes.Where( Crash => UsersIds.Contains( Crash.UserId ) ).ToList();
+				}
+
+				try
+				{
+					Results =
+					(
+						from CrashDetail in Crashes
+						group CrashDetail by CrashDetail.TimeOfCrash.Date into GroupCount
+						orderby GroupCount.Key
+						select new { Count = GroupCount.Count(), Date = GroupCount.Key }
+					).ToDictionary( x => x.Date, y => y.Count );
+				}
+				catch( Exception Ex )
+				{
+					Debug.WriteLine( "Exception in GetDailyCountsByGroup: " + Ex.ToString() );
+				}
+
+				return Results;
+			}
+		}
+
+		System.Data.Linq.Table<FunctionCall> FunctionCalls
+		{
+			get
+			{
+				return this.CrashRepository.Context.FunctionCalls;
+			}
+		}
+
+
+		void BuildPattern( Crash CrashInstance )
+		{
+
+			List<string> Pattern = new List<string>();
+
+			// Get an array of callstack items
+			CallStackContainer CallStack = new CallStackContainer( CrashInstance );
+			CallStack.bDisplayFunctionNames = true;
+
+			if( CrashInstance.Pattern == null )
+			{
+				// Set the module based on the modules in the callstack
+				CrashInstance.Module = CallStack.GetModuleName();
+				try
+				{
+					using( FAutoScopedLogTimer LogTimer = new FAutoScopedLogTimer( this.GetType().ToString() + "(Id=" + CrashInstance.Id + ")" ) )
+					{
+						foreach( CallStackEntry Entry in CallStack.CallStackEntries.Take( 64 ) )
+						{
+							FunctionCall CurrentFunctionCall = new FunctionCall();
+
+							if( FunctionCalls.Where( f => f.Call == Entry.FunctionName ).Count() > 0 )
+							{
+								CurrentFunctionCall = FunctionCalls.Where( f => f.Call == Entry.FunctionName ).First();
+							}
+							else
+							{
+								CurrentFunctionCall = new FunctionCall();
+								CurrentFunctionCall.Call = Entry.FunctionName;
+								FunctionCalls.InsertOnSubmit( CurrentFunctionCall );
+							}
+
+							//CrashRepository.Context.SubmitChanges();
+
+							Pattern.Add( CurrentFunctionCall.Id.ToString() );
+						}
+
+						//CrashInstance.Pattern = "+";
+						CrashInstance.Pattern = string.Join( "+", Pattern );
+						// We need something like this +1+2+3+5+ for searching for exact pattern like +5+
+						//CrashInstance.Pattern += "+";
+
+						CrashRepository.Context.SubmitChanges();
+					}
+
+				}
+				catch( Exception Ex )
+				{
+					FLogger.WriteEvent( "Exception in BuildPattern: " + Ex.ToString() );
+				}
+			}
+		}
+
+		/// <summary>
+		/// Retrieve all Buggs matching the search criteria.
+		/// </summary>
+		/// <param name="FormData">The incoming form of search criteria from the client.</param>
+		/// <returns>A view to display the filtered Buggs.</returns>
+		public ReportsViewModel GetResults( FormHelper FormData )
+		{
+			using( FAutoScopedLogTimer LogTimer = new FAutoScopedLogTimer( this.GetType().ToString() ) )
+			{
+				string AnonumousGroup = "Anonymous";
+				List<String> Users = CrashRepository.GetUsersForGroup( AnonumousGroup );
+				int AnonymousGroupID = BuggRepository.GetIdFromUserGroup( AnonumousGroup );
+				HashSet<int> AnonumousIDs = BuggRepository.GetUserIdsFromUserGroup( AnonumousGroup );
+				int AnonymousID = AnonumousIDs.First();
+				HashSet<string> UserNamesForUserGroup = BuggRepository.GetUserNamesFromUserGroups( AnonumousGroup );
+
+				//FormData.DateFrom = DateTime.Now.AddDays( -1 );
+
+				var Crashes = CrashRepository
+					.FilterByDate( CrashRepository.ListAll(), FormData.DateFrom, FormData.DateTo.AddDays( 1 ) )
+					// Only crashes and asserts
+					.Where( Crash => Crash.CrashType == 1 || Crash.CrashType == 2 )
+					// Only anonymous user
+					.Where( Crash => Crash.UserNameId.Value == AnonymousID )
+					.Select( Crash => new
+					{
+						ID = Crash.Id,
+						TimeOfCrash = Crash.TimeOfCrash.Value,
+						//UserID = Crash.UserNameId.Value, 
+						BuildVersion = Crash.BuildVersion,
+						JIRA = Crash.TTPID,
+						Platform = Crash.PlatformName,
+						FixCL = Crash.FixedChangeList,
+						BuiltFromCL = Crash.ChangeListVersion,
+						Pattern = Crash.Pattern,
+						MachineID = Crash.ComputerName,
+					} )
+					.ToList();
+				int NumCrashes = Crashes.Count;
+
+				/*
+				// Build patterns for crashes where patters is null.
+				var CrashesWithoutPattern = CrashRepository
+					.FilterByDate( CrashRepository.ListAll(), FormData.DateFrom, FormData.DateTo.AddDays( 1 ) )
+					// Only crashes and asserts
+					.Where( Crash => Crash.Pattern == null || Crash.Pattern == "" )
+					.Select( Crash => Crash )
+					.ToList();
+
+				foreach( var Crash in CrashesWithoutPattern )
+				{
+					////BuildPattern( Crash );
+				}
+				*/
+
+				// Total # of ALL (Anonymous) crashes in timeframe
+				int TotalAnonymousCrashes = NumCrashes;
+
+				// Total # of UNIQUE (Anonymous) crashes in timeframe
+				HashSet<string> UniquePatterns = new HashSet<string>();
+				HashSet<string> UniqueMachines = new HashSet<string>();
+				Dictionary<string, int> PatternToCount = new Dictionary<string, int>();
+
+				//List<int> CrashesWithoutPattern = new List<int>();
+				//List<DateTime> CrashesWithoutPatternDT = new List<DateTime>();
+
+				foreach( var Crash in Crashes )
+				{
+					if( string.IsNullOrEmpty( Crash.Pattern ) )
+					{
+						//CrashesWithoutPattern.Add( Crash.ID );
+						//CrashesWithoutPatternDT.Add( Crash.TimeOfCrash );
+						continue;
+					}
+
+					UniquePatterns.Add( Crash.Pattern );
+					UniqueMachines.Add( Crash.MachineID );
+
+					bool bAdd = !PatternToCount.ContainsKey( Crash.Pattern );
+					if( bAdd )
+					{
+						PatternToCount.Add( Crash.Pattern, 1 );
+					}
+					else
+					{
+						PatternToCount[Crash.Pattern]++;
+					}
+				}
+				var PatternToCountOrdered = PatternToCount.OrderByDescending( X => X.Value ).ToList();
+				// The 100 top.
+				var PatternAndCount = PatternToCountOrdered.Take( 100 ).ToDictionary( X => X.Key, Y => Y.Value );
+
+				int TotalUniqueAnonymousCrashes = UniquePatterns.Count;
+
+				// Total # of AFFECTED USERS (Anonymous) in timeframe
+				int TotalAffectedUsers = UniqueMachines.Count;
+
+				var RealBuggs = BuggRepository.GetDataContext().Buggs.Where( Bugg => PatternAndCount.Keys.Contains( Bugg.Pattern ) ).ToList();
+
+				List<Bugg> Buggs = new List<Bugg>( 100 );
+				foreach( var Top in PatternAndCount )
+				{
+					Bugg NewBugg = new Bugg();
+
+					Bugg RealBugg = RealBuggs.Where( X => X.Pattern == Top.Key ).First();
+					if( RealBugg != null )
+					{
+						using( FAutoScopedLogTimer TopTimer = new FAutoScopedLogTimer( string.Format("{0}:{1}", Buggs.Count+1, RealBugg.Id ) ) )
+						{
+							// Index												// number
+							NewBugg.Id = RealBugg.Id;								// CrashGroup URL (a link to the Bugg)
+							NewBugg.TTPID = RealBugg.TTPID;							// JIRA
+							NewBugg.FixedChangeList = RealBugg.FixedChangeList;		// FixCL
+							NewBugg.NumberOfCrashes = Top.Value;					// # Occurrences
+
+							//NewBugg.BuildVersion = 
+
+							var CrashesForBugg = Crashes.Where( Crash => Crash.Pattern == Top.Key ).ToList();
+
+							NewBugg.AffectedBuilds = new SortedSet<string>();
+
+							HashSet<string> MachineIds = new HashSet<string>();
+							foreach( var Crash in CrashesForBugg )
+							{
+								// Only add machine if the number has 32 characters
+								if( Crash.MachineID.Length == 32 )
+								{
+									MachineIds.Add( Crash.MachineID );
+								}
+
+								// Ignore bad build versions.
+								if( Crash.BuildVersion.StartsWith( "4." ) )
+								{
+									NewBugg.AffectedBuilds.Add( Crash.BuildVersion );
+								}
+							}
+							NewBugg.BuildVersion = NewBugg.AffectedBuilds.Last();	// Latest Version Affected
+							NewBugg.NumberOfUniqueMachines = MachineIds.Count;		// # Affected Users
+							string LatestCLAffected = CrashesForBugg.				// CL of the latest build
+								Where( Crash => Crash.BuildVersion == NewBugg.BuildVersion ).
+								Max( Crash => Crash.BuiltFromCL );
+
+							int ILatestCLAffected = -1;
+							int.TryParse( LatestCLAffected, out ILatestCLAffected );
+							NewBugg.LatestCLAffected = ILatestCLAffected;			// Latest CL Affected
+
+							string LatestOSAffected = CrashesForBugg.OrderByDescending( Crash => Crash.TimeOfCrash ).First().Platform;
+
+							NewBugg.LatestOSAffected = LatestOSAffected;			// Latest Environment Affected
+
+							NewBugg.TimeOfFirstCrash = RealBugg.TimeOfFirstCrash;	// First Crash Timestamp
+
+							Buggs.Add( NewBugg );
+						}
+					}
+					else
+					{
+						FLogger.WriteEvent( "Bugg for pattern " + Top.Key + " not found" );
+					}
+				}
+
+				return new ReportsViewModel
+				{
+					Buggs = Buggs,
+					/*Crashes = Crashes,*/
+					DateFrom = (long)( FormData.DateFrom - CrashesViewModel.Epoch ).TotalMilliseconds,
+					DateTo = (long)( FormData.DateTo - CrashesViewModel.Epoch ).TotalMilliseconds,
+					TotalAffectedUsers = TotalAffectedUsers,
+					TotalAnonymousCrashes = TotalAnonymousCrashes,
+					TotalUniqueAnonymousCrashes = TotalUniqueAnonymousCrashes
+				};
+			}
+		}
+
+		/// <summary>
+		/// 
+		/// </summary>
+		/// <param name="ReportsForm"></param>
+		/// <returns></returns>
+		public ActionResult Index( FormCollection ReportsForm )
+		{
+			using( FAutoScopedLogTimer LogTimer = new FAutoScopedLogTimer( this.GetType().ToString() ) )
+			{
+				FormHelper FormData = new FormHelper( Request, ReportsForm, "JustReport" );
+				ReportsViewModel Results = GetResults( FormData );
+				return View( "Index", Results );
+			}
+		}
+	}
+}
\ No newline at end of file
diff --git a/Engine/Source/Programs/CrashReporter/CrashReportWebSite/CrashReportWebSite.csproj b/Engine/Source/Programs/CrashReporter/CrashReportWebSite/CrashReportWebSite.csproj
--- a/Engine/Source/Programs/CrashReporter/CrashReportWebSite/CrashReportWebSite.csproj
+++ b/Engine/Source/Programs/CrashReporter/CrashReportWebSite/CrashReportWebSite.csproj
@@ -86,6 +86,7 @@
     </Compile>
     <Compile Include="Controllers\DashboardController.cs" />
     <Compile Include="Controllers\ErrorController.cs" />
+    <Compile Include="Controllers\ReportsController.cs" />
     <Compile Include="Controllers\UsersController.cs" />
     <Compile Include="Models\BuggsViewModel.cs" />
     <Compile Include="Models\CachedDataService.cs" />
@@ -103,6 +104,7 @@
     <Compile Include="Models\CrashViewModel.cs" />
     <Compile Include="Models\FormHelper.cs" />
     <Compile Include="Models\PerformanceTimers.cs" />
+    <Compile Include="Models\ReportsViewModel.cs" />
     <Compile Include="Models\UsersViewModel.cs" />
     <Compile Include="Properties\Settings.Designer.cs">
       <AutoGen>True</AutoGen>
@@ -200,6 +202,8 @@
     <Content Include="Views\Crashes\Index.aspx" />
     <Content Include="Views\Crashes\Show.aspx" />
     <Content Include="Views\Crashes\ViewCrash.ascx" />
+    <Content Include="Views\Reports\Index.aspx" />
+    <Content Include="Views\Reports\ViewReports.ascx" />
     <Content Include="Views\Users\Index.aspx" />
     <Content Include="Web.config">
       <SubType>Designer</SubType>
diff --git a/Engine/Source/Programs/CrashReporter/CrashReportWebSite/Models/BuggRepository.cs b/Engine/Source/Programs/CrashReporter/CrashReportWebSite/Models/BuggRepository.cs
--- a/Engine/Source/Programs/CrashReporter/CrashReportWebSite/Models/BuggRepository.cs
+++ b/Engine/Source/Programs/CrashReporter/CrashReportWebSite/Models/BuggRepository.cs
@@ -153,6 +153,7 @@ namespace Tools.CrashReporter.CrashReportWebSite.Models
 
 		/// <summary>
 		/// Associate a set of crashes and their owners with a Bugg.
+		/// NOT USED, FIX!
 		/// </summary>
 		/// <param name="Bugg">The Bugg to get the data.</param>
 		/// <param name="Crashes">The set of crashes to add to the Bugg.</param>
@@ -621,7 +622,6 @@ namespace Tools.CrashReporter.CrashReportWebSite.Models
 				{
 					// Get the group id and grab all buggs for the specified group.
 					int NumResults = Results.Count();
-					HashSet<int> ResultIds = new HashSet<int>( Results.Select( x => x.Id ) );
 					HashSet<string> UserNamesForUserGroup = GetUserNamesFromUserGroups( UserGroup );
 
 
diff --git a/Engine/Source/Programs/CrashReporter/CrashReportWebSite/Models/CrashReport.cs b/Engine/Source/Programs/CrashReporter/CrashReportWebSite/Models/CrashReport.cs
--- a/Engine/Source/Programs/CrashReporter/CrashReportWebSite/Models/CrashReport.cs
+++ b/Engine/Source/Programs/CrashReporter/CrashReportWebSite/Models/CrashReport.cs
@@ -73,6 +73,12 @@ namespace Tools.CrashReporter.CrashReportWebSite.Models
 		/// <summary></summary>
 		public int NumberOfUniqueMachines { get; set; }
 
+		/// <summary></summary>
+		public int LatestCLAffected { get; set; }
+
+		/// <summary></summary>
+		public string LatestOSAffected { get; set; }
+
 		/// <summary></summary>
 		public SortedSet<string> AffectedBuilds { get; set; }
 
diff --git a/Engine/Source/Programs/CrashReporter/CrashReportWebSite/Models/PerformanceTimers.cs b/Engine/Source/Programs/CrashReporter/CrashReportWebSite/Models/PerformanceTimers.cs
--- a/Engine/Source/Programs/CrashReporter/CrashReportWebSite/Models/PerformanceTimers.cs
+++ b/Engine/Source/Programs/CrashReporter/CrashReportWebSite/Models/PerformanceTimers.cs
@@ -12,7 +12,23 @@ using System.Threading;
 public class FLogger
 {
 	/// <summary>Current log file to write debug progress info to</summary>
-	static public LogWriter Log = new LogWriter( "CrashReportWebSite-" + GetCleanUserName(), "f:/CrashReportWebsiteLogs" );
+	static public LogWriter Log = CreateSafeLogWriter();
+
+	/// <summary>
+	/// 
+	/// </summary>
+	/// <returns></returns>
+	static LogWriter CreateSafeLogWriter()
+	{
+		bool bHasDDrive = Directory.Exists( "D:" );
+		bool bHasFDrive = Directory.Exists( "F:" );
+		
+		string LogsPath = bHasFDrive ? "F:/" : (bHasDDrive ? "D:/" : "C:/");
+		LogsPath += "CrashReportWebsiteLogs";
+
+		LogWriter Log = new LogWriter( "CrashReportWebSite-" + GetCleanUserName(), LogsPath );
+		return Log;
+	}
 
 	/// <summary> Returns the user name valid as filename. </summary>
 	static public string GetCleanUserName()
@@ -33,6 +49,17 @@ public class FLogger
 		}
 	}
 
+	/// <summary>
+	/// Writes an exception status to the log file.
+	/// </summary>
+	static public void WriteException( string Message )
+	{
+		if( Message != null && Message.Length > 2 )
+		{
+			Log.Print( "[EXCEPT] " + Message );
+		}
+	}
+
 	/// <summary>
 	/// Writes a performance info to the log file.
 	/// </summary>
diff --git a/Engine/Source/Programs/CrashReporter/CrashReportWebSite/Models/ReportsViewModel.cs b/Engine/Source/Programs/CrashReporter/CrashReportWebSite/Models/ReportsViewModel.cs
new file mode 100644
--- /dev/null
+++ b/Engine/Source/Programs/CrashReporter/CrashReportWebSite/Models/ReportsViewModel.cs
@@ -0,0 +1,45 @@
+﻿// Copyright 1998-2015 Epic Games, Inc. All Rights Reserved.
+
+using System;
+using System.Collections.Generic;
+using System.Linq;
+using System.Web.Mvc;
+
+namespace Tools.CrashReporter.CrashReportWebSite.Models
+{
+	/// <summary>
+	/// The view model for the crash summary page.
+	/// </summary>
+	public class ReportsViewModel
+	{
+		/// <summary>A container of sorted Buggs.</summary>
+		public List<Bugg> Buggs { get; set; }
+
+		/// <summary>A container of sorted Crashes.</summary>
+		public List<Crash> Crashes { get; set; }
+
+		/// <summary>The date of the earliest crash in a Bugg.</summary>
+		public long DateFrom { get; set; }
+
+		/// <summary>The date of the most recent crash in a Bugg.</summary>
+		public long DateTo { get; set; }
+
+		/// <summary>
+		/// 
+		/// </summary>
+		public int TotalAnonymousCrashes { get; set; }
+
+		/// <summary>
+		///		
+		/// </summary>
+		public int TotalUniqueAnonymousCrashes { get; set; }
+
+		/// <summary>
+		/// 
+		/// </summary>
+		public int TotalAffectedUsers { get; set; }
+
+		/// <summary>User input from the client.</summary>
+		public FormCollection FormCollection { get; set; }
+	}
+}
\ No newline at end of file
diff --git a/Engine/Source/Programs/CrashReporter/CrashReportWebSite/Views/Reports/Index.aspx b/Engine/Source/Programs/CrashReporter/CrashReportWebSite/Views/Reports/Index.aspx
new file mode 100644
--- /dev/null
+++ b/Engine/Source/Programs/CrashReporter/CrashReportWebSite/Views/Reports/Index.aspx
@@ -0,0 +1,70 @@
+﻿<%-- // Copyright 1998-2015 Epic Games, Inc. All Rights Reserved. --%>
+
+<%@ Page Title="" Language="C#" MasterPageFile="~/Views/Shared/Site.Master" Inherits="System.Web.Mvc.ViewPage<ReportsViewModel>" %>
+<%@ Import Namespace="Tools.CrashReporter.CrashReportWebSite.Models" %>
+
+<asp:Content ID="StyleSheet" ContentPlaceHolderID="CssContent" runat="server">
+	<link href="../../Content/Site.css" rel="stylesheet" type="text/css" />
+</asp:Content>
+
+<asp:Content ID="TitleContent" ContentPlaceHolderID="TitleContent" runat="server">
+	Crash Reports
+</asp:Content>
+
+<asp:Content ID="ScriptContent" ContentPlaceHolderID="ScriptContent" runat="server" >
+	<script type="text/javascript">
+
+		$(function ()
+		{
+			$("#dateFromVisible")
+				.datepicker({ maxDate: '+0D' })
+				.datepicker('setDate', new Date(parseInt($('#dateFrom').val())));
+
+			$("#dateToVisible")
+				.datepicker({ maxDate: '+0D' })
+				.datepicker('setDate', new Date(parseInt($('#dateTo').val())));
+
+		});
+
+		$.datepicker.setDefaults({
+			onSelect: function ()
+			{
+				$("#dateFrom").val($("#dateFromVisible").datepicker('getDate').getTime());
+				$("#dateTo").val($("#dateToVisible").datepicker('getDate').getTime());
+				$("#ReporsForm").submit();
+			}
+		});
+
+	</script>
+</asp:Content>
+
+<asp:Content ID="AboveMainContent" ContentPlaceHolderID="AboveMainContent" runat="server" >
+
+<div id='SearchForm' style="clear:both;">
+<%using( Html.BeginForm( "", "Reports", FormMethod.Get, new { id = "ReporsForm" } ) )
+{ %>
+	<script>$.datepicker.setDefaults($.datepicker.regional['']);</script>
+
+	<span style="margin-left: 10px; font-weight: bold;">Filter by Date </span>
+
+	<span>From: 
+	<input id="dateFromVisible" type="text" class="date" autocomplete="OFF" /></span>
+	<input id="dateFrom" name="dateFrom" type="hidden" value="<%=Model.DateFrom %>" autocomplete="OFF" />
+
+	<span>To: 
+	<input id="dateToVisible" type="text" class="date" autocomplete="OFF" /></span>
+	<input id="dateTo" name="dateTo" type="hidden" value="<%=Model.DateTo %>" autocomplete="OFF" />
+
+	<%} %>
+
+</div>
+</asp:Content>
+
+<asp:Content ID="ReportsContent" ContentPlaceHolderID="ReportsContent" runat="server">
+	<div>
+		<div>
+			<% Html.RenderPartial("/Views/Reports/ViewReports.ascx"); %>
+		</div>
+	</div>
+
+</asp:Content>
diff --git a/Engine/Source/Programs/CrashReporter/CrashReportWebSite/Views/Reports/ViewReports.ascx b/Engine/Source/Programs/CrashReporter/CrashReportWebSite/Views/Reports/ViewReports.ascx
new file mode 100644
--- /dev/null
+++ b/Engine/Source/Programs/CrashReporter/CrashReportWebSite/Views/Reports/ViewReports.ascx
@@ -0,0 +1,39 @@
+<%-- // Copyright 1998-2015 Epic Games, Inc. All Rights Reserved. --%>
+
+<%@ Control Language="C#" Inherits="System.Web.Mvc.ViewUserControl<BuggsViewModel>" %>
+<%@ Import Namespace="Tools.CrashReporter.CrashReportWebSite.Models" %>
+
+<style>
+
+table, tr, td, th
+{
+	border: solid;
+	border-width: 1px;
+
+	color: white;
+
+	font-family: 'Segoe UI', Tahoma, Geneva, Verdana, sans-serif;
+	font-size: 12px;
+}
+
+</style>
+
+<table>
+	<tr>
+		<td>
+			TotalAnonymousCrashes:&nbsp;<%Model.TotalAnonymousCrashes %>
+		</td>
+	</tr>
+
+	<tr>
+		<td>
+			TotalUniqueAnonymousCrashes:&nbsp;<%Model.TotalUniqueAnonymousCrashes %>
+		</td>
+	</tr>
+
+	<tr>
+		<td>
+			TotalAffectedUsers:&nbsp;<%Model.TotalAffectedUsers %>
+		</td>
+	</tr>
+</table>
diff --git a/Engine/Source/Programs/CrashReporter/CrashReportWebSite/Views/Shared/Site.Master b/Engine/Source/Programs/CrashReporter/CrashReportWebSite/Views/Shared/Site.Master
--- a/Engine/Source/Programs/CrashReporter/CrashReportWebSite/Views/Shared/Site.Master
+++ b/Engine/Source/Programs/CrashReporter/CrashReportWebSite/Views/Shared/Site.Master
@@ -31,6 +31,8 @@
 				<br />
 			</div>
 			<div id="utils">
+				&nbsp;|&nbsp;
+				<span><a href="/Reports/Index/">Reports</a></span>
 				&nbsp;|&nbsp;
 				<span><a href="/Dashboard/Index/">Dashboard</a></span>
 				&nbsp;|&nbsp;
@@ -46,7 +48,11 @@
 				<div id="mainLeft">
 					<asp:ContentPlaceHolder ID="MainContent" runat="server" />
 				</div>
-			</div>
+			</div>		
+		</div>
+
+		<div style="clear:both;">
+			<asp:ContentPlaceHolder ID="ReportsContent" runat="server" />
 		</div>
 
 		<div id="footer" style="clear:both;">
