Core - UE-9320 Mac crashes are not uploading to Crash Reporter

diff --git a/Engine/Source/Programs/CrashReporter/CrashReportClient/Private/CrashReportClient.cpp b/Engine/Source/Programs/CrashReporter/CrashReportClient/Private/CrashReportClient.cpp
--- a/Engine/Source/Programs/CrashReporter/CrashReportClient/Private/CrashReportClient.cpp
+++ b/Engine/Source/Programs/CrashReporter/CrashReportClient/Private/CrashReportClient.cpp
@@ -40,6 +40,7 @@ FCrashReportClient::FCrashReportClient(const FPlatformErrorReport& InErrorReport
 	, bBeginUploadCalled(false)
 	, bShouldWindowBeHidden(false)
 	, bAllowToBeContacted(true)
+	, bSendData(false)
 {
 
 	if (!ErrorReport.TryReadDiagnosticsFile(DiagnosticText) && !FParse::Param(FCommandLine::Get(), TEXT("no-local-diagnosis")))
@@ -65,6 +66,7 @@ FCrashReportClient::~FCrashReportClient()
 
 FReply FCrashReportClient::Submit()
 {
+	bSendData = true;
 	StoreCommentAndUpload();
 	bShouldWindowBeHidden = true;
 	return FReply::Handled();
@@ -94,6 +96,9 @@ void FCrashReportClient::UserCommentChanged(const FText& Comment, ETextCommit::T
 
 void FCrashReportClient::RequestCloseWindow(const TSharedRef<SWindow>& Window)
 {
+	// Don't send the data.
+	bSendData = false;
+
 	// We may still processing minidump etc. so start the main ticker.
 	StartTicker();
 	bShouldWindowBeHidden = true;
@@ -133,19 +138,23 @@ bool FCrashReportClient::Tick(float UnusedDeltaTime)
 	{
 		return true;
 	}
-	else if( !bBeginUploadCalled )
+	
+	if( bSendData )
 	{
-		// Can be called only when we have all files.
-		Uploader.BeginUpload(ErrorReport);
-		bBeginUploadCalled = true;
-	}
-
-	// IsWorkDone will always return true here (since uploader can't finish until the diagnosis has been sent), but it
-	//  has the side effect of joining the worker thread.
-	if( !Uploader.IsFinished() )
-	{
-		// More ticks, please
-		return true;
+		if( !bBeginUploadCalled )
+		{
+			// Can be called only when we have all files.
+			Uploader.BeginUpload( ErrorReport );
+			bBeginUploadCalled = true;
+		}
+
+		// IsWorkDone will always return true here (since uploader can't finish until the diagnosis has been sent), but it
+		//  has the side effect of joining the worker thread.
+		if( !Uploader.IsFinished() )
+		{
+			// More ticks, please
+			return true;
+		}
 	}
 
 	FPlatformMisc::RequestExit(false);
diff --git a/Engine/Source/Programs/CrashReporter/CrashReportClient/Private/CrashReportClient.h b/Engine/Source/Programs/CrashReporter/CrashReportClient/Private/CrashReportClient.h
--- a/Engine/Source/Programs/CrashReporter/CrashReportClient/Private/CrashReportClient.h
+++ b/Engine/Source/Programs/CrashReporter/CrashReportClient/Private/CrashReportClient.h
@@ -27,6 +27,11 @@ public:
 	 */
 	void DoWork();
 
+	FORCEINLINE TStatId GetStatId() const
+	{
+		return TStatId();
+	}
+
 	/** 
 	 * @return The name to display in external event viewers
 	 */
@@ -154,6 +159,9 @@ private:
 	/** Whether the user allowed us to be contacted. */
 	bool bAllowToBeContacted;
 
+	/** Whether we send the data. */
+	bool bSendData;
+
 };
 
 #endif // !CRASH_REPORT_UNATTENDED_ONLY
diff --git a/Engine/Source/Programs/CrashReporter/CrashReportClient/Private/CrashUpload.cpp b/Engine/Source/Programs/CrashReporter/CrashReportClient/Private/CrashUpload.cpp
--- a/Engine/Source/Programs/CrashReporter/CrashReportClient/Private/CrashUpload.cpp
+++ b/Engine/Source/Programs/CrashReporter/CrashReportClient/Private/CrashUpload.cpp
@@ -50,7 +50,7 @@ bool FCrashUpload::PingTimeout(float DeltaTime)
 void FCrashUpload::BeginUpload(const FPlatformErrorReport& PlatformErrorReport)
 {
 	ErrorReport = PlatformErrorReport;
-	PendingFiles = ErrorReport.GetFilesToUpload();
+	PendingFiles += ErrorReport.GetFilesToUpload();
 	UE_LOG(CrashReportClientLog, Log, TEXT("Got %d pending files to upload from '%s'"), PendingFiles.Num(), *ErrorReport.GetReportDirectoryLeafName());
 
 	PauseState = EUploadState::Finished;
@@ -222,24 +222,24 @@ void FCrashUpload::CompressAndSendData()
 	delete CompressedDataRaw;
 	CompressedDataRaw = nullptr;
 
- 	// Set up request for upload
- 	UE_LOG(CrashReportClientLog, Log, TEXT("Sending HTTP request (posting file)"));
- 	auto Request = CreateHttpRequest();
- 	Request->SetVerb(TEXT("POST"));
- 	Request->SetHeader(TEXT("Content-Type"), TEXT("application/octet-stream"));
- 	Request->SetURL(UrlPrefix / TEXT("UploadReportFile"));
- 	Request->SetContent(CompressedData);
- 	Request->SetHeader(TEXT("DirectoryName"), *ErrorReport.GetReportDirectoryLeafName());
- 	Request->SetHeader(TEXT("FileName"), Filename);
- 	Request->SetHeader(TEXT("FileLength"), TTypeToString<int32>::ToString(CompressedData.Num()) );
+	// Set up request for upload
+	UE_LOG(CrashReportClientLog, Log, TEXT("Sending HTTP request (posting file)"));
+	auto Request = CreateHttpRequest();
+	Request->SetVerb(TEXT("POST"));
+	Request->SetHeader(TEXT("Content-Type"), TEXT("application/octet-stream"));
+	Request->SetURL(UrlPrefix / TEXT("UploadReportFile"));
+	Request->SetContent(CompressedData);
+	Request->SetHeader(TEXT("DirectoryName"), *ErrorReport.GetReportDirectoryLeafName());
+	Request->SetHeader(TEXT("FileName"), Filename);
+	Request->SetHeader(TEXT("FileLength"), TTypeToString<int32>::ToString(CompressedData.Num()) );
 	Request->SetHeader(TEXT("CompressedSize"), TTypeToString<int32>::ToString(CompressedSize) );
 	Request->SetHeader(TEXT("UncompressedSize"), TTypeToString<int32>::ToString(UncompressedSize) );
 	Request->SetHeader(TEXT("NumberOfFiles"), TTypeToString<int32>::ToString(CurrentFileIndex) );
 
- 	if (Request->ProcessRequest())
- 	{
- 		return;
- 	}
+	if (Request->ProcessRequest())
+	{
+		return;
+	}
 	else
 	{
 		UE_LOG(CrashReportClientLog, Warning, TEXT("Failed to send file upload request"));
