CrashReportClient - Tweaks to the UI (fixed scrollbars not visible)

diff --git a/Engine/Source/Programs/CrashReporter/CrashReportClient/Private/CrashReportClientStyle.cpp b/Engine/Source/Programs/CrashReporter/CrashReportClient/Private/CrashReportClientStyle.cpp
--- a/Engine/Source/Programs/CrashReporter/CrashReportClient/Private/CrashReportClientStyle.cpp
+++ b/Engine/Source/Programs/CrashReporter/CrashReportClient/Private/CrashReportClientStyle.cpp
@@ -25,7 +25,7 @@ void FCrashReportClientStyle::Shutdown()
 }
 
 #define TTF_FONT(RelativePath, ...) FSlateFontInfo(ContentFromEngine(TEXT(RelativePath), TEXT(".ttf")), __VA_ARGS__)
-#define IMAGE_BRUSH( RelativePath, ... ) FSlateImageBrush(Style.RootToContentDir(RelativePath, TEXT(".png")), __VA_ARGS__)
+#define IMAGE_BRUSH( RelativePath, ... ) FSlateImageBrush(ContentFromEngine(RelativePath, TEXT(".png")), __VA_ARGS__)
 
 namespace
 {
@@ -42,7 +42,7 @@ TSharedRef< FSlateStyleSet > FCrashReportClientStyle::Create()
 	FSlateStyleSet& Style = StyleRef.Get();
 
 	const FTextBlockStyle DefaultText = FTextBlockStyle()
-		.SetFont(TTF_FONT("Fonts/Roboto-Black", 10))
+		.SetFont(TTF_FONT("Fonts/Roboto-Bold", 10))
 		.SetColorAndOpacity(FSlateColor::UseForeground())
 		.SetShadowOffset(FVector2D::ZeroVector)
 		.SetShadowColorAndOpacity(FLinearColor::Black);
@@ -65,11 +65,13 @@ TSharedRef< FSlateStyleSet > FCrashReportClientStyle::Create()
 	FSlateBrush* GenericWhiteBox = new IMAGE_BRUSH( "Old/White", Icon16x16 );
 
 	// SEditableTextBox defaults...
+	const FScrollBarStyle& ScrollBarStyle = FCoreStyle::Get().GetWidgetStyle<FScrollBarStyle>( "ScrollBar" );
 	const FEditableTextBoxStyle NormalEditableTextBoxStyle = FEditableTextBoxStyle()
 		.SetBackgroundImageNormal( *GenericWhiteBox )
 		.SetBackgroundImageHovered( *GenericWhiteBox )
 		.SetBackgroundImageFocused( *GenericWhiteBox )
-		.SetBackgroundImageReadOnly( *GenericWhiteBox );
+		.SetBackgroundImageReadOnly( *GenericWhiteBox )
+		.SetScrollBarStyle( ScrollBarStyle );
 	{
 		Style.Set( "NormalEditableTextBox", NormalEditableTextBoxStyle );
 	}
diff --git a/Engine/Source/Programs/CrashReporter/CrashReportClient/Private/SCrashReportClient.cpp b/Engine/Source/Programs/CrashReporter/CrashReportClient/Private/SCrashReportClient.cpp
--- a/Engine/Source/Programs/CrashReporter/CrashReportClient/Private/SCrashReportClient.cpp
+++ b/Engine/Source/Programs/CrashReporter/CrashReportClient/Private/SCrashReportClient.cpp
@@ -24,6 +24,10 @@ void SCrashReportClient::Construct(const FArguments& InArgs, TSharedRef<FCrashRe
 {
 	CrashReportClient = Client;
 
+	TSharedPtr<SScrollBar> HiddenScrollbar;
+	SAssignNew( HiddenScrollbar, SScrollBar );
+	HiddenScrollbar->SetVisibility( EVisibility::Collapsed );
+
 	auto CrashedAppName = CrashReportClient->GetCrashedAppName();
 
 	// Set the text displaying the name of the crashed app, if available
@@ -152,6 +156,7 @@ void SCrashReportClient::Construct(const FArguments& InArgs, TSharedRef<FCrashRe
 					.BackgroundColor( FSlateColor( FLinearColor::Black ) )
 					.ForegroundColor( FSlateColor( FLinearColor::White * 0.8f ) )
 					.Text( Client, &FCrashReportClient::GetDiagnosticText )
+					.HScrollBar( HiddenScrollbar )
 				]
 
 
