Various crash reporter improvements: made scrollbars visible, added splitter between editable field and crash details, added "Send and Restart" button to restart the crashed application.

diff --git a/Engine/Source/Programs/CrashReporter/CrashReportClient/Private/CrashDescription.cpp b/Engine/Source/Programs/CrashReporter/CrashReportClient/Private/CrashDescription.cpp
--- a/Engine/Source/Programs/CrashReporter/CrashReportClient/Private/CrashDescription.cpp
+++ b/Engine/Source/Programs/CrashReporter/CrashReportClient/Private/CrashDescription.cpp
@@ -78,7 +78,7 @@ FCrashDescription::FCrashDescription( FString WERXMLFilepath ) :
 				EngineVersionComponents++;
 			}
 
-			/*const FXmlNode* Parameter8Node = ProblemSignaturesNode->FindChildNode( TEXT( "Parameter8" ) );
+			const FXmlNode* Parameter8Node = ProblemSignaturesNode->FindChildNode( TEXT( "Parameter8" ) );
 			if( Parameter8Node )
 			{
 				const FString Parameter8Value = Parameter8Node->GetContent();
@@ -95,7 +95,7 @@ FCrashDescription::FCrashDescription( FString WERXMLFilepath ) :
 				{
 					ErrorMessage.Add( ParsedParameters8[2] );
 				}
-			}*/
+			}
 
 			const FXmlNode* Parameter9Node = ProblemSignaturesNode->FindChildNode( TEXT( "Parameter9" ) );
 			if( Parameter9Node )
diff --git a/Engine/Source/Programs/CrashReporter/CrashReportClient/Private/CrashReportClient.cpp b/Engine/Source/Programs/CrashReporter/CrashReportClient/Private/CrashReportClient.cpp
--- a/Engine/Source/Programs/CrashReporter/CrashReportClient/Private/CrashReportClient.cpp
+++ b/Engine/Source/Programs/CrashReporter/CrashReportClient/Private/CrashReportClient.cpp
@@ -72,6 +72,19 @@ FReply FCrashReportClient::Submit()
 	return FReply::Handled();
 }
 
+FReply FCrashReportClient::SubmitAndRestart()
+{
+	Submit();
+
+	const FString CrashedAppPath = ErrorReport.FindCrashedAppPath();
+	const FCrashDescription& CrashDescription = GetCrashDescription();
+	const FString CommandLineArguments = CrashDescription.CommandLine;
+
+	FPlatformProcess::CreateProc(*CrashedAppPath, *CommandLineArguments, true, false, false, NULL, 0, NULL, NULL);
+
+	return FReply::Handled();
+}
+
 FReply FCrashReportClient::CopyCallstack()
 {
 	FPlatformMisc::ClipboardCopy(*DiagnosticText.ToString());
diff --git a/Engine/Source/Programs/CrashReporter/CrashReportClient/Private/CrashReportClient.h b/Engine/Source/Programs/CrashReporter/CrashReportClient/Private/CrashReportClient.h
--- a/Engine/Source/Programs/CrashReporter/CrashReportClient/Private/CrashReportClient.h
+++ b/Engine/Source/Programs/CrashReporter/CrashReportClient/Private/CrashReportClient.h
@@ -64,6 +64,12 @@ public:
 	 */
 	FReply Submit();
 
+	/**
+	* Respond to the user pressing Submit and Restart
+	* @return Whether the request was handled
+	*/
+	FReply SubmitAndRestart();
+
 	/**
 	 * Respond to the user requesting the callstack to be copied to the clipboard
 	 * @return Whether the request was handled
diff --git a/Engine/Source/Programs/CrashReporter/CrashReportClient/Private/CrashReportClientStyle.cpp b/Engine/Source/Programs/CrashReporter/CrashReportClient/Private/CrashReportClientStyle.cpp
--- a/Engine/Source/Programs/CrashReporter/CrashReportClient/Private/CrashReportClientStyle.cpp
+++ b/Engine/Source/Programs/CrashReporter/CrashReportClient/Private/CrashReportClientStyle.cpp
@@ -25,7 +25,8 @@ void FCrashReportClientStyle::Shutdown()
 }
 
 #define TTF_FONT(RelativePath, ...) FSlateFontInfo(ContentFromEngine(TEXT(RelativePath), TEXT(".ttf")), __VA_ARGS__)
-#define IMAGE_BRUSH( RelativePath, ... ) FSlateImageBrush(ContentFromEngine(RelativePath, TEXT(".png")), __VA_ARGS__)
+#define IMAGE_BRUSH( RelativePath, ... ) FSlateImageBrush(ContentFromEngine(TEXT(RelativePath), TEXT(".png")), __VA_ARGS__)
+#define BOX_BRUSH( RelativePath, ... ) FSlateBoxBrush(ContentFromEngine(TEXT(RelativePath), TEXT(".png")), __VA_ARGS__)
 
 namespace
 {
@@ -64,14 +65,23 @@ TSharedRef< FSlateStyleSet > FCrashReportClientStyle::Create()
 	const FVector2D Icon16x16( 16.0f, 16.0f );
 	FSlateBrush* GenericWhiteBox = new IMAGE_BRUSH( "Old/White", Icon16x16 );
 
+	// Scrollbar
+	const FScrollBarStyle ScrollBar = FScrollBarStyle()
+		.SetVerticalTopSlotImage(IMAGE_BRUSH("Common/Scrollbar_Background_Vertical", FVector2D(8, 8)))
+		.SetVerticalBottomSlotImage(IMAGE_BRUSH("Common/Scrollbar_Background_Vertical", FVector2D(8, 8)))
+		.SetHorizontalTopSlotImage(IMAGE_BRUSH("Common/Scrollbar_Background_Horizontal", FVector2D(8, 8)))
+		.SetHorizontalBottomSlotImage(IMAGE_BRUSH("Common/Scrollbar_Background_Horizontal", FVector2D(8, 8)))
+		.SetNormalThumbImage(BOX_BRUSH("Common/Scrollbar_Thumb", FMargin(4.f / 16.f)))
+		.SetDraggedThumbImage(BOX_BRUSH("Common/Scrollbar_Thumb", FMargin(4.f / 16.f)))
+		.SetHoveredThumbImage(BOX_BRUSH("Common/Scrollbar_Thumb", FMargin(4.f / 16.f)));
+
 	// SEditableTextBox defaults...
-	const FScrollBarStyle& ScrollBarStyle = FCoreStyle::Get().GetWidgetStyle<FScrollBarStyle>( "ScrollBar" );
 	const FEditableTextBoxStyle NormalEditableTextBoxStyle = FEditableTextBoxStyle()
-		.SetBackgroundImageNormal( *GenericWhiteBox )
-		.SetBackgroundImageHovered( *GenericWhiteBox )
-		.SetBackgroundImageFocused( *GenericWhiteBox )
-		.SetBackgroundImageReadOnly( *GenericWhiteBox )
-		.SetScrollBarStyle( ScrollBarStyle );
+		.SetBackgroundImageNormal(*GenericWhiteBox)
+		.SetBackgroundImageHovered(*GenericWhiteBox)
+		.SetBackgroundImageFocused(*GenericWhiteBox)
+		.SetBackgroundImageReadOnly(*GenericWhiteBox)
+		.SetScrollBarStyle(ScrollBar);
 	{
 		Style.Set( "NormalEditableTextBox", NormalEditableTextBoxStyle );
 	}
diff --git a/Engine/Source/Programs/CrashReporter/CrashReportClient/Private/GenericErrorReport.cpp b/Engine/Source/Programs/CrashReporter/CrashReportClient/Private/GenericErrorReport.cpp
--- a/Engine/Source/Programs/CrashReporter/CrashReportClient/Private/GenericErrorReport.cpp
+++ b/Engine/Source/Programs/CrashReporter/CrashReportClient/Private/GenericErrorReport.cpp
@@ -224,7 +224,18 @@ bool FGenericErrorReport::FindFirstReportFileWithExtension(FString& OutFilename,
 
 FString FGenericErrorReport::FindCrashedAppName() const
 {
-	UE_LOG( LogTemp, Warning, TEXT( "FGenericPlatformMemory::Init not implemented on this platform" ) );
-	return FString( TEXT( "GenericAppName" ) );
+	FString AppPath = FindCrashedAppPath();
+	if (!AppPath.IsEmpty())
+	{
+		return FPaths::GetCleanFilename(AppPath);
+	}
+
+	return AppPath;
+}
+
+FString FGenericErrorReport::FindCrashedAppPath() const
+{
+	UE_LOG( LogTemp, Warning, TEXT( "FGenericErrorReport::FindCrashedAppPath not implemented on this platform" ) );
+	return FString( TEXT( "GenericAppPath" ) );
 }
 
diff --git a/Engine/Source/Programs/CrashReporter/CrashReportClient/Private/GenericErrorReport.h b/Engine/Source/Programs/CrashReporter/CrashReportClient/Private/GenericErrorReport.h
--- a/Engine/Source/Programs/CrashReporter/CrashReportClient/Private/GenericErrorReport.h
+++ b/Engine/Source/Programs/CrashReporter/CrashReportClient/Private/GenericErrorReport.h
@@ -87,6 +87,11 @@ public:
 	 */
 	FString FindCrashedAppName() const;
 
+	/**
+	 * Get the full path of the crashed app from the report
+	 */
+	FString FindCrashedAppPath() const;
+
 	/**
 	 * Is there anything to upload?
 	 */
diff --git a/Engine/Source/Programs/CrashReporter/CrashReportClient/Private/Mac/MacErrorReport.cpp b/Engine/Source/Programs/CrashReporter/CrashReportClient/Private/Mac/MacErrorReport.cpp
--- a/Engine/Source/Programs/CrashReporter/CrashReportClient/Private/Mac/MacErrorReport.cpp
+++ b/Engine/Source/Programs/CrashReporter/CrashReportClient/Private/Mac/MacErrorReport.cpp
@@ -29,7 +29,7 @@ void FMacErrorReport::ShutDown()
 	CrashHelperModule->ShutdownModule();
 }
 
-FString FMacErrorReport::FindCrashedAppName() const
+FString FMacErrorReport::FindCrashedAppPath() const
 {
 	TArray<uint8> Data;
 	if(FFileHelper::LoadFileToArray(Data, *(ReportDirectory / TEXT("Report.wer"))))
@@ -53,7 +53,7 @@ FString FMacErrorReport::FindCrashedAppName() const
 			{
 				PathData = PathData.Left(LineEnd);
 			}
-			return FPaths::GetCleanFilename(PathData);
+			return PathData;
 		}
 	}
 	else
diff --git a/Engine/Source/Programs/CrashReporter/CrashReportClient/Private/Mac/MacErrorReport.h b/Engine/Source/Programs/CrashReporter/CrashReportClient/Private/Mac/MacErrorReport.h
--- a/Engine/Source/Programs/CrashReporter/CrashReportClient/Private/Mac/MacErrorReport.h
+++ b/Engine/Source/Programs/CrashReporter/CrashReportClient/Private/Mac/MacErrorReport.h
@@ -40,9 +40,9 @@ public:
 	FText DiagnoseReport() const;
 
 	/**
-	 * Get the name of the crashed app from the report (hides implementation in FGenericErrorReport)
+	 * Get the full path of the crashed app from the report (hides implementation in FGenericErrorReport)
 	 */
-	FString FindCrashedAppName() const;
+	FString FindCrashedAppPath() const;
 
 	/**
 	 * Look for the most recent Mac Error Report
diff --git a/Engine/Source/Programs/CrashReporter/CrashReportClient/Private/SCrashReportClient.cpp b/Engine/Source/Programs/CrashReporter/CrashReportClient/Private/SCrashReportClient.cpp
--- a/Engine/Source/Programs/CrashReporter/CrashReportClient/Private/SCrashReportClient.cpp
+++ b/Engine/Source/Programs/CrashReporter/CrashReportClient/Private/SCrashReportClient.cpp
@@ -24,10 +24,6 @@ void SCrashReportClient::Construct(const FArguments& InArgs, TSharedRef<FCrashRe
 {
 	CrashReportClient = Client;
 
-	TSharedPtr<SScrollBar> HiddenScrollbar;
-	SAssignNew( HiddenScrollbar, SScrollBar );
-	HiddenScrollbar->SetVisibility( EVisibility::Collapsed );
-
 	auto CrashedAppName = CrashReportClient->GetCrashedAppName();
 
 	// Set the text displaying the name of the crashed app, if available
@@ -87,86 +83,93 @@ void SCrashReportClient::Construct(const FArguments& InArgs, TSharedRef<FCrashRe
 			]
 
 			+SVerticalBox::Slot()
-			.FillHeight(0.3f)
-			.Padding( FMargin( 4, 10, 4, 4 ) )
+			.Padding(FMargin(4, 10, 4, 4))
 			[
-				SNew( SOverlay )
-
-				+ SOverlay::Slot()
+				SNew(SSplitter)
+				.Orientation(Orient_Vertical)
+				+SSplitter::Slot()
+				.Value(0.3f)
 				[
-					SAssignNew( CrashDetailsInformation, SMultiLineEditableTextBox )
-					.Style( &FCrashReportClientStyle::Get().GetWidgetStyle<FEditableTextBoxStyle>( "NormalEditableTextBox" ) )
-					.OnTextCommitted( CrashReportClient.ToSharedRef(), &FCrashReportClient::UserCommentChanged )
-					.Font( FSlateFontInfo( FPaths::EngineContentDir() / TEXT( "Slate/Fonts/Roboto-Regular.ttf" ), 9 ) )
-					.AutoWrapText( true )
-					.BackgroundColor( FSlateColor( FLinearColor::Black ) )
-					.ForegroundColor( FSlateColor( FLinearColor::White * 0.8f ) )
+					SNew( SOverlay )
+					+ SOverlay::Slot()
+					[
+						SAssignNew( CrashDetailsInformation, SMultiLineEditableTextBox )
+						.Style( &FCrashReportClientStyle::Get().GetWidgetStyle<FEditableTextBoxStyle>( "NormalEditableTextBox" ) )
+						.OnTextCommitted( CrashReportClient.ToSharedRef(), &FCrashReportClient::UserCommentChanged )
+						.Font( FSlateFontInfo( FPaths::EngineContentDir() / TEXT( "Slate/Fonts/Roboto-Regular.ttf" ), 9 ) )
+						.AutoWrapText( true )
+						.BackgroundColor( FSlateColor( FLinearColor::Black ) )
+						.ForegroundColor( FSlateColor( FLinearColor::White * 0.8f ) )
+					]
+
+					// HintText is not implemented in SMultiLineEditableTextBox, so this is a workaround.
+					+ SOverlay::Slot()
+					[
+						SNew(STextBlock)
+						.Margin( FMargin(4,2,0,0) )
+						.Font( FSlateFontInfo( FPaths::EngineContentDir() / TEXT( "Slate/Testing/Fonts/Roboto-Italic.ttf" ), 9 ) )
+						.ColorAndOpacity( FSlateColor( FLinearColor::White * 0.5f ) )
+						.Text( LOCTEXT( "CrashProvide", "Please provide detailed information about what you were doing when the crash occurred." ) )
+						.Visibility( this, &SCrashReportClient::IsHintTextVisible )
+					]	
 				]
 
-				// HintText is not implemented in SMultiLineEditableTextBox, so this is a workaround.
-				+ SOverlay::Slot()
-				
-				[
-					SNew(STextBlock)
-					.Margin( FMargin(4,2,0,0) )
-					.Font( FSlateFontInfo( FPaths::EngineContentDir() / TEXT( "Slate/Testing/Fonts/Roboto-Italic.ttf" ), 9 ) )
-					.ColorAndOpacity( FSlateColor( FLinearColor::White * 0.5f ) )
-					.Text( LOCTEXT( "CrashProvide", "Please provide detailed information about what you were doing when the crash occurred." ) )
-					.Visibility( this, &SCrashReportClient::IsHintTextVisible )
-				]	
-			]
-
-			+ SVerticalBox::Slot()
-			.AutoHeight()
-			.Padding( FMargin( 4, 10, 4, 0 ) )
-			[
-				SNew( SOverlay )
-
-				+ SOverlay::Slot()			
+				+SSplitter::Slot()
+				.Value(0.7f)
 				[
-					SNew( SColorBlock )
-					.Color( FLinearColor::Black )
-				]
-
-				+ SOverlay::Slot()
-				[
-					SNew( STextBlock )
-					.Margin( FMargin( 4, 2, 0, 8 ) )
-					.Font( FSlateFontInfo( FPaths::EngineContentDir() / TEXT( "Slate/Testing/Fonts/Roboto-Italic.ttf" ), 9 ) )
-					.Text( LOCTEXT( "CrashReportData", "Crash report data:" ) )
-					.ColorAndOpacity( FSlateColor( FLinearColor::White * 0.5f ) )
-					.AutoWrapText( false )
-				]
-			]
-
-			+SVerticalBox::Slot()
-			.FillHeight(0.7f)
-			.Padding( FMargin( 4, 0 ) )
-			[
-				SNew(SOverlay)
+					SNew(SVerticalBox)
+					+ SVerticalBox::Slot()
+					.AutoHeight()
+					.Padding(FMargin(4, 10, 4, 0))
+					[
+						SNew(SOverlay)
+
+						+ SOverlay::Slot()			
+						[
+							SNew(SColorBlock)
+							.Color(FLinearColor::Black)
+						]
+
+						+ SOverlay::Slot()
+						[
+							SNew(STextBlock)
+							.Margin(FMargin(4, 2, 0, 8))
+							.Font(FSlateFontInfo(FPaths::EngineContentDir() / TEXT("Slate/Testing/Fonts/Roboto-Italic.ttf"), 9))
+							.Text(LOCTEXT("CrashReportData", "Crash report data:"))
+							.ColorAndOpacity(FSlateColor(FLinearColor::White * 0.5f))
+							.AutoWrapText(false)
+						]
+					]
+
+					+ SVerticalBox::Slot()
+					.FillHeight(0.7f)
+					.Padding(FMargin(4, 0))
+					[
+						SNew(SOverlay)
 		
-				+SOverlay::Slot()
-				[
-					SNew( SMultiLineEditableTextBox )
-					.Style( &FCrashReportClientStyle::Get().GetWidgetStyle<FEditableTextBoxStyle>( "NormalEditableTextBox" ) )
-					.Font( FSlateFontInfo( FPaths::EngineContentDir() / TEXT( "Slate/Fonts/Roboto-Regular.ttf" ), 8 ) )
-					.AutoWrapText( false )
-					.IsReadOnly( true )
-					.ReadOnlyForegroundColor( FSlateColor( FLinearColor::Black ) )
-					.BackgroundColor( FSlateColor( FLinearColor::Black ) )
-					.ForegroundColor( FSlateColor( FLinearColor::White * 0.8f ) )
-					.Text( Client, &FCrashReportClient::GetDiagnosticText )
-					.HScrollBar( HiddenScrollbar )
-				]
-
-
-				+SOverlay::Slot()
-				.HAlign(HAlign_Center)
-				.VAlign(VAlign_Center)
-				[
-					SNew(SThrobber)
-					.Visibility(CrashReportClient.ToSharedRef(), &FCrashReportClient::IsThrobberVisible)
-					.NumPieces(5)
+						+ SOverlay::Slot()
+						[
+							SNew( SMultiLineEditableTextBox )
+							.Style( &FCrashReportClientStyle::Get().GetWidgetStyle<FEditableTextBoxStyle>( "NormalEditableTextBox" ) )
+							.Font( FSlateFontInfo( FPaths::EngineContentDir() / TEXT( "Slate/Fonts/Roboto-Regular.ttf" ), 8 ) )
+							.AutoWrapText( false )
+							.IsReadOnly( true )
+							.ReadOnlyForegroundColor( FSlateColor( FLinearColor::Black ) )
+							.BackgroundColor( FSlateColor( FLinearColor::Black ) )
+							.ForegroundColor( FSlateColor( FLinearColor::White * 0.8f ) )
+							.Text( Client, &FCrashReportClient::GetDiagnosticText )
+						]
+
+
+						+ SOverlay::Slot()
+						.HAlign(HAlign_Center)
+						.VAlign(VAlign_Center)
+						[
+							SNew(SThrobber)
+							.Visibility(CrashReportClient.ToSharedRef(), &FCrashReportClient::IsThrobberVisible)
+							.NumPieces(5)
+						]
+					]
 				]
 			]
 
@@ -211,6 +214,18 @@ void SCrashReportClient::Construct(const FArguments& InArgs, TSharedRef<FCrashRe
 					SNew(SSpacer)
 				]
 
+				+SHorizontalBox::Slot()
+				.HAlign(HAlign_Center)
+				.VAlign(VAlign_Center)
+				.AutoWidth()
+				.Padding( FMargin(6) )
+				[
+					SNew(SButton)
+					.ContentPadding( FMargin(8,2) )
+					.Text(LOCTEXT("SendAndRestartEditor", "Send and Restart Editor"))
+					.OnClicked(Client, &FCrashReportClient::SubmitAndRestart)
+				]
+
 				+SHorizontalBox::Slot()
 				.HAlign(HAlign_Center)
 				.VAlign(VAlign_Center)
diff --git a/Engine/Source/Programs/CrashReporter/CrashReportClient/Private/Windows/WindowsErrorReport.cpp b/Engine/Source/Programs/CrashReporter/CrashReportClient/Private/Windows/WindowsErrorReport.cpp
--- a/Engine/Source/Programs/CrashReporter/CrashReportClient/Private/Windows/WindowsErrorReport.cpp
+++ b/Engine/Source/Programs/CrashReporter/CrashReportClient/Private/Windows/WindowsErrorReport.cpp
@@ -73,6 +73,11 @@ void FWindowsErrorReport::ShutDown()
 	CrashHelperModule->ShutdownModule();
 }
 
+FString FWindowsErrorReport::FindCrashedAppPath() const
+{
+	return FWindowsReportParser::Find(ReportDirectory, TEXT("AppPath="));
+}
+
 FText FWindowsErrorReport::DiagnoseReport() const
 {
 	// Should check if there are local PDBs before doing anything
diff --git a/Engine/Source/Programs/CrashReporter/CrashReportClient/Private/Windows/WindowsErrorReport.h b/Engine/Source/Programs/CrashReporter/CrashReportClient/Private/Windows/WindowsErrorReport.h
--- a/Engine/Source/Programs/CrashReporter/CrashReportClient/Private/Windows/WindowsErrorReport.h
+++ b/Engine/Source/Programs/CrashReporter/CrashReportClient/Private/Windows/WindowsErrorReport.h
@@ -39,6 +39,11 @@ public:
 	 */
 	FText DiagnoseReport() const;
 
+	/**
+	 * Get the full path of the crashed app from the report
+	 */
+	FString FindCrashedAppPath() const;
+
 	/**
 	 * Look for the most recent Windows Error Report
 	 * @return Full path to the most recent report, or an empty string if none found
diff --git a/Engine/Source/Runtime/Core/Private/Misc/CoreMisc.cpp b/Engine/Source/Runtime/Core/Private/Misc/CoreMisc.cpp
--- a/Engine/Source/Runtime/Core/Private/Misc/CoreMisc.cpp
+++ b/Engine/Source/Runtime/Core/Private/Misc/CoreMisc.cpp
@@ -545,6 +545,7 @@ bool FFileHelper::LoadANSITextFileToStrings(const TCHAR* InFilename, IFileManage
 
 bool FCommandLine::bIsInitialized = false;
 TCHAR FCommandLine::CmdLine[FCommandLine::MaxCommandLineSize] = TEXT("");
+TCHAR FCommandLine::OriginalCmdLine[FCommandLine::MaxCommandLineSize] = TEXT("");
 FString FCommandLine::SubprocessCommandLine(TEXT(" -Multiprocess"));
 
 bool FCommandLine::IsInitialized()
@@ -558,8 +559,19 @@ const TCHAR* FCommandLine::Get()
 	return CmdLine;
 }
 
+const TCHAR* FCommandLine::GetOriginal()
+{
+	UE_CLOG(!bIsInitialized, LogInit, Fatal, TEXT("Attempting to get the command line but it hasn't been initialized yet."));
+	return OriginalCmdLine;
+}
+
 bool FCommandLine::Set(const TCHAR* NewCommandLine)
-{		
+{
+	if (!bIsInitialized)
+	{
+		FCString::Strncpy(OriginalCmdLine, NewCommandLine, ARRAY_COUNT(CmdLine));
+	}
+
 	FCString::Strncpy( CmdLine, NewCommandLine, ARRAY_COUNT(CmdLine) );
 	bIsInitialized = true;
 
diff --git a/Engine/Source/Runtime/Core/Private/Windows/WindowsPlatformCrashContext.cpp b/Engine/Source/Runtime/Core/Private/Windows/WindowsPlatformCrashContext.cpp
--- a/Engine/Source/Runtime/Core/Private/Windows/WindowsPlatformCrashContext.cpp
+++ b/Engine/Source/Runtime/Core/Private/Windows/WindowsPlatformCrashContext.cpp
@@ -185,7 +185,7 @@ void SetReportParameters( HREPORT ReportHandle, EXCEPTION_POINTERS* ExceptionInf
 	}
 
 	// AssertLog should be ErrorMessage, but this require crash server changes, so don't change this.
-	StringCchPrintf( StringBuffer, MAX_SPRINTF, TEXT( "!%s!AssertLog=\"%s\"" ), FCommandLine::Get(), LocalBuffer );
+	StringCchPrintf( StringBuffer, MAX_SPRINTF, TEXT( "!%s!AssertLog=\"%s\"" ), FCommandLine::GetOriginal(), LocalBuffer );
 	Result = WerReportSetParameter( ReportHandle, WER_P8, TEXT( "Commandline" ), StringBuffer );
 
 	StringCchPrintf( StringBuffer, MAX_SPRINTF, TEXT( "%s!%s!%s!%d" ), TEXT( BRANCH_NAME ), FPlatformProcess::BaseDir(), FPlatformMisc::GetEngineMode(), BUILT_FROM_CHANGELIST );
diff --git a/Engine/Source/Runtime/Core/Public/Misc/CoreMisc.h b/Engine/Source/Runtime/Core/Public/Misc/CoreMisc.h
--- a/Engine/Source/Runtime/Core/Public/Misc/CoreMisc.h
+++ b/Engine/Source/Runtime/Core/Public/Misc/CoreMisc.h
@@ -67,10 +67,15 @@ struct CORE_API FCommandLine
 	static const uint32 MaxCommandLineSize = 16384;
 
 	/** 
-	 * Returns the executables command line. 
+	 * Returns an edited version of the executable's command line with the game name and certain other parameters removed. 
 	 */
 	static const TCHAR* Get();
 	
+	/**
+	 * Returns the command line originally passed to the executable.
+	 */
+	static const TCHAR* GetOriginal();
+
 	/** 
 	 * Checks if the command line has been initialized. 
 	 */
@@ -120,6 +125,8 @@ private:
 	static bool bIsInitialized;
 	/** character buffer containing the command line */
 	static TCHAR CmdLine[MaxCommandLineSize];
+	/** character buffer containing the original command line */
+	static TCHAR OriginalCmdLine[MaxCommandLineSize];
 	/** subprocess command line */
 	static FString SubprocessCommandLine;
 };
