Tweaks to CrashReportClient.

diff --git a/Engine/Source/Programs/CrashReporter/CrashReportClient/Private/SCrashReportClient.cpp b/Engine/Source/Programs/CrashReporter/CrashReportClient/Private/SCrashReportClient.cpp
--- a/Engine/Source/Programs/CrashReporter/CrashReportClient/Private/SCrashReportClient.cpp
+++ b/Engine/Source/Programs/CrashReporter/CrashReportClient/Private/SCrashReportClient.cpp
@@ -120,7 +120,6 @@ void SCrashReportClient::Construct(const FArguments& InArgs, TSharedRef<FCrashRe
 					SNew(SVerticalBox)
 					+ SVerticalBox::Slot()
 					.AutoHeight()
-					.Padding(FMargin(4, 10, 4, 0))
 					[
 						SNew(SOverlay)
 
@@ -143,7 +142,6 @@ void SCrashReportClient::Construct(const FArguments& InArgs, TSharedRef<FCrashRe
 
 					+ SVerticalBox::Slot()
 					.FillHeight(0.7f)
-					.Padding(FMargin(4, 0))
 					[
 						SNew(SOverlay)
 		
@@ -222,7 +220,7 @@ void SCrashReportClient::Construct(const FArguments& InArgs, TSharedRef<FCrashRe
 				[
 					SNew(SButton)
 					.ContentPadding( FMargin(8,2) )
-					.Text(LOCTEXT("SendAndRestartEditor", "Send and Restart Editor"))
+					.Text(LOCTEXT("SendAndRestartEditor", "Send and Restart"))
 					.OnClicked(Client, &FCrashReportClient::SubmitAndRestart)
 				]
 
