SQL

diff --git a/Engine/Source/Programs/CrashReporter/CrashReportWebSite/Models/Query-GetTodayCrashes.sql b/Engine/Source/Programs/CrashReporter/CrashReportWebSite/Models/Query-GetTodayCrashes.sql
--- a/Engine/Source/Programs/CrashReporter/CrashReportWebSite/Models/Query-GetTodayCrashes.sql
+++ b/Engine/Source/Programs/CrashReporter/CrashReportWebSite/Models/Query-GetTodayCrashes.sql
@@ -2,6 +2,7 @@
 SELECT *
 from crashes
 where [TimeOfCrash] > CAST(GETDATE() AS DATE) 
+order by [TimeOfCrash]
 --where [EpicAccountId] is not null and buildversion != '4.5.0.0' and buildversion != '4.5.1.0
 --where buildversion like '4.5.1%' and [ChangeListVersion] != '0' and [PlatformName] not like 'Mac%' and [Branch] like 'UE4-Re%'
 
