CrashReportWebsite
UECORE-49 A direct connection between CrashReporter and JIRA. (Adding callstack to the JIRA tickets, added support for automatics Platform gathering for JIRA tickets)

diff --git a/Engine/Source/Programs/CrashReporter/CrashReportWebSite/Controllers/BuggsController.cs b/Engine/Source/Programs/CrashReporter/CrashReportWebSite/Controllers/BuggsController.cs
--- a/Engine/Source/Programs/CrashReporter/CrashReportWebSite/Controllers/BuggsController.cs
+++ b/Engine/Source/Programs/CrashReporter/CrashReportWebSite/Controllers/BuggsController.cs
@@ -95,10 +95,11 @@ namespace Tools.CrashReporter.CrashReportWebSite.Controllers
 					return RedirectToAction( "" );
 				}
 
+				// @TODO yrx 2015-02-17 JIRA
 				using( FAutoScopedLogTimer GetCrashesTimer = new FAutoScopedLogTimer( "Bugg.GetCrashes().ToList" ) )
 				{
 					Crashes = Bugg.GetCrashes();
-					Bugg.AffectedBuilds = new SortedSet<string>();
+					Bugg.AffectedVersions = new SortedSet<string>();
 
 					HashSet<string> MachineIds = new HashSet<string>();
 					foreach( Crash Crash in Crashes )
@@ -107,7 +108,7 @@ namespace Tools.CrashReporter.CrashReportWebSite.Controllers
 						// Ignore bad build versions.
 						if( Crash.BuildVersion.StartsWith( "4." ) )
 						{
-							Bugg.AffectedBuilds.Add( Crash.BuildVersion );
+							Bugg.AffectedVersions.Add( Crash.BuildVersion );
 						}
 
 						if( Crash.User == null )
@@ -136,7 +137,7 @@ namespace Tools.CrashReporter.CrashReportWebSite.Controllers
 					if( !string.IsNullOrEmpty( BuggsForm["SetTTP"] ) )
 					{
 						Bugg.TTPID = BuggsForm["SetTTP"];
-						LocalCrashRepository.SetBuggTTPID( Bugg.TTPID, id );
+						BuggRepository.SetJIRAForBuggAndCrashes( Bugg.TTPID, id );
 					}
 
 					if( !string.IsNullOrEmpty( BuggsForm["Description"] ) )
diff --git a/Engine/Source/Programs/CrashReporter/CrashReportWebSite/Controllers/ReportsController.cs b/Engine/Source/Programs/CrashReporter/CrashReportWebSite/Controllers/ReportsController.cs
--- a/Engine/Source/Programs/CrashReporter/CrashReportWebSite/Controllers/ReportsController.cs
+++ b/Engine/Source/Programs/CrashReporter/CrashReportWebSite/Controllers/ReportsController.cs
@@ -194,12 +194,16 @@ namespace Tools.CrashReporter.CrashReportWebSite.Controllers
 		/// Retrieve all Buggs matching the search criteria.
 		/// </summary>
 		/// <param name="FormData">The incoming form of search criteria from the client.</param>
+		/// <param name="BuggIDToBeAddedToJira">ID of the bugg that will be added to JIRA</param>
 		/// <returns>A view to display the filtered Buggs.</returns>
-		public ReportsViewModel GetResults( FormHelper FormData )
+		public ReportsViewModel GetResults( FormHelper FormData, int BuggIDToBeAddedToJira )
 		{
+			// @TODO yrx 2015-02-17 BuggIDToBeAddedToJira replace with List<int> based on check box and Submit?
 			// Enumerate JIRA projects if needed.
 			// https://jira.ol.epicgames.net//rest/api/2/project
-
+			var JC = JiraConnection.Get();
+			var JiraComponents = JC.GetNameToComponents();
+			var JiraVersions = JC.GetNameToVersions();
 
 			using( FAutoScopedLogTimer LogTimer = new FAutoScopedLogTimer( this.GetType().ToString() ) )
 			{
@@ -230,6 +234,7 @@ namespace Tools.CrashReporter.CrashReportWebSite.Controllers
 						BuiltFromCL = Crash.ChangeListVersion,
 						Pattern = Crash.Pattern,
 						MachineID = Crash.ComputerName,
+						Branch = Crash.Branch,
 					} )
 					.ToList();
 				int NumCrashes = Crashes.Count;
@@ -312,14 +317,20 @@ namespace Tools.CrashReporter.CrashReportWebSite.Controllers
 							NewBugg.TTPID = RealBugg.TTPID;							// JIRA
 							NewBugg.FixedChangeList = RealBugg.FixedChangeList;		// FixCL
 							NewBugg.NumberOfCrashes = Top.Value;					// # Occurrences
+							NewBugg.Pattern = RealBugg.Pattern;					// # Occurrences
 
 							//NewBugg.BuildVersion = 
 
 							var CrashesForBugg = Crashes.Where( Crash => Crash.Pattern == Top.Key ).ToList();
 
-							NewBugg.AffectedBuilds = new SortedSet<string>();
+							NewBugg.AffectedVersions = new SortedSet<string>();
+							NewBugg.AffectedMajorVersions = new SortedSet<string>(); // 4.4, 4.5 and so
+							NewBugg.BranchesFoundIn = new SortedSet<string>();
+							NewBugg.AffectedPlatforms = new SortedSet<string>();
 
 							HashSet<string> MachineIds = new HashSet<string>();
+							int FirstCLAffected = int.MaxValue;
+
 							foreach( var Crash in CrashesForBugg )
 							{
 								// Only add machine if the number has 32 characters
@@ -329,15 +340,54 @@ namespace Tools.CrashReporter.CrashReportWebSite.Controllers
 								}
 
 								// Ignore bad build versions.
+								// @TODO yrx 2015-02-17 What about launcher?
 								if( Crash.BuildVersion.StartsWith( "4." ) )
 								{
-									NewBugg.AffectedBuilds.Add( Crash.BuildVersion );
-								}
+									if( !string.IsNullOrEmpty(Crash.BuildVersion) )
+									{
+										NewBugg.AffectedVersions.Add( Crash.BuildVersion );
+									}
+									if( !string.IsNullOrEmpty( Crash.Branch ) && Crash.Branch.StartsWith( "UE4" ) )
+									{
+										NewBugg.BranchesFoundIn.Add( Crash.Branch );
+									}
+
+									int CrashBuiltFromCL = 0;
+									int.TryParse( Crash.BuiltFromCL, out CrashBuiltFromCL );
+									if( CrashBuiltFromCL > 0 )
+									{
+										FirstCLAffected = Math.Min( FirstCLAffected, CrashBuiltFromCL );
+									}
+
+									if( !string.IsNullOrEmpty( Crash.Platform ) )
+									{
+										// Platform = "Platform [Desc]";
+										var PlatSubs = Crash.Platform.Split( " ".ToCharArray(), StringSplitOptions.RemoveEmptyEntries );
+										if( PlatSubs.Length >= 1 )
+										{
+											NewBugg.AffectedPlatforms.Add( PlatSubs[0] );
+										}
+									}
+								}	
 							}
 
-							if( NewBugg.AffectedBuilds.Count > 0 )
+							// CopyToJira 
+							NewBugg.ToJiraFirstCLAffected = FirstCLAffected;
+
+							if( NewBugg.AffectedVersions.Count > 0 )
 							{
-								NewBugg.BuildVersion = NewBugg.AffectedBuilds.Last();	// Latest Version Affected
+								NewBugg.BuildVersion = NewBugg.AffectedVersions.Last();	// Latest Version Affected
+							}
+
+							
+							foreach( var AffectedBuild in NewBugg.AffectedVersions )
+							{
+								var Subs = AffectedBuild.Split( ".".ToCharArray(), StringSplitOptions.RemoveEmptyEntries );
+								if( Subs.Length >= 2 )
+								{
+									string MajorVersion = Subs[0] + "." + Subs[1];
+									NewBugg.AffectedMajorVersions.Add( MajorVersion );
+								}
 							}
 
 							NewBugg.NumberOfUniqueMachines = MachineIds.Count;		// # Affected Users
@@ -348,31 +398,60 @@ namespace Tools.CrashReporter.CrashReportWebSite.Controllers
 							int ILatestCLAffected = -1;
 							int.TryParse( LatestCLAffected, out ILatestCLAffected );
 							NewBugg.LatestCLAffected = ILatestCLAffected;			// Latest CL Affected
-
+							
 							string LatestOSAffected = CrashesForBugg.OrderByDescending( Crash => Crash.TimeOfCrash ).First().Platform;
-
 							NewBugg.LatestOSAffected = LatestOSAffected;			// Latest Environment Affected
 
 							NewBugg.TimeOfFirstCrash = RealBugg.TimeOfFirstCrash;	// First Crash Timestamp
 
+							// ToJiraSummary
+							var Callstack = RealBugg.GetFunctionCalls();
+							NewBugg.ToJiraSummary = Callstack.Count > 1 ? Callstack[0] : "No valid callstack found";
+
+							// ToJiraVersions
+							NewBugg.ToJiraVersions = new List<object>();
+							foreach( var Version in NewBugg.AffectedMajorVersions )
+							{
+								bool bValid = JC.GetNameToVersions().ContainsKey( Version );
+								if( bValid )
+								{
+									NewBugg.ToJiraVersions.Add( JC.GetNameToVersions()[Version] );
+								}
+							}
+
+							// ToJiraBranches
+							NewBugg.ToJiraBranches = new List<object>();
+							foreach( var Platform in NewBugg.BranchesFoundIn )
+							{
+								string CleanedBranch = Platform.Contains( "UE4-Releases" ) ? "UE4-Releases" : Platform;
+								Dictionary<string, object> JiraBranch = null;
+								JC.GetNameToBranchFoundIn().TryGetValue( CleanedBranch, out JiraBranch );
+								if( JiraBranch != null && !NewBugg.ToJiraBranches.Contains( JiraBranch ) )
+								{
+									NewBugg.ToJiraBranches.Add( JiraBranch );
+								}
+							}
+
+							// ToJiraPlatforms
+							NewBugg.ToJiraPlatforms = new List<object>();
+							foreach( var Platform in NewBugg.AffectedPlatforms )
+							{
+								bool bValid = JC.GetNameToPlatform().ContainsKey( Platform );
+								if( bValid )
+								{
+									NewBugg.ToJiraPlatforms.Add( JC.GetNameToPlatform()[Platform] );
+								}
+							}
+
 							// Verify valid JiraID, this may be still a TTP 
-							if( !string.IsNullOrEmpty( RealBugg.TTPID ) )
+							if( !string.IsNullOrEmpty( NewBugg.TTPID ) )
 							{
 								int TTPID = 0;
-								int.TryParse( RealBugg.TTPID, out TTPID );
+								int.TryParse( NewBugg.TTPID, out TTPID );
 
 								if( TTPID == 0 )
 								{
-									string JiraID = RealBugg.TTPID;
-									FoundJiras.Add( "key = " + JiraID );
-
-									bool bAdd = !JiraIDtoBugg.ContainsKey( JiraID );
-									if( bAdd )
-									{
-										JiraIDtoBugg.Add( JiraID, new List<Bugg>() );
-									}
-
-									JiraIDtoBugg[JiraID].Add( NewBugg );
+									AddBuggJiraMapping(NewBugg, ref FoundJiras, ref JiraIDtoBugg);						
 								}
 							}
 
@@ -385,14 +464,26 @@ namespace Tools.CrashReporter.CrashReportWebSite.Controllers
 					}
 				}
 
-				// Grab the data form JIRA.
-				string JiraSearchQuery = string.Join( " OR ", FoundJiras );
+				if( BuggIDToBeAddedToJira > 0 )
+				{
+					var Bugg = Buggs.Where( X => X.Id == BuggIDToBeAddedToJira ).FirstOrDefault();
+					if( Bugg != null )
+					{
+						Bugg.CopyToJira();
+						AddBuggJiraMapping( Bugg, ref FoundJiras, ref JiraIDtoBugg );
+					}
+				}
 
-				using( FAutoScopedLogTimer JiraResultsTimer = new FAutoScopedLogTimer( "JiraResults" ) )
+				if( JC.CanBeUsed() )
 				{
-					var JiraResults = JiraConnection.Get().SearchJiraTickets(
-						JiraSearchQuery,
-						new string[] 
+					// Grab the data form JIRA.
+					string JiraSearchQuery = string.Join( " OR ", FoundJiras );
+
+					using( FAutoScopedLogTimer JiraResultsTimer = new FAutoScopedLogTimer( "JiraResults" ) )
+					{
+						var JiraResults = JC.SearchJiraTickets(
+							JiraSearchQuery,
+							new string[] 
 						{ 
 							"key",				// string
 							"summary",			// string
@@ -403,45 +494,46 @@ namespace Tools.CrashReporter.CrashReportWebSite.Controllers
 						} );
 
 
-					// Jira Key, Summary, Components, Resolution, Fix version, Fix changelist
-					foreach( var Jira in JiraResults )
-					{
-						string JiraID = Jira.Key;
+						// Jira Key, Summary, Components, Resolution, Fix version, Fix changelist
+						foreach( var Jira in JiraResults )
+						{
+							string JiraID = Jira.Key;
 
-						string Summary = (string)Jira.Value["summary"];
+							string Summary = (string)Jira.Value["summary"];
 
-						string ComponentsText = "";
-						System.Collections.ArrayList Components = (System.Collections.ArrayList)Jira.Value["components"];
-						foreach( Dictionary<string, object> Component in Components )
-						{
-							ComponentsText += (string)Component["name"];
-							ComponentsText += " ";
-						}
+							string ComponentsText = "";
+							System.Collections.ArrayList Components = (System.Collections.ArrayList)Jira.Value["components"];
+							foreach( Dictionary<string, object> Component in Components )
+							{
+								ComponentsText += (string)Component["name"];
+								ComponentsText += " ";
+							}
 
-						Dictionary<string, object> ResolutionFields = (Dictionary<string, object>)Jira.Value["resolution"];
-						string Resolution = ResolutionFields != null ? (string)ResolutionFields["name"] : "";
+							Dictionary<string, object> ResolutionFields = (Dictionary<string, object>)Jira.Value["resolution"];
+							string Resolution = ResolutionFields != null ? (string)ResolutionFields["name"] : "";
 
-						string FixVersionsText = "";
-						System.Collections.ArrayList FixVersions = (System.Collections.ArrayList)Jira.Value["fixVersions"];
-						foreach( Dictionary<string, object> FixVersion in FixVersions )
-						{
-							FixVersionsText += (string)FixVersion["name"];
-							FixVersionsText += " ";
-						}
+							string FixVersionsText = "";
+							System.Collections.ArrayList FixVersions = (System.Collections.ArrayList)Jira.Value["fixVersions"];
+							foreach( Dictionary<string, object> FixVersion in FixVersions )
+							{
+								FixVersionsText += (string)FixVersion["name"];
+								FixVersionsText += " ";
+							}
 
-						int FixCL = Jira.Value["customfield_11200"] != null ? (int)(decimal)Jira.Value["customfield_11200"] : 0;
+							int FixCL = Jira.Value["customfield_11200"] != null ? (int)(decimal)Jira.Value["customfield_11200"] : 0;
 
-						var BuggsForJira = JiraIDtoBugg[JiraID];
+							var BuggsForJira = JiraIDtoBugg[JiraID];
 
-						foreach( Bugg Bugg in BuggsForJira )
-						{
-							Bugg.JiraSummary = Summary;
-							Bugg.JiraComponentsText = ComponentsText;
-							Bugg.JiraResolution = Resolution;
-							Bugg.JiraFixVersionsText = FixVersionsText;
-							if( FixCL != 0 )
+							foreach( Bugg Bugg in BuggsForJira )
 							{
-								Bugg.JiraFixCL = FixCL.ToString();
+								Bugg.JiraSummary = Summary;
+								Bugg.JiraComponentsText = ComponentsText;
+								Bugg.JiraResolution = Resolution;
+								Bugg.JiraFixVersionsText = FixVersionsText;
+								if( FixCL != 0 )
+								{
+									Bugg.JiraFixCL = FixCL.ToString();
+								}
 							}
 						}
 					}
@@ -460,6 +552,20 @@ namespace Tools.CrashReporter.CrashReportWebSite.Controllers
 			}
 		}
 
+		private void AddBuggJiraMapping( Bugg NewBugg, ref List<string> FoundJiras, ref Dictionary<string, List<Bugg>> JiraIDtoBugg )
+		{
+			string JiraID = NewBugg.TTPID;
+			FoundJiras.Add( "key = " + JiraID );
+
+			bool bAdd = !JiraIDtoBugg.ContainsKey( JiraID );
+			if( bAdd )
+			{
+				JiraIDtoBugg.Add( JiraID, new List<Bugg>() );
+			}
+
+			JiraIDtoBugg[JiraID].Add( NewBugg );
+		}
+
 		/// <summary>
 		/// 
 		/// </summary>
@@ -470,7 +576,19 @@ namespace Tools.CrashReporter.CrashReportWebSite.Controllers
 			using( FAutoScopedLogTimer LogTimer = new FAutoScopedLogTimer( this.GetType().ToString() ) )
 			{
 				FormHelper FormData = new FormHelper( Request, ReportsForm, "JustReport" );
-				ReportsViewModel Results = GetResults( FormData );
+
+				// Handle 'CopyToJira' button
+				int BuggIDToBeAddedToJira = 0;
+				foreach( var Entry in ReportsForm )
+				{
+					
+					if( int.TryParse( Entry.ToString(), out BuggIDToBeAddedToJira ) )
+					{
+						break;
+					}
+				}
+
+				ReportsViewModel Results = GetResults( FormData, BuggIDToBeAddedToJira );
 				Results.GenerationTime = LogTimer.GetElapsedSeconds().ToString( "F2" );
 				return View( "Index", Results );
 			}
diff --git a/Engine/Source/Programs/CrashReporter/CrashReportWebSite/Models/BuggRepository.cs b/Engine/Source/Programs/CrashReporter/CrashReportWebSite/Models/BuggRepository.cs
--- a/Engine/Source/Programs/CrashReporter/CrashReportWebSite/Models/BuggRepository.cs
+++ b/Engine/Source/Programs/CrashReporter/CrashReportWebSite/Models/BuggRepository.cs
@@ -21,14 +21,36 @@ namespace Tools.CrashReporter.CrashReportWebSite.Models
 	{
 		private const string DefaultUserGroup = "General";
 
-		private readonly CrashReportDataContext Context;
+		private static readonly CrashReportDataContext Context = new CrashReportDataContext();
+
+		/// <summary> Submits enqueue changes to the database. </summary>
+		public static void SubmitChanges()
+		{
+			Context.SubmitChanges();
+		}
 
 		/// <summary>
-		/// The default constructor to create a data context to access the database.
+		/// Sets the JIRA for all crashes in a Bugg.
 		/// </summary>
-		public BuggRepository()
+		/// <param name="JIRA">A string representing a TTP.</param>
+		/// <param name="BuggId">The id of the Bugg to update the crashes for.</param>
+		public static void SetJIRAForBuggAndCrashes( string JIRA, int BuggId )
 		{
-			Context = new CrashReportDataContext();
+			try
+			{
+				using( FAutoScopedLogTimer LogTimer = new FAutoScopedLogTimer( "SetJIRAForBuggAndCrashes (" + BuggId + ")" ) )
+				{
+					string Query = "UPDATE Crashes SET TTPID = {0} WHERE Id IN ( SELECT CrashId FROM Buggs_Crashes WHERE BuggId = {1} )";
+					Context.ExecuteCommand( Query, JIRA, BuggId );
+
+					Query = "UPDATE Buggs SET TTPID = {0} WHERE id = {1}";
+					Context.ExecuteCommand( Query, JIRA, BuggId );
+				}
+			}
+			catch( Exception Ex )
+			{
+				FLogger.WriteException( "SetBuggTTPID: " + Ex.ToString() );
+			}
 		}
 
 		/// <summary>
diff --git a/Engine/Source/Programs/CrashReporter/CrashReportWebSite/Models/CrashReport.cs b/Engine/Source/Programs/CrashReporter/CrashReportWebSite/Models/CrashReport.cs
--- a/Engine/Source/Programs/CrashReporter/CrashReportWebSite/Models/CrashReport.cs
+++ b/Engine/Source/Programs/CrashReporter/CrashReportWebSite/Models/CrashReport.cs
@@ -6,6 +6,7 @@ using System.Linq;
 using System;
 using System.Data.Linq;
 using System.Diagnostics;
+using Tools.DotNETCommon;
 
 namespace Tools.CrashReporter.CrashReportWebSite.Models
 {
@@ -79,10 +80,18 @@ namespace Tools.CrashReporter.CrashReportWebSite.Models
 		/// <summary></summary>
 		public string LatestOSAffected { get; set; }
 
+		/// <summary>  </summary>
+		public SortedSet<string> AffectedPlatforms { get; set; }
+
 		/// <summary></summary>
-		public SortedSet<string> AffectedBuilds { get; set; }
+		public SortedSet<string> AffectedVersions { get; set; }
 
+		/// <summary> 4.4, 4.5 and so. </summary>
+		public SortedSet<string> AffectedMajorVersions { get; set; }
 
+		/// <summary>  </summary>
+		public SortedSet<string> BranchesFoundIn { get; set; }
+		
 		/// <summary></summary>
 		public string JiraSummary { get; set; }
 
@@ -98,10 +107,131 @@ namespace Tools.CrashReporter.CrashReportWebSite.Models
 		/// <summary></summary>
 		public string JiraFixCL { get; set; }
 
+		/// <summary>The CL of the oldest build when this bugg is occurring</summary>
+		public int ToJiraFirstCLAffected { get; set; }
+
+		/// <summary>The first line of the callstack</summary>
+		public string ToJiraSummary { get; set; }
+
+		/// <summary>Branches in JIRA</summary>
+		public List<object> ToJiraBranches { get; set; }
+
+		/// <summary>Versions in JIRA</summary>
+		public List<object> ToJiraVersions { get; set; }
+
+		/// <summary>Platforms in JIRA</summary>
+		public List<object> ToJiraPlatforms { get; set; }
+
 		/// <summary> Helper method, display this Bugg as a human readable string. Debugging purpose. </summary>
 		public override string ToString()
 		{
-			return string.Format( "Id={0} NOC={1} NOU={2} TLC={3} BV={4} Game={5}", Id, NumberOfCrashes, NumberOfUsers, TimeOfLastCrash, BuildVersion, Game );
+			return string.Format( "Id={0} NOC={1} TLC={2} BV={3} Pattern={4}", Id, NumberOfCrashes, TimeOfLastCrash, BuildVersion, Pattern != null ? Pattern.Length : 0 );
+		}
+
+		/// <summary>
+		/// Adds the bugg as a new JIRA ticket
+		/// </summary>
+		public void CopyToJira()
+		{
+			var JC = JiraConnection.Get();
+			if( JC.CanBeUsed() && string.IsNullOrEmpty( this.TTPID ) )
+			{
+				Dictionary<string, object> Fields = new Dictionary<string, object>();
+				Fields.Add( "project", new Dictionary<string, object> { { "id", 11205 } } );	// UE
+				Fields.Add( "summary", "[CrashReport] " + ToJiraSummary );						// Call Stack, Line 1
+				Fields.Add( "description", "____________" );									// Description
+				Fields.Add( "issuetype", new Dictionary<string, object> { { "id", "1" } } );	// Bug
+				Fields.Add( "labels", new string[] { "crash" } );								// <label>crash</label>
+				Fields.Add( "customfield_11500", ToJiraFirstCLAffected );						// Changelist # / Found Changelist
+				Fields.Add( "environment", LatestOSAffected );									// Platform
+
+				// Components
+				var SupportComponent = JC.GetNameToComponents()["Support"];
+				Fields.Add( "components", new object[] { SupportComponent } );
+
+				// ToJiraVersions			
+				Fields.Add( "versions", ToJiraVersions );
+
+				// ToJiraBranches
+				Fields.Add( "customfield_11201", ToJiraBranches );
+
+				// ToJiraPlatforms
+				Fields.Add( "customfield_11203", ToJiraPlatforms );
+
+				// Callstack customfield_11807
+				var Callstack = GetFunctionCalls();
+				string JiraCallstack = string.Join( "\r\n", Callstack );
+				Fields.Add( "customfield_11807", JiraCallstack );								// Callstack
+
+				string Key = JC.AddJiraTicket( Fields );
+				if( !string.IsNullOrEmpty( Key ) )
+				{
+					TTPID = Key;
+					BuggRepository.SetJIRAForBuggAndCrashes( Key, Id );
+
+					// Update the JIRA in the bugg.
+					Dictionary<string, object> FieldsToUpdate = new Dictionary<string, object>();			
+					string BuggLink = "http://crashreporter/Buggs/Show/" + Id;
+					FieldsToUpdate.Add( "customfield_11205", BuggLink );									// Additional Info URL / Link to Crash/Bugg
+					JC.UpdateJiraTicket( Key, FieldsToUpdate );
+				}
+
+			}
+		}
+
+		/// <summary> Creates a previews of the bugg, for verifying JIRA's fields. </summary>
+		public string ToTooltip()
+		{
+			string NL = "&#013;";
+
+			string Tooltip = NL;
+			Tooltip += "Project: UE" + NL;
+			Tooltip += "Summary: " + ToJiraSummary + NL;
+			Tooltip += "Description: " + "____________" + NL;
+			Tooltip += "Issuetype: " + "1 (bug)" + NL;
+			Tooltip += "Labels: " + "crash" + NL;
+			Tooltip += "Changelist # / Found Changelist: " + ToJiraFirstCLAffected + NL;
+			Tooltip += "LatestOSAffected: " + LatestOSAffected + NL;
+
+			// "name"
+			string JiraVersions = GetFieldsFrom(ToJiraVersions, "name");
+			Tooltip += "JiraVersions: " + JiraVersions + NL;
+
+			// "value"
+			string JiraBranches = GetFieldsFrom( ToJiraBranches, "value" );
+			Tooltip += "JiraBranches: " + JiraBranches + NL;
+
+			// "value"
+			string JiraPlatforms = GetFieldsFrom( ToJiraPlatforms, "value" );
+			Tooltip += "JiraPlatforms: " + JiraPlatforms + NL;
+
+			var Callstack = GetFunctionCalls();
+			string JiraCallstack = "Callstack:" + NL + string.Join( NL, Callstack ) + NL;
+			Tooltip += JiraCallstack;
+
+			return Tooltip;
+		}
+
+		/// <summary> Returns concatenated string of fields from the specified JIRA list. </summary>
+		public string GetFieldsFrom( List<object> JiraList, string FieldName )
+		{
+			string Results = "";
+
+			foreach( var Obj in JiraList )
+			{
+				Dictionary<string, object> Dict = (Dictionary<string, object>)Obj;
+				try
+				{
+					Results += (string)Dict[FieldName];
+					Results += " ";
+				}
+				catch( System.Exception /*E*/ )
+				{
+					
+				}
+			}
+
+			return Results;
 		}
 
 		/// <summary>
@@ -112,19 +242,11 @@ namespace Tools.CrashReporter.CrashReportWebSite.Models
 		{
 			CrashRepository CrashRepo = new CrashRepository();
 			var CrashList =
-							(
-								from BuggCrash in CrashRepo.Context.Buggs_Crashes
-								where BuggCrash.BuggId == Id
-								select BuggCrash.Crash/*new
-								{
-									Id = BuggCrash.Crash.Id,
-									BuildVersion = BuggCrash.Crash.BuildVersion,
-									TimeofCrash = BuggCrash.Crash.TimeOfCrash.Value,
-									CL = BuggCrash.Crash.ChangeListVersion,
-									UserName = BuggCrash.Crash.UserName,
-									Desc = BuggCrash.Crash.Description
-								}*/
-							).AsEnumerable().ToList();
+				(
+					from BuggCrash in CrashRepo.Context.Buggs_Crashes
+					where BuggCrash.BuggId == Id
+					select BuggCrash.Crash
+				).AsEnumerable().ToList();
 			return CrashList;
 		}
 
@@ -174,15 +296,15 @@ namespace Tools.CrashReporter.CrashReportWebSite.Models
 		/// <summary>
 		/// 
 		/// </summary>
-		public string GetAffectedBuilds()
+		public string GetAffectedVersions()
 		{
-			if( AffectedBuilds.Count == 1 )
+			if( AffectedVersions.Count == 1 )
 			{
-				return AffectedBuilds.Max;
+				return AffectedVersions.Max;
 			}
 			else
 			{
-				return AffectedBuilds.Min + " - " + AffectedBuilds.Max;
+				return AffectedVersions.Min + " - " + AffectedVersions.Max;
 			}
 		}
 
diff --git a/Engine/Source/Programs/CrashReporter/CrashReportWebSite/Models/CrashRepository.cs b/Engine/Source/Programs/CrashReporter/CrashReportWebSite/Models/CrashRepository.cs
--- a/Engine/Source/Programs/CrashReporter/CrashReportWebSite/Models/CrashRepository.cs
+++ b/Engine/Source/Programs/CrashReporter/CrashReportWebSite/Models/CrashRepository.cs
@@ -161,27 +161,6 @@ namespace Tools.CrashReporter.CrashReportWebSite.Models
 			}
 		}
 
-		/// <summary>
-		/// Sets the TTPID for all crashes in a Bugg.
-		/// </summary>
-		/// <param name="TTPID">A string representing a TTP.</param>
-		/// <param name="BuggId">The id of the Bugg to update the crashes for.</param>
-		public void SetBuggTTPID( string TTPID, int BuggId )
-		{
-			try
-			{
-				string Query = "UPDATE Crashes SET TTPID = {0} WHERE Id IN ( SELECT CrashId FROM Buggs_Crashes WHERE BuggId = {1} )";
-				Context.ExecuteCommand( Query, TTPID, BuggId );
-
-				Query = "UPDATE Buggs SET TTPID = {0} WHERE id = {1}";
-				Context.ExecuteCommand( Query, TTPID, BuggId );
-			}
-			catch( Exception Ex )
-			{
-				FLogger.WriteException( "SetBuggTTPID: " + Ex.ToString() );
-			}
-		}
-
 		/// <summary>
 		/// Sets the UserName and UserGroupName as derived data for a crash.
 		/// </summary>
diff --git a/Engine/Source/Programs/CrashReporter/CrashReportWebSite/Views/Buggs/Show.aspx b/Engine/Source/Programs/CrashReporter/CrashReportWebSite/Views/Buggs/Show.aspx
--- a/Engine/Source/Programs/CrashReporter/CrashReportWebSite/Views/Buggs/Show.aspx
+++ b/Engine/Source/Programs/CrashReporter/CrashReportWebSite/Views/Buggs/Show.aspx
@@ -101,7 +101,7 @@
 			<dd><%=Html.DisplayFor( m => Model.Bugg.BuildVersion )%></dd>--%>
 
 		<dt>Affected Build Versions</dt>
-			<dd><%=Model.Bugg.GetAffectedBuilds()%></dd>
+			<dd><%=Model.Bugg.GetAffectedVersions()%></dd>
 			
 		<dt>Time of Latest Crash</dt> 
 			<dd class='even' style='width:8em'><%=Model.Bugg.TimeOfLastCrash%></dd>
diff --git a/Engine/Source/Programs/CrashReporter/CrashReportWebSite/Views/Reports/ViewReports.ascx b/Engine/Source/Programs/CrashReporter/CrashReportWebSite/Views/Reports/ViewReports.ascx
--- a/Engine/Source/Programs/CrashReporter/CrashReportWebSite/Views/Reports/ViewReports.ascx
+++ b/Engine/Source/Programs/CrashReporter/CrashReportWebSite/Views/Reports/ViewReports.ascx
@@ -17,8 +17,33 @@ table, table tr, table th, table td
 	font-size: 12px;
 }
 
-</style>
+.CopyToJiraStyle
+{
+	/*padding: 0;*/
+
+	cursor:pointer;
+
+	border: solid;
+	border-width: 1px;
+
+	color: white;
+	background-color: black;
+
+	font-family: 'Segoe UI', Tahoma, Geneva, Verdana, sans-serif;
+	font-size: 12px;
+
+}
+
+.CopyToJiraStyle:hover
+{
+	color: blue;
+}
 
+</style>
+	
+<body>
+<form runat="server">
+	
 <table >
 	<tr>
 		<td>
@@ -39,29 +64,7 @@ table, table tr, table th, table td
 	</tr>
 </table>
 <br />
-<%--
-	
-Number
-NewBugg.Id = RealBugg.Id;								// CrashGroup URL (a link to the Bugg)
-NewBugg.TTPID = RealBugg.TTPID;							// JIRA
-NewBugg.FixedChangeList = RealBugg.FixedChangeList;		// FixCL
-NewBugg.NumberOfCrashes = Top.Value;					// # Occurrences
-NewBugg.AffectedBuilds = new SortedSet<string>();
-NewBugg.BuildVersion = NewBugg.AffectedBuilds.Last();	// Latest Version Affected
-NewBugg.NumberOfUniqueMachines = MachineIds.Count;		// # Affected Users
-NewBugg.LatestCLAffected = ILatestCLAffected;			// Latest CL Affected
-NewBugg.LatestOSAffected = LatestOSAffected;			// Latest Environment Affected
-NewBugg.TimeOfFirstCrash = RealBugg.TimeOfFirstCrash;	// First Crash Timestamp	
-
-#From XML JIRA	
-public string JiraSummary { get; set; }
-public string JiraComponentsText { get; set; }
-public string JiraResolution { get; set; }
-public string JiraFixVersionsText { get; set; }
-public int JiraFixCL { get; set; }
-
-
---%>
+
 <table>
 	<tr>
 		<th>
@@ -128,10 +131,24 @@ public int JiraFixCL { get; set; }
 			<%--<a href="https://jira.ol.epicgames.net/browse/<%=Model.Bugg.TTPID%>" target="_blank"><%=Model.Bugg.TTPID%></a>--%>
 			<%=Html.ActionLink( Bugg.Id.ToString(), "Show", new { controller = "Buggs", id = Bugg.Id }, null )%>
 		</td>
-		<td>
-			<%--NewBugg.TTPID = RealBugg.TTPID;							// JIRA--%>
-			<a href="https://jira.ol.epicgames.net/browse/<%=Bugg.TTPID%>" target="_blank"><%=Bugg.TTPID%></a>
-		</td>
+
+		<%if( string.IsNullOrEmpty( Bugg.TTPID ) && !string.IsNullOrEmpty(Bugg.Pattern ) )
+		{ %>
+			<td>
+				<input type="submit" value="CopyToJira" title="<%=Bugg.ToTooltip()%>" id="id <%=Bugg.Id%>" name="<%=Bugg.Id%>" class="CopyToJiraStyle" />
+			</td>
+		<% } else if( string.IsNullOrEmpty( Bugg.TTPID ) && string.IsNullOrEmpty(Bugg.Pattern ) )
+		{ %>
+			<td>
+				No pattern!
+			</td>
+		<% } else { %>
+			<td>
+				<%--NewBugg.TTPID = RealBugg.TTPID;							// JIRA--%>
+				<a href="https://jira.ol.epicgames.net/browse/<%=Bugg.TTPID%>" target="_blank"><%=Bugg.TTPID%></a>
+			</td>
+		<% } %>
+
 		<td>
 			<%--NewBugg.FixedChangeList = RealBugg.FixedChangeList;		// FixCL--%>
 			<%=Bugg.FixedChangeList%>
@@ -141,7 +158,7 @@ public int JiraFixCL { get; set; }
 			<%=Bugg.NumberOfCrashes%>
 		</td>
 		<td>
-			<%--NewBugg.BuildVersion = NewBugg.AffectedBuilds.Last();	// Latest Version Affected--%>
+			<%--NewBugg.BuildVersion = NewBugg.AffectedVersions.Last();	// Latest Version Affected--%>
 			<%=Bugg.BuildVersion%>
 		</td>
 		<td>
@@ -174,9 +191,19 @@ public int JiraFixCL { get; set; }
 		<td>
 			<%=Bugg.JiraFixVersionsText%>
 		</td>
-		<td>
-			<%=Bugg.JiraFixCL%>
-		</td>
+		
+		<%if( !string.IsNullOrEmpty( Bugg.JiraFixCL ) )
+		{ %>
+			<td style="background-color:blue">
+				<%=Bugg.JiraFixCL%>
+			</td>
+		<%}
+		else
+		{%>
+			<td>
+			</td>
+		<%} %>
+
 	</tr>
 	<%
 		}	
@@ -184,4 +211,7 @@ public int JiraFixCL { get; set; }
 
 </table>
 
-<small>Generated in <%=Model.GenerationTime%> second(s)</small>
\ No newline at end of file
+<small>Generated in <%=Model.GenerationTime%> second(s)</small>
+
+</form>
+</body>
\ No newline at end of file
