CrashReportWebsite - More stats

diff --git a/Engine/Source/Programs/CrashReporter/CrashReportWebSite/Models/CrashRepository.cs b/Engine/Source/Programs/CrashReporter/CrashReportWebSite/Models/CrashRepository.cs
--- a/Engine/Source/Programs/CrashReporter/CrashReportWebSite/Models/CrashRepository.cs
+++ b/Engine/Source/Programs/CrashReporter/CrashReportWebSite/Models/CrashRepository.cs
@@ -923,48 +923,51 @@ namespace Tools.CrashReporter.CrashReportWebSite.Models
 		/// <param name="CrashInstance">A crash that was recently added to the database.</param>
 		public void BuildPattern( Crash CrashInstance )
 		{
-			List<string> Pattern = new List<string>();
+			using( FAutoScopedLogTimer LogTimer = new FAutoScopedLogTimer( this.GetType().ToString() + "(Crash=" + CrashInstance.Id + ")" ) )
+			{
+				List<string> Pattern = new List<string>();
 
-			// Get an array of callstack items
-			CallStackContainer CallStack = new CallStackContainer( CrashInstance );
-			CallStack.bDisplayFunctionNames = true;
+				// Get an array of callstack items
+				CallStackContainer CallStack = new CallStackContainer( CrashInstance );
+				CallStack.bDisplayFunctionNames = true;
 
-			if( CrashInstance.Pattern == null )
-			{
-				// Set the module based on the modules in the callstack
-				CrashInstance.Module = CallStack.GetModuleName();
-				try
+				if( CrashInstance.Pattern == null )
 				{
-					foreach( CallStackEntry Entry in CallStack.CallStackEntries )
+					// Set the module based on the modules in the callstack
+					CrashInstance.Module = CallStack.GetModuleName();
+					try
 					{
-						FunctionCall CurrentFunctionCall = new FunctionCall();
-
-						if( FunctionCalls.Where( f => f.Call == Entry.FunctionName ).Count() > 0 )
-						{
-							CurrentFunctionCall = FunctionCalls.Where( f => f.Call == Entry.FunctionName ).First();
-						}
-						else
+						foreach( CallStackEntry Entry in CallStack.CallStackEntries )
 						{
-							CurrentFunctionCall = new FunctionCall();
-							CurrentFunctionCall.Call = Entry.FunctionName;
-							FunctionCalls.InsertOnSubmit( CurrentFunctionCall );
-						}
+							FunctionCall CurrentFunctionCall = new FunctionCall();
 
-						SubmitChanges();
+							if( FunctionCalls.Where( f => f.Call == Entry.FunctionName ).Count() > 0 )
+							{
+								CurrentFunctionCall = FunctionCalls.Where( f => f.Call == Entry.FunctionName ).First();
+							}
+							else
+							{
+								CurrentFunctionCall = new FunctionCall();
+								CurrentFunctionCall.Call = Entry.FunctionName;
+								FunctionCalls.InsertOnSubmit( CurrentFunctionCall );
+							}
 
-						Pattern.Add( CurrentFunctionCall.Id.ToString() );
-					}
+							SubmitChanges();
 
-					//CrashInstance.Pattern = "+";
-					CrashInstance.Pattern = string.Join( "+", Pattern );
-					// We need something like this +1+2+3+5+ for searching for exact pattern like +5+
-					//CrashInstance.Pattern += "+";
+							Pattern.Add( CurrentFunctionCall.Id.ToString() );
+						}
 
-					SubmitChanges();
-				}
-				catch( Exception Ex )
-				{
-					FLogger.WriteException( "BuildPattern: " + Ex.ToString() );
+						//CrashInstance.Pattern = "+";
+						CrashInstance.Pattern = string.Join( "+", Pattern );
+						// We need something like this +1+2+3+5+ for searching for exact pattern like +5+
+						//CrashInstance.Pattern += "+";
+
+						SubmitChanges();
+					}
+					catch( Exception Ex )
+					{
+						FLogger.WriteException( "BuildPattern: " + Ex.ToString() );
+					}
 				}
 			}
 		}
