CrashReport - Enabled writing a Diagnostics.txt to disk in UE4 release build, if a callstack is valid, it will sent to the server.

diff --git a/Engine/Source/Programs/CrashReporter/CrashReportClient/Private/Windows/WindowsErrorReport.cpp b/Engine/Source/Programs/CrashReporter/CrashReportClient/Private/Windows/WindowsErrorReport.cpp
--- a/Engine/Source/Programs/CrashReporter/CrashReportClient/Private/Windows/WindowsErrorReport.cpp
+++ b/Engine/Source/Programs/CrashReporter/CrashReportClient/Private/Windows/WindowsErrorReport.cpp
@@ -102,12 +102,8 @@ FText FWindowsErrorReport::DiagnoseReport() const
 		return LOCTEXT("NoDebuggingSymbols", "You do not have any debugging symbols required to display the callstack for this crash.");
 	}
 
-	// Don't write a Diagnostics.txt to disk in UE4 release build. It will be displayed in the UI but not sent to the server.
-	if ( !FRocketSupport::IsRocket() )
-	{
-		// There's a callstack, so write it out to save the server trying to do it
-		CrashDebugHelper->CrashInfo.GenerateReport(ReportDirectory / GDiagnosticsFilename);
-	}
+	// There's a callstack, so write it out to save the server trying to do it
+	CrashDebugHelper->CrashInfo.GenerateReport(ReportDirectory / GDiagnosticsFilename);
 
 	const auto& Exception = CrashDebugHelper->CrashInfo.Exception;
 	const FString Assertion = FWindowsReportParser::Find( ReportDirectory, TEXT( "AssertLog=" ) );
