CrashReport - Logs are always uploaded, even for UE4 releases

diff --git a/Engine/Source/Programs/CrashReporter/CrashReportClient/Private/GenericErrorReport.cpp b/Engine/Source/Programs/CrashReporter/CrashReportClient/Private/GenericErrorReport.cpp
--- a/Engine/Source/Programs/CrashReporter/CrashReportClient/Private/GenericErrorReport.cpp
+++ b/Engine/Source/Programs/CrashReporter/CrashReportClient/Private/GenericErrorReport.cpp
@@ -118,12 +118,6 @@ TArray<FString> FGenericErrorReport::GetFilesToUpload() const
 
 	for (const auto& Filename: ReportFilenames)
 	{
-		if (FRocketSupport::IsRocket() && Filename.EndsWith(TEXT(".log")))
-		{
-			// Temporarily side-step privacy concerns by not uploading UE4 release builds' logs
-			continue;
-		}
-
 		FilesToUpload.Push(ReportDirectory / Filename);
 	}
 	return FilesToUpload;
