CrashReportClient - Improved app shutting down

diff --git a/Engine/Source/Programs/CrashReporter/CrashReportClient/Private/CrashReportClientApp.cpp b/Engine/Source/Programs/CrashReporter/CrashReportClient/Private/CrashReportClientApp.cpp
--- a/Engine/Source/Programs/CrashReporter/CrashReportClient/Private/CrashReportClientApp.cpp
+++ b/Engine/Source/Programs/CrashReporter/CrashReportClient/Private/CrashReportClientApp.cpp
@@ -196,7 +196,11 @@ void RunCrashReportClient(const TCHAR* CommandLine)
 	}
 
 	FPlatformErrorReport::ShutDown();
-	GEngineLoop.AppExit();
+
+	FEngineLoop::AppPreExit();
+	FTaskGraphInterface::Shutdown();
+
+	FEngineLoop::AppExit();
 }
 
 void CrashReportClientCheck(bool bCondition, const TCHAR* Location)
