CrashReportClient - Added basic configuration for the crash report receiver

diff --git a/Engine/Config/BaseEngine.ini b/Engine/Config/BaseEngine.ini
--- a/Engine/Config/BaseEngine.ini
+++ b/Engine/Config/BaseEngine.ini
@@ -2230,4 +2230,9 @@ NameForGameView5="GameView5"
 ; These are the defaults for HTML5 settings, and they need to be in the .ini since UBT reads the .ini settings, without instantiating the class
 [/Script/HTML5PlatformEditor.HTML5TargetSettings]
 HeapSizeDevelopment=1024
-HeapSizeShipping=1024DeployServerPort=53501
\ No newline at end of file
+HeapSizeShipping=1024
+DeployServerPort=53501
+
+[CrashReportClient]
+;CrashReportReceiverIP="http://localhost:57005"
+CrashReportReceiverIP="http://crashreporter.epicgames.com:57005"
\ No newline at end of file
diff --git a/Engine/Source/Programs/CrashReporter/CrashReportClient/Private/CrashReportClient.cpp b/Engine/Source/Programs/CrashReporter/CrashReportClient/Private/CrashReportClient.cpp
--- a/Engine/Source/Programs/CrashReporter/CrashReportClient/Private/CrashReportClient.cpp
+++ b/Engine/Source/Programs/CrashReporter/CrashReportClient/Private/CrashReportClient.cpp
@@ -12,18 +12,6 @@
 
 #define LOCTEXT_NAMESPACE "CrashReportClient"
 
-//#define DO_LOCAL_TESTING 1
-
-#if	DO_LOCAL_TESTING
-	const TCHAR* GServerIP = TEXT( "http://localhost:57005" );
-#else
-	const TCHAR* GServerIP = TEXT( "http://crashreporter.epicgames.com:57005" );
-#endif // DO_LOCAL_TESTING
-
-// Must match filename specified in RunMinidumpDiagnostics
-const TCHAR* GDiagnosticsFilename = TEXT("Diagnostics.txt");
-
-
 FCrashDescription& GetCrashDescription()
 {
 	static FCrashDescription Singleton;
@@ -36,7 +24,7 @@ FCrashReportClient::FCrashReportClient(const FPlatformErrorReport& InErrorReport
 	: DiagnosticText( LOCTEXT("ProcessingReport", "Processing crash report ...") )
 	, DiagnoseReportTask(nullptr)
 	, ErrorReport( InErrorReport )
-	, Uploader(GServerIP)
+	, Uploader( FCrashReportClientConfig::Get().GetReceiverAddress() )
 	, bBeginUploadCalled(false)
 	, bShouldWindowBeHidden(false)
 	, bAllowToBeContacted(true)
@@ -183,7 +171,7 @@ void FCrashReportClient::FinalizeDiagnoseReportWorker( FText ReportText )
 {
 	DiagnosticText = FCrashReportUtil::FormatDiagnosticText( ReportText, GetCrashDescription().MachineId, GetCrashDescription().EpicAccountId, GetCrashDescription().UserName );
 
-	auto DiagnosticsFilePath = ErrorReport.GetReportDirectory() / GDiagnosticsFilename;
+	auto DiagnosticsFilePath = ErrorReport.GetReportDirectory() / FCrashReportClientConfig::Get().GetDiagnosticsFilename();
 	Uploader.LocalDiagnosisComplete(FPaths::FileExists(DiagnosticsFilePath) ? DiagnosticsFilePath : TEXT(""));
 }
 
diff --git a/Engine/Source/Programs/CrashReporter/CrashReportClient/Private/CrashReportClientApp.cpp b/Engine/Source/Programs/CrashReporter/CrashReportClient/Private/CrashReportClientApp.cpp
--- a/Engine/Source/Programs/CrashReporter/CrashReportClient/Private/CrashReportClientApp.cpp
+++ b/Engine/Source/Programs/CrashReporter/CrashReportClient/Private/CrashReportClientApp.cpp
@@ -92,6 +92,16 @@ FPlatformErrorReport LoadErrorReport()
 #endif
 }
 
+FCrashReportClientConfig::FCrashReportClientConfig()
+	: DiagnosticsFilename( TEXT( "Diagnostics.txt" ))
+{
+	if( !GConfig->GetString( TEXT( "CrashReportClient" ), TEXT( "CrashReportReceiverIP" ), CrashReportReceiverIP, GEngineIni ) )
+	{
+		// Use the default value.
+		CrashReportReceiverIP = TEXT( "http://crashreporter.epicgames.com:57005" );
+	}
+}
+
 void RunCrashReportClient(const TCHAR* CommandLine)
 {
 	// Override the stack size for the thread pool.
@@ -100,6 +110,9 @@ void RunCrashReportClient(const TCHAR* CommandLine)
 	// Set up the main loop
 	GEngineLoop.PreInit(CommandLine);
 
+	// Initialize config.
+	FCrashReportClientConfig::Get();
+
 	const bool bUnattended = 
 #if CRASH_REPORT_UNATTENDED_ONLY
 		true;
diff --git a/Engine/Source/Programs/CrashReporter/CrashReportClient/Private/CrashReportClientUnattended.cpp b/Engine/Source/Programs/CrashReporter/CrashReportClient/Private/CrashReportClientUnattended.cpp
--- a/Engine/Source/Programs/CrashReporter/CrashReportClient/Private/CrashReportClientUnattended.cpp
+++ b/Engine/Source/Programs/CrashReporter/CrashReportClient/Private/CrashReportClientUnattended.cpp
@@ -4,7 +4,7 @@
 #include "CrashReportClientUnattended.h"
 
 FCrashReportClientUnattended::FCrashReportClientUnattended(const FPlatformErrorReport& ErrorReport)
-	: Uploader(GServerIP)
+	: Uploader( FCrashReportClientConfig::Get().GetReceiverAddress() )
 {
 	// Prevent uploader waiting for a diagnosis file
 	Uploader.LocalDiagnosisSkipped();	
diff --git a/Engine/Source/Programs/CrashReporter/CrashReportClient/Private/CrashReportUtil.h b/Engine/Source/Programs/CrashReporter/CrashReportClient/Private/CrashReportUtil.h
--- a/Engine/Source/Programs/CrashReporter/CrashReportClient/Private/CrashReportUtil.h
+++ b/Engine/Source/Programs/CrashReporter/CrashReportClient/Private/CrashReportUtil.h
@@ -2,6 +2,40 @@
 
 #pragma once
 
+#include "Core.h"
+
+/** Holds basic configuration for the crash report client. */
+struct FCrashReportClientConfig
+{
+	/** Accesses the singleton. */
+	static FCrashReportClientConfig& Get()
+	{
+		static FCrashReportClientConfig Instance;
+		return Instance;
+	}
+
+	/** Initialization constructor. */
+	FCrashReportClientConfig();
+
+	const FString& GetReceiverAddress() const
+	{
+		return CrashReportReceiverIP;
+	}
+
+	const FString& GetDiagnosticsFilename() const
+	{
+		return DiagnosticsFilename;
+	}
+
+protected:
+	/** IP address of crash report receiver. */
+	FString CrashReportReceiverIP;
+
+	/** Filename to use when saving diagnostics report, if generated locally. */
+	FString DiagnosticsFilename;
+};
+
+
 /**
  * Helper class for MakeDirectoryVisitor
  */
diff --git a/Engine/Source/Programs/CrashReporter/CrashReportClient/Private/GenericErrorReport.cpp b/Engine/Source/Programs/CrashReporter/CrashReportClient/Private/GenericErrorReport.cpp
--- a/Engine/Source/Programs/CrashReporter/CrashReportClient/Private/GenericErrorReport.cpp
+++ b/Engine/Source/Programs/CrashReporter/CrashReportClient/Private/GenericErrorReport.cpp
@@ -138,7 +138,7 @@ bool FGenericErrorReport::LoadWindowsReportXmlFile( FString& OutString ) const
 bool FGenericErrorReport::TryReadDiagnosticsFile(FText& OutReportDescription)
 {
 	FString FileContent;
-	if (!FFileHelper::LoadFileToString(FileContent, *(ReportDirectory / GDiagnosticsFilename)))
+	if (!FFileHelper::LoadFileToString(FileContent, *(ReportDirectory / FCrashReportClientConfig::Get().GetDiagnosticsFilename())))
 	{
 		// No diagnostics file
 		return false;
diff --git a/Engine/Source/Programs/CrashReporter/CrashReportClient/Private/Mac/MacErrorReport.cpp b/Engine/Source/Programs/CrashReporter/CrashReportClient/Private/Mac/MacErrorReport.cpp
--- a/Engine/Source/Programs/CrashReporter/CrashReportClient/Private/Mac/MacErrorReport.cpp
+++ b/Engine/Source/Programs/CrashReporter/CrashReportClient/Private/Mac/MacErrorReport.cpp
@@ -123,9 +123,9 @@ FText FMacErrorReport::DiagnoseReport() const
 	else
 	{
 		FString CrashDump;
-		FString DiagnosticsPath = ReportDirectory / FString( TEXT( "Diagnostics.txt" ) );
+		FString DiagnosticsPath = ReportDirectory / FCrashReportClientConfig::Get().GetDiagnosticsFilename();
 		CrashDebugHelper->CrashInfo.GenerateReport( DiagnosticsPath );
-		if ( FFileHelper::LoadFileToString( CrashDump, *(ReportDirectory / TEXT("Diagnostics.txt")) ) )
+		if ( FFileHelper::LoadFileToString( CrashDump, *(ReportDirectory / FCrashReportClientConfig::Get().GetDiagnosticsFilename() ) )
 		{
 			return FText::FromString(CrashDump);
 		}
diff --git a/Engine/Source/Programs/CrashReporter/CrashReportClient/Private/Windows/WindowsErrorReport.cpp b/Engine/Source/Programs/CrashReporter/CrashReportClient/Private/Windows/WindowsErrorReport.cpp
--- a/Engine/Source/Programs/CrashReporter/CrashReportClient/Private/Windows/WindowsErrorReport.cpp
+++ b/Engine/Source/Programs/CrashReporter/CrashReportClient/Private/Windows/WindowsErrorReport.cpp
@@ -103,7 +103,7 @@ FText FWindowsErrorReport::DiagnoseReport() const
 	}
 
 	// There's a callstack, so write it out to save the server trying to do it
-	CrashDebugHelper->CrashInfo.GenerateReport(ReportDirectory / GDiagnosticsFilename);
+	CrashDebugHelper->CrashInfo.GenerateReport(ReportDirectory / FCrashReportClientConfig::Get().GetDiagnosticsFilename());
 
 	const auto& Exception = CrashDebugHelper->CrashInfo.Exception;
 	const FString Assertion = FWindowsReportParser::Find( ReportDirectory, TEXT( "AssertLog=" ) );
diff --git a/Engine/Source/Programs/CrashReporter/CrashReportClient/Public/CrashReportClientApp.h b/Engine/Source/Programs/CrashReporter/CrashReportClient/Public/CrashReportClientApp.h
--- a/Engine/Source/Programs/CrashReporter/CrashReportClient/Public/CrashReportClientApp.h
+++ b/Engine/Source/Programs/CrashReporter/CrashReportClient/Public/CrashReportClientApp.h
@@ -19,15 +19,6 @@ DECLARE_LOG_CATEGORY_EXTERN(CrashReportClientLog, Log, All)
 
 #define FILE_LINE_STRING TEXT(__FILE__) TEXT("(") TEXT(STRINGIZE(__LINE__)) TEXT(")")
 
-/** IP of server to send report to */
-extern const TCHAR* GServerIP;
-
-/** Filename to use when saving diagnostics report (if generated locally) */
-extern const TCHAR* GDiagnosticsFilename;
-
-/** A user ID and/or name to add to the report and display */
-extern FString GCrashUserId;
-
 /** This writes and error to the log rather than calling 'check': don't want the crash reporter crashing */
 void CrashReportClientCheck(bool bCondition, const TCHAR* Location);
 #define CRASHREPORTCLIENT_CHECK(COND) CrashReportClientCheck(COND, FILE_LINE_STRING)
