cis fix

diff --git a/Engine/Source/Programs/CrashReporter/CrashReportClient/Private/Mac/MacErrorReport.cpp b/Engine/Source/Programs/CrashReporter/CrashReportClient/Private/Mac/MacErrorReport.cpp
--- a/Engine/Source/Programs/CrashReporter/CrashReportClient/Private/Mac/MacErrorReport.cpp
+++ b/Engine/Source/Programs/CrashReporter/CrashReportClient/Private/Mac/MacErrorReport.cpp
@@ -125,7 +125,7 @@ FText FMacErrorReport::DiagnoseReport() const
 		FString CrashDump;
 		FString DiagnosticsPath = ReportDirectory / FCrashReportClientConfig::Get().GetDiagnosticsFilename();
 		CrashDebugHelper->CrashInfo.GenerateReport( DiagnosticsPath );
-		if ( FFileHelper::LoadFileToString( CrashDump, *(ReportDirectory / FCrashReportClientConfig::Get().GetDiagnosticsFilename() ) )
+		if ( FFileHelper::LoadFileToString( CrashDump, *(ReportDirectory / FCrashReportClientConfig::Get().GetDiagnosticsFilename() ) ) )
 		{
 			return FText::FromString(CrashDump);
 		}
