CrashReportClient - Minor tweak

diff --git a/Engine/Source/Programs/CrashReporter/CrashReportClient/Private/CrashReportClientApp.cpp b/Engine/Source/Programs/CrashReporter/CrashReportClient/Private/CrashReportClientApp.cpp
--- a/Engine/Source/Programs/CrashReporter/CrashReportClient/Private/CrashReportClientApp.cpp
+++ b/Engine/Source/Programs/CrashReporter/CrashReportClient/Private/CrashReportClientApp.cpp
@@ -100,6 +100,8 @@ FCrashReportClientConfig::FCrashReportClientConfig()
 		// Use the default value.
 		CrashReportReceiverIP = TEXT( "http://crashreporter.epicgames.com:57005" );
 	}
+
+	UE_LOG( CrashReportClientLog, Log, TEXT( "CrashReportReceiverIP: %s" ), *CrashReportReceiverIP );
 }
 
 void RunCrashReportClient(const TCHAR* CommandLine)
