CrashReport - Removed obsolete tables from the DB

diff --git a/Engine/Source/Programs/CrashReporter/CrashReportWebSite/Models/CreateDatabase.sql b/Engine/Source/Programs/CrashReporter/CrashReportWebSite/Models/CreateDatabase.sql
--- a/Engine/Source/Programs/CrashReporter/CrashReportWebSite/Models/CreateDatabase.sql
+++ b/Engine/Source/Programs/CrashReporter/CrashReportWebSite/Models/CreateDatabase.sql
@@ -1,4 +1,177 @@
-﻿USE [CrashReport]GO/****** Object:  Table [dbo].[Users]    Script Date: 27/01/2014 13:59:11 ******/SET ANSI_NULLS ONGOSET QUOTED_IDENTIFIER ONGOSET ANSI_PADDING ONGOCREATE TABLE [dbo].[Users](	[UserName] [varchar](64) NOT NULL,	[UserGroup] [varchar](50) NOT NULL,	[UserGroupId] [int] NULL, CONSTRAINT [PK_Users] PRIMARY KEY CLUSTERED (	[UserName] ASC)WITH (PAD_INDEX  = OFF, STATISTICS_NORECOMPUTE  = OFF, IGNORE_DUP_KEY  = OFF, ALLOW_ROW_LOCKS  = ON, ALLOW_PAGE_LOCKS  = ON) ON [PRIMARY]) ON [PRIMARY]GOSET ANSI_PADDING OFFGO/****** Object:  Table [dbo].[Crashes]    Script Date: 27/01/2014 13:57:24 ******/SET ANSI_NULLS ONGOSET QUOTED_IDENTIFIER ONGOSET ANSI_PADDING ONGOCREATE TABLE [dbo].[Crashes](	[Id] [int] IDENTITY(1,1) NOT NULL,	[Title] [nchar](20) NULL,	[Summary] [varchar](512) NULL,	[GameName] [varchar](64) NULL,	[Status] [varchar](64) NULL,	[TimeOfCrash] [datetime] NULL,	[ChangeListVersion] [varchar](64) NULL,	[PlatformName] [varchar](64) NULL,	[EngineMode] [varchar](64) NULL,	[Description] [varchar](512) NULL,	[RawCallStack] [varchar](max) NULL,	[Pattern] [varchar](800) NULL,	[CommandLine] [varchar](512) NULL,	[ComputerName] [varchar](64) NULL,	[Selected] [bit] NULL,	[FixedChangeList] [varchar](64) NULL,	[LanguageExt] [varchar](64) NULL,	[Module] [varchar](128) NULL,	[BuildVersion] [varchar](128) NULL,	[BaseDir] [varchar](512) NULL,	[Version] [int] NULL,	[UserName] [varchar](64) NULL,	[TTPID] [varchar](64) NULL,	[AutoReporterID] [int] NULL,	[Processed] [bit] NULL,	[HasLogFile] [bit] NULL,	[HasMiniDumpFile] [bit] NULL,	[HasVideoFile] [bit] NULL,	[HasDiagnosticsFile] [bit] NULL,	[HasNewLogFile] [bit] NULL,	[Branch] [varchar](32) NULL,	[CrashType] [smallint] NOT NULL,	[UserNameId] [int] NULL,	[HasMetaData] [bit] NULL,	[SourceContext] [varchar](max) NULL,	[EpicAccountId] [varchar](64) NULL,	[EngineVersion] [varchar](64) NULL, CONSTRAINT [PK_Crashes] PRIMARY KEY CLUSTERED (	[Id] ASC)WITH (PAD_INDEX = OFF, STATISTICS_NORECOMPUTE = OFF, IGNORE_DUP_KEY = OFF, ALLOW_ROW_LOCKS = ON, ALLOW_PAGE_LOCKS = ON) ON [PRIMARY]) ON [PRIMARY] TEXTIMAGE_ON [PRIMARY]GOSET ANSI_PADDING OFFGOEXEC sys.sp_addextendedproperty @name=N'MS_Description', @value=N'NOT USED OR OBSOLETE' , @level0type=N'SCHEMA',@level0name=N'dbo', @level1type=N'TABLE',@level1name=N'Crashes', @level2type=N'COLUMN',@level2name=N'Title'EXEC sys.sp_addextendedproperty @name=N'MS_Description', @value=N'NOT USED OR OBSOLETE' , @level0type=N'SCHEMA',@level0name=N'dbo', @level1type=N'TABLE',@level1name=N'Crashes', @level2type=N'COLUMN',@level2name=N'Selected'EXEC sys.sp_addextendedproperty @name=N'MS_Description', @value=N'NOT USED OR OBSOLETE' , @level0type=N'SCHEMA',@level0name=N'dbo', @level1type=N'TABLE',@level1name=N'Crashes', @level2type=N'COLUMN',@level2name=N'Version'EXEC sys.sp_addextendedproperty @name=N'MS_Description', @value=N'NOT USED OR OBSOLETE' , @level0type=N'SCHEMA',@level0name=N'dbo', @level1type=N'TABLE',@level1name=N'Crashes', @level2type=N'COLUMN',@level2name=N'AutoReporterID'EXEC sys.sp_addextendedproperty @name=N'MS_Description', @value=N'NOT USED OR OBSOLETE' , @level0type=N'SCHEMA',@level0name=N'dbo', @level1type=N'TABLE',@level1name=N'Crashes', @level2type=N'COLUMN',@level2name=N'Processed'EXEC sys.sp_addextendedproperty @name=N'MS_Description', @value=N'NOT USED OR OBSOLETE' , @level0type=N'SCHEMA',@level0name=N'dbo', @level1type=N'TABLE',@level1name=N'Crashes', @level2type=N'COLUMN',@level2name=N'HasDiagnosticsFile'EXEC sys.sp_addextendedproperty @name=N'MS_Description', @value=N'NOT USED OR OBSOLETE' , @level0type=N'SCHEMA',@level0name=N'dbo', @level1type=N'TABLE',@level1name=N'Crashes', @level2type=N'COLUMN',@level2name=N'HasNewLogFile'EXEC sys.sp_addextendedproperty @name=N'MS_Description', @value=N'The name of the game that crashed. (FCrashDescription.GameName)' , @level0type=N'SCHEMA',@level0name=N'dbo', @level1type=N'TABLE',@level1name=N'Crashes', @level2type=N'COLUMN',@level2name=N'GameName'EXEC sys.sp_addextendedproperty @name=N'MS_Description', @value=N'The mode the game was in e.g. editor. (FCrashDescription.EngineMode)' , @level0type=N'SCHEMA',@level0name=N'dbo', @level1type=N'TABLE',@level1name=N'Crashes', @level2type=N'COLUMN',@level2name=N'EngineMode'EXEC sys.sp_addextendedproperty @name=N'MS_Description', @value=N'The platform that crashed e.g. Win64. (FCrashDescription.Platform)' , @level0type=N'SCHEMA',@level0name=N'dbo', @level1type=N'TABLE',@level1name=N'Crashes', @level2type=N'COLUMN',@level2name=N'PlatformName'EXEC sys.sp_addextendedproperty @name=N'MS_Description', @value=N'Encoded engine version.  BuildVersion-BuiltFromCL-BranchName E.g. 4.3.0.0-2215663+UE4-Releases+4.3 (FCrashDescription.EngineVersion)' , @level0type=N'SCHEMA',@level0name=N'dbo', @level1type=N'TABLE',@level1name=N'Crashes', @level2type=N'COLUMN',@level2name=N'EngineVersion'EXEC sys.sp_addextendedproperty @name=N'MS_Description', @value=N'The three component version of the app e.g. 4.4.1 (FCrashDescription.BuildVersion)' , @level0type=N'SCHEMA',@level0name=N'dbo', @level1type=N'TABLE',@level1name=N'Crashes', @level2type=N'COLUMN',@level2name=N'BuildVersion'EXEC sys.sp_addextendedproperty @name=N'MS_Description', @value=N'Built from changelist. (FCrashDescription.BuiltFromCL)' , @level0type=N'SCHEMA',@level0name=N'dbo', @level1type=N'TABLE',@level1name=N'Crashes', @level2type=N'COLUMN',@level2name=N'ChangeListVersion'EXEC sys.sp_addextendedproperty @name=N'MS_Description', @value=N'The name of the branch this game was built out of. (FCrashDescription.BranchName)' , @level0type=N'SCHEMA',@level0name=N'dbo', @level1type=N'TABLE',@level1name=N'Crashes', @level2type=N'COLUMN',@level2name=N'Branch'EXEC sys.sp_addextendedproperty @name=N'MS_Description', @value=N'The command line of the application that crashed. (FCrashDescription.CommandLine)' , @level0type=N'SCHEMA',@level0name=N'dbo', @level1type=N'TABLE',@level1name=N'Crashes', @level2type=N'COLUMN',@level2name=N'CommandLine'EXEC sys.sp_addextendedproperty @name=N'MS_Description', @value=N'The base directory where the app was running. (FCrashDescription.BaseDir)' , @level0type=N'SCHEMA',@level0name=N'dbo', @level1type=N'TABLE',@level1name=N'Crashes', @level2type=N'COLUMN',@level2name=N'BaseDir'EXEC sys.sp_addextendedproperty @name=N'MS_Description', @value=N'The language of the system. (FCrashDescription.Language)' , @level0type=N'SCHEMA',@level0name=N'dbo', @level1type=N'TABLE',@level1name=N'Crashes', @level2type=N'COLUMN',@level2name=N'LanguageExt'EXEC sys.sp_addextendedproperty @name=N'MS_Description', @value=N'The name of the user that caused this crash. (FCrashDescription.UserName)' , @level0type=N'SCHEMA',@level0name=N'dbo', @level1type=N'TABLE',@level1name=N'Crashes', @level2type=N'COLUMN',@level2name=N'UserName'EXEC sys.sp_addextendedproperty @name=N'MS_Description', @value=N'The unique ID used to identify the machine the crash occurred on. (FCrashDescription.MachineGuid)' , @level0type=N'SCHEMA',@level0name=N'dbo', @level1type=N'TABLE',@level1name=N'Crashes', @level2type=N'COLUMN',@level2name=N'ComputerName'EXEC sys.sp_addextendedproperty @name=N'MS_Description', @value=N'The Epic account ID for the user who last used the Launcher. (FCrashDescription.EpicAccountId)' , @level0type=N'SCHEMA',@level0name=N'dbo', @level1type=N'TABLE',@level1name=N'Crashes', @level2type=N'COLUMN',@level2name=N'EpicAccountId'EXEC sys.sp_addextendedproperty @name=N'MS_Description', @value=N'An array of strings representing the callstack of the crash. (FCrashDescription.CallStack)' , @level0type=N'SCHEMA',@level0name=N'dbo', @level1type=N'TABLE',@level1name=N'Crashes', @level2type=N'COLUMN',@level2name=N'RawCallStack'EXEC sys.sp_addextendedproperty @name=N'MS_Description', @value=N'An array of strings showing the source code around the crash. (FCrashDescription.SourceContext)' , @level0type=N'SCHEMA',@level0name=N'dbo', @level1type=N'TABLE',@level1name=N'Crashes', @level2type=N'COLUMN',@level2name=N'SourceContext'EXEC sys.sp_addextendedproperty @name=N'MS_Description', @value=N'An array of strings representing the user description of the crash. (FCrashDescription.UserDescription)' , @level0type=N'SCHEMA',@level0name=N'dbo', @level1type=N'TABLE',@level1name=N'Crashes', @level2type=N'COLUMN',@level2name=N'Description'EXEC sys.sp_addextendedproperty @name=N'MS_Description', @value=N'The error message, can be assertion message, ensure message or message from the fatal error. (FCrashDescription.ErrorMessage)' , @level0type=N'SCHEMA',@level0name=N'dbo', @level1type=N'TABLE',@level1name=N'Crashes', @level2type=N'COLUMN',@level2name=N'Summary'EXEC sys.sp_addextendedproperty @name=N'MS_Description', @value=N'The UTC time the crash occurred. (FCrashDescription.TimeofCrash)' , @level0type=N'SCHEMA',@level0name=N'dbo', @level1type=N'TABLE',@level1name=N'Crashes', @level2type=N'COLUMN',@level2name=N'TimeOfCrash'EXEC sys.sp_addextendedproperty @name=N'MS_Description', @value=N'Whether this crash has a minidump file. (FCrashDescription.bHasMiniDumpFile)' , @level0type=N'SCHEMA',@level0name=N'dbo', @level1type=N'TABLE',@level1name=N'Crashes', @level2type=N'COLUMN',@level2name=N'HasMiniDumpFile'EXEC sys.sp_addextendedproperty @name=N'MS_Description', @value=N'Whether this crash has a log file. (FCrashDescription.bHasLogFile)' , @level0type=N'SCHEMA',@level0name=N'dbo', @level1type=N'TABLE',@level1name=N'Crashes', @level2type=N'COLUMN',@level2name=N'HasLogFile'EXEC sys.sp_addextendedproperty @name=N'MS_Description', @value=N'Whether this crash has a video file. (FCrashDescription.bHasVideoFile)' , @level0type=N'SCHEMA',@level0name=N'dbo', @level1type=N'TABLE',@level1name=N'Crashes', @level2type=N'COLUMN',@level2name=N'HasVideoFile'EXEC sys.sp_addextendedproperty @name=N'MS_Description', @value=N'Status of the crash' , @level0type=N'SCHEMA',@level0name=N'dbo', @level1type=N'TABLE',@level1name=N'Crashes', @level2type=N'COLUMN',@level2name=N'Status'EXEC sys.sp_addextendedproperty @name=N'MS_Description', @value=N'Pattern used to group cashes into buggs' , @level0type=N'SCHEMA',@level0name=N'dbo', @level1type=N'TABLE',@level1name=N'Crashes', @level2type=N'COLUMN',@level2name=N'Pattern'EXEC sys.sp_addextendedproperty @name=N'MS_Description', @value=N'Changelist that fixed this crash' , @level0type=N'SCHEMA',@level0name=N'dbo', @level1type=N'TABLE',@level1name=N'Crashes', @level2type=N'COLUMN',@level2name=N'FixedChangeList'EXEC sys.sp_addextendedproperty @name=N'MS_Description', @value=N'Name of the module where the crash occurred' , @level0type=N'SCHEMA',@level0name=N'dbo', @level1type=N'TABLE',@level1name=N'Crashes', @level2type=N'COLUMN',@level2name=N'Module'EXEC sys.sp_addextendedproperty @name=N'MS_Description', @value=N'The ID of the user name' , @level0type=N'SCHEMA',@level0name=N'dbo', @level1type=N'TABLE',@level1name=N'Crashes', @level2type=N'COLUMN',@level2name=N'UserNameId'EXEC sys.sp_addextendedproperty @name=N'MS_Description', @value=N'The ID of the TTP' , @level0type=N'SCHEMA',@level0name=N'dbo', @level1type=N'TABLE',@level1name=N'Crashes', @level2type=N'COLUMN',@level2name=N'TTPID'EXEC sys.sp_addextendedproperty @name=N'MS_Description', @value=N'The type of the crash, crash(1), assert(2), ensure(3)' , @level0type=N'SCHEMA',@level0name=N'dbo', @level1type=N'TABLE',@level1name=N'Crashes', @level2type=N'COLUMN',@level2name=N'CrashType'/****** Object:  Table [dbo].[UserGroups]    Script Date: 27/01/2014 13:58:26 ******/SET ANSI_NULLS ONGOSET QUOTED_IDENTIFIER ONGOSET ANSI_PADDING ONGOCREATE TABLE [dbo].[UserGroups](	[Id] [int] IDENTITY(1,1) NOT NULL,	[Name] [varchar](50) NOT NULL, CONSTRAINT [PK_UserGroups] PRIMARY KEY CLUSTERED (	[Id] ASC)WITH (PAD_INDEX  = OFF, STATISTICS_NORECOMPUTE  = OFF, IGNORE_DUP_KEY  = OFF, ALLOW_ROW_LOCKS  = ON, ALLOW_PAGE_LOCKS  = ON) ON [PRIMARY]) ON [PRIMARY]GOSET ANSI_PADDING OFFGO/****** Object:  Table [dbo].[FunctionCalls]    Script Date: 27/01/2014 13:57:51 ******/SET ANSI_NULLS ONGOSET QUOTED_IDENTIFIER ONGOSET ANSI_PADDING ONGOCREATE TABLE [dbo].[FunctionCalls](	[Id] [int] IDENTITY(1,1) NOT NULL,	[Call] [varchar](max) NULL, CONSTRAINT [PK_FunctionCalls] PRIMARY KEY CLUSTERED (	[Id] ASC)WITH (PAD_INDEX  = OFF, STATISTICS_NORECOMPUTE  = OFF, IGNORE_DUP_KEY  = OFF, ALLOW_ROW_LOCKS  = ON, ALLOW_PAGE_LOCKS  = ON) ON [PRIMARY]) ON [PRIMARY] TEXTIMAGE_ON [PRIMARY]GOSET ANSI_PADDING OFFGO/****** Object:  Table [dbo].[Crashes_FunctionCalls]    Script Date: 27/01/2014 13:57:39 ******/SET ANSI_NULLS ONGOSET QUOTED_IDENTIFIER ONGOCREATE TABLE [dbo].[Crashes_FunctionCalls](	[CrashId] [int] NOT NULL,	[FunctionCallId] [int] NOT NULL, CONSTRAINT [PK_Crashes_FunctionCalls] PRIMARY KEY CLUSTERED (	[CrashId] ASC,	[FunctionCallId] ASC)WITH (PAD_INDEX  = OFF, STATISTICS_NORECOMPUTE  = OFF, IGNORE_DUP_KEY  = OFF, ALLOW_ROW_LOCKS  = ON, ALLOW_PAGE_LOCKS  = ON) ON [PRIMARY]) ON [PRIMARY]GO/****** Object:  Table [dbo].[Buggs]    Script Date: 27/01/2014 13:55:40 ******/SET ANSI_NULLS ONGOSET QUOTED_IDENTIFIER ONGOSET ANSI_PADDING ONGOCREATE TABLE [dbo].[Buggs](		[Id] [int] IDENTITY(1,1) NOT NULL,
+﻿USE [CrashReport]
+GO
+/****** Object:  Table [dbo].[Users]    Script Date: 27/01/2014 13:59:11 ******/
+SET ANSI_NULLS ON
+GO
+SET QUOTED_IDENTIFIER ON
+GO
+SET ANSI_PADDING ON
+GO
+CREATE TABLE [dbo].[Users](
+	[UserName] [varchar](64) NOT NULL,
+	[UserGroup] [varchar](50) NOT NULL,
+	[UserGroupId] [int] NULL,
+ CONSTRAINT [PK_Users] PRIMARY KEY CLUSTERED 
+(
+	[UserName] ASC
+)WITH (PAD_INDEX  = OFF, STATISTICS_NORECOMPUTE  = OFF, IGNORE_DUP_KEY  = OFF, ALLOW_ROW_LOCKS  = ON, ALLOW_PAGE_LOCKS  = ON) ON [PRIMARY]
+) ON [PRIMARY]
+GO
+SET ANSI_PADDING OFF
+GO
+/****** Object:  Table [dbo].[Crashes]    Script Date: 27/01/2014 13:57:24 ******/
+SET ANSI_NULLS ON
+GO
+SET QUOTED_IDENTIFIER ON
+GO
+SET ANSI_PADDING ON
+GO
+CREATE TABLE [dbo].[Crashes](
+	[Id] [int] IDENTITY(1,1) NOT NULL,
+	[Title] [nchar](20) NULL,
+	[Summary] [varchar](512) NULL,
+	[GameName] [varchar](64) NULL,
+	[Status] [varchar](64) NULL,
+	[TimeOfCrash] [datetime] NULL,
+	[ChangeListVersion] [varchar](64) NULL,
+	[PlatformName] [varchar](64) NULL,
+	[EngineMode] [varchar](64) NULL,
+	[Description] [varchar](512) NULL,
+	[RawCallStack] [varchar](max) NULL,
+	[Pattern] [varchar](800) NULL,
+	[CommandLine] [varchar](512) NULL,
+	[ComputerName] [varchar](64) NULL,
+	[Selected] [bit] NULL,
+	[FixedChangeList] [varchar](64) NULL,
+	[LanguageExt] [varchar](64) NULL,
+	[Module] [varchar](128) NULL,
+	[BuildVersion] [varchar](128) NULL,
+	[BaseDir] [varchar](512) NULL,
+	[Version] [int] NULL,
+	[UserName] [varchar](64) NULL,
+	[TTPID] [varchar](64) NULL,
+	[AutoReporterID] [int] NULL,
+	[Processed] [bit] NULL,
+	[HasLogFile] [bit] NULL,
+	[HasMiniDumpFile] [bit] NULL,
+	[HasVideoFile] [bit] NULL,
+	[HasDiagnosticsFile] [bit] NULL,
+	[HasNewLogFile] [bit] NULL,
+	[Branch] [varchar](32) NULL,
+	[CrashType] [smallint] NOT NULL,
+	[UserNameId] [int] NULL,
+	[HasMetaData] [bit] NULL,
+	[SourceContext] [varchar](max) NULL,
+	[EpicAccountId] [varchar](64) NULL,
+	[EngineVersion] [varchar](64) NULL,
+ CONSTRAINT [PK_Crashes] PRIMARY KEY CLUSTERED 
+(
+	[Id] ASC
+)WITH (PAD_INDEX = OFF, STATISTICS_NORECOMPUTE = OFF, IGNORE_DUP_KEY = OFF, ALLOW_ROW_LOCKS = ON, ALLOW_PAGE_LOCKS = ON) ON [PRIMARY]
+) ON [PRIMARY] TEXTIMAGE_ON [PRIMARY]
+GO
+SET ANSI_PADDING OFF
+GO
+
+EXEC sys.sp_addextendedproperty @name=N'MS_Description', @value=N'NOT USED OR OBSOLETE' , @level0type=N'SCHEMA',@level0name=N'dbo', @level1type=N'TABLE',@level1name=N'Crashes', @level2type=N'COLUMN',@level2name=N'Title'
+EXEC sys.sp_addextendedproperty @name=N'MS_Description', @value=N'NOT USED OR OBSOLETE' , @level0type=N'SCHEMA',@level0name=N'dbo', @level1type=N'TABLE',@level1name=N'Crashes', @level2type=N'COLUMN',@level2name=N'Selected'
+EXEC sys.sp_addextendedproperty @name=N'MS_Description', @value=N'NOT USED OR OBSOLETE' , @level0type=N'SCHEMA',@level0name=N'dbo', @level1type=N'TABLE',@level1name=N'Crashes', @level2type=N'COLUMN',@level2name=N'Version'
+EXEC sys.sp_addextendedproperty @name=N'MS_Description', @value=N'NOT USED OR OBSOLETE' , @level0type=N'SCHEMA',@level0name=N'dbo', @level1type=N'TABLE',@level1name=N'Crashes', @level2type=N'COLUMN',@level2name=N'AutoReporterID'
+EXEC sys.sp_addextendedproperty @name=N'MS_Description', @value=N'NOT USED OR OBSOLETE' , @level0type=N'SCHEMA',@level0name=N'dbo', @level1type=N'TABLE',@level1name=N'Crashes', @level2type=N'COLUMN',@level2name=N'Processed'
+EXEC sys.sp_addextendedproperty @name=N'MS_Description', @value=N'NOT USED OR OBSOLETE' , @level0type=N'SCHEMA',@level0name=N'dbo', @level1type=N'TABLE',@level1name=N'Crashes', @level2type=N'COLUMN',@level2name=N'HasDiagnosticsFile'
+EXEC sys.sp_addextendedproperty @name=N'MS_Description', @value=N'NOT USED OR OBSOLETE' , @level0type=N'SCHEMA',@level0name=N'dbo', @level1type=N'TABLE',@level1name=N'Crashes', @level2type=N'COLUMN',@level2name=N'HasNewLogFile'
+
+EXEC sys.sp_addextendedproperty @name=N'MS_Description', @value=N'The name of the game that crashed. (FCrashDescription.GameName)' , @level0type=N'SCHEMA',@level0name=N'dbo', @level1type=N'TABLE',@level1name=N'Crashes', @level2type=N'COLUMN',@level2name=N'GameName'
+EXEC sys.sp_addextendedproperty @name=N'MS_Description', @value=N'The mode the game was in e.g. editor. (FCrashDescription.EngineMode)' , @level0type=N'SCHEMA',@level0name=N'dbo', @level1type=N'TABLE',@level1name=N'Crashes', @level2type=N'COLUMN',@level2name=N'EngineMode'
+EXEC sys.sp_addextendedproperty @name=N'MS_Description', @value=N'The platform that crashed e.g. Win64. (FCrashDescription.Platform)' , @level0type=N'SCHEMA',@level0name=N'dbo', @level1type=N'TABLE',@level1name=N'Crashes', @level2type=N'COLUMN',@level2name=N'PlatformName'
+EXEC sys.sp_addextendedproperty @name=N'MS_Description', @value=N'Encoded engine version.  BuildVersion-BuiltFromCL-BranchName E.g. 4.3.0.0-2215663+UE4-Releases+4.3 (FCrashDescription.EngineVersion)' , @level0type=N'SCHEMA',@level0name=N'dbo', @level1type=N'TABLE',@level1name=N'Crashes', @level2type=N'COLUMN',@level2name=N'EngineVersion'
+EXEC sys.sp_addextendedproperty @name=N'MS_Description', @value=N'The three component version of the app e.g. 4.4.1 (FCrashDescription.BuildVersion)' , @level0type=N'SCHEMA',@level0name=N'dbo', @level1type=N'TABLE',@level1name=N'Crashes', @level2type=N'COLUMN',@level2name=N'BuildVersion'
+EXEC sys.sp_addextendedproperty @name=N'MS_Description', @value=N'Built from changelist. (FCrashDescription.BuiltFromCL)' , @level0type=N'SCHEMA',@level0name=N'dbo', @level1type=N'TABLE',@level1name=N'Crashes', @level2type=N'COLUMN',@level2name=N'ChangeListVersion'
+EXEC sys.sp_addextendedproperty @name=N'MS_Description', @value=N'The name of the branch this game was built out of. (FCrashDescription.BranchName)' , @level0type=N'SCHEMA',@level0name=N'dbo', @level1type=N'TABLE',@level1name=N'Crashes', @level2type=N'COLUMN',@level2name=N'Branch'
+EXEC sys.sp_addextendedproperty @name=N'MS_Description', @value=N'The command line of the application that crashed. (FCrashDescription.CommandLine)' , @level0type=N'SCHEMA',@level0name=N'dbo', @level1type=N'TABLE',@level1name=N'Crashes', @level2type=N'COLUMN',@level2name=N'CommandLine'
+EXEC sys.sp_addextendedproperty @name=N'MS_Description', @value=N'The base directory where the app was running. (FCrashDescription.BaseDir)' , @level0type=N'SCHEMA',@level0name=N'dbo', @level1type=N'TABLE',@level1name=N'Crashes', @level2type=N'COLUMN',@level2name=N'BaseDir'
+EXEC sys.sp_addextendedproperty @name=N'MS_Description', @value=N'The language of the system. (FCrashDescription.Language)' , @level0type=N'SCHEMA',@level0name=N'dbo', @level1type=N'TABLE',@level1name=N'Crashes', @level2type=N'COLUMN',@level2name=N'LanguageExt'
+EXEC sys.sp_addextendedproperty @name=N'MS_Description', @value=N'The name of the user that caused this crash. (FCrashDescription.UserName)' , @level0type=N'SCHEMA',@level0name=N'dbo', @level1type=N'TABLE',@level1name=N'Crashes', @level2type=N'COLUMN',@level2name=N'UserName'
+EXEC sys.sp_addextendedproperty @name=N'MS_Description', @value=N'The unique ID used to identify the machine the crash occurred on. (FCrashDescription.MachineGuid)' , @level0type=N'SCHEMA',@level0name=N'dbo', @level1type=N'TABLE',@level1name=N'Crashes', @level2type=N'COLUMN',@level2name=N'ComputerName'
+EXEC sys.sp_addextendedproperty @name=N'MS_Description', @value=N'The Epic account ID for the user who last used the Launcher. (FCrashDescription.EpicAccountId)' , @level0type=N'SCHEMA',@level0name=N'dbo', @level1type=N'TABLE',@level1name=N'Crashes', @level2type=N'COLUMN',@level2name=N'EpicAccountId'
+EXEC sys.sp_addextendedproperty @name=N'MS_Description', @value=N'An array of strings representing the callstack of the crash. (FCrashDescription.CallStack)' , @level0type=N'SCHEMA',@level0name=N'dbo', @level1type=N'TABLE',@level1name=N'Crashes', @level2type=N'COLUMN',@level2name=N'RawCallStack'
+EXEC sys.sp_addextendedproperty @name=N'MS_Description', @value=N'An array of strings showing the source code around the crash. (FCrashDescription.SourceContext)' , @level0type=N'SCHEMA',@level0name=N'dbo', @level1type=N'TABLE',@level1name=N'Crashes', @level2type=N'COLUMN',@level2name=N'SourceContext'
+EXEC sys.sp_addextendedproperty @name=N'MS_Description', @value=N'An array of strings representing the user description of the crash. (FCrashDescription.UserDescription)' , @level0type=N'SCHEMA',@level0name=N'dbo', @level1type=N'TABLE',@level1name=N'Crashes', @level2type=N'COLUMN',@level2name=N'Description'
+EXEC sys.sp_addextendedproperty @name=N'MS_Description', @value=N'The error message, can be assertion message, ensure message or message from the fatal error. (FCrashDescription.ErrorMessage)' , @level0type=N'SCHEMA',@level0name=N'dbo', @level1type=N'TABLE',@level1name=N'Crashes', @level2type=N'COLUMN',@level2name=N'Summary'
+EXEC sys.sp_addextendedproperty @name=N'MS_Description', @value=N'The UTC time the crash occurred. (FCrashDescription.TimeofCrash)' , @level0type=N'SCHEMA',@level0name=N'dbo', @level1type=N'TABLE',@level1name=N'Crashes', @level2type=N'COLUMN',@level2name=N'TimeOfCrash'
+EXEC sys.sp_addextendedproperty @name=N'MS_Description', @value=N'Whether this crash has a minidump file. (FCrashDescription.bHasMiniDumpFile)' , @level0type=N'SCHEMA',@level0name=N'dbo', @level1type=N'TABLE',@level1name=N'Crashes', @level2type=N'COLUMN',@level2name=N'HasMiniDumpFile'
+EXEC sys.sp_addextendedproperty @name=N'MS_Description', @value=N'Whether this crash has a log file. (FCrashDescription.bHasLogFile)' , @level0type=N'SCHEMA',@level0name=N'dbo', @level1type=N'TABLE',@level1name=N'Crashes', @level2type=N'COLUMN',@level2name=N'HasLogFile'
+EXEC sys.sp_addextendedproperty @name=N'MS_Description', @value=N'Whether this crash has a video file. (FCrashDescription.bHasVideoFile)' , @level0type=N'SCHEMA',@level0name=N'dbo', @level1type=N'TABLE',@level1name=N'Crashes', @level2type=N'COLUMN',@level2name=N'HasVideoFile'
+
+
+EXEC sys.sp_addextendedproperty @name=N'MS_Description', @value=N'Status of the crash' , @level0type=N'SCHEMA',@level0name=N'dbo', @level1type=N'TABLE',@level1name=N'Crashes', @level2type=N'COLUMN',@level2name=N'Status'
+EXEC sys.sp_addextendedproperty @name=N'MS_Description', @value=N'Pattern used to group cashes into buggs' , @level0type=N'SCHEMA',@level0name=N'dbo', @level1type=N'TABLE',@level1name=N'Crashes', @level2type=N'COLUMN',@level2name=N'Pattern'
+EXEC sys.sp_addextendedproperty @name=N'MS_Description', @value=N'Changelist that fixed this crash' , @level0type=N'SCHEMA',@level0name=N'dbo', @level1type=N'TABLE',@level1name=N'Crashes', @level2type=N'COLUMN',@level2name=N'FixedChangeList'
+EXEC sys.sp_addextendedproperty @name=N'MS_Description', @value=N'Name of the module where the crash occurred' , @level0type=N'SCHEMA',@level0name=N'dbo', @level1type=N'TABLE',@level1name=N'Crashes', @level2type=N'COLUMN',@level2name=N'Module'
+EXEC sys.sp_addextendedproperty @name=N'MS_Description', @value=N'The ID of the user name' , @level0type=N'SCHEMA',@level0name=N'dbo', @level1type=N'TABLE',@level1name=N'Crashes', @level2type=N'COLUMN',@level2name=N'UserNameId'
+EXEC sys.sp_addextendedproperty @name=N'MS_Description', @value=N'The ID of the TTP' , @level0type=N'SCHEMA',@level0name=N'dbo', @level1type=N'TABLE',@level1name=N'Crashes', @level2type=N'COLUMN',@level2name=N'TTPID'
+EXEC sys.sp_addextendedproperty @name=N'MS_Description', @value=N'The type of the crash, crash(1), assert(2), ensure(3)' , @level0type=N'SCHEMA',@level0name=N'dbo', @level1type=N'TABLE',@level1name=N'Crashes', @level2type=N'COLUMN',@level2name=N'CrashType'
+
+/****** Object:  Table [dbo].[UserGroups]    Script Date: 27/01/2014 13:58:26 ******/
+SET ANSI_NULLS ON
+GO
+SET QUOTED_IDENTIFIER ON
+GO
+SET ANSI_PADDING ON
+GO
+CREATE TABLE [dbo].[UserGroups](
+	[Id] [int] IDENTITY(1,1) NOT NULL,
+	[Name] [varchar](50) NOT NULL,
+ CONSTRAINT [PK_UserGroups] PRIMARY KEY CLUSTERED 
+(
+	[Id] ASC
+)WITH (PAD_INDEX  = OFF, STATISTICS_NORECOMPUTE  = OFF, IGNORE_DUP_KEY  = OFF, ALLOW_ROW_LOCKS  = ON, ALLOW_PAGE_LOCKS  = ON) ON [PRIMARY]
+) ON [PRIMARY]
+GO
+SET ANSI_PADDING OFF
+GO
+/****** Object:  Table [dbo].[FunctionCalls]    Script Date: 27/01/2014 13:57:51 ******/
+SET ANSI_NULLS ON
+GO
+SET QUOTED_IDENTIFIER ON
+GO
+SET ANSI_PADDING ON
+GO
+CREATE TABLE [dbo].[FunctionCalls](
+	[Id] [int] IDENTITY(1,1) NOT NULL,
+	[Call] [varchar](max) NULL,
+ CONSTRAINT [PK_FunctionCalls] PRIMARY KEY CLUSTERED 
+(
+	[Id] ASC
+)WITH (PAD_INDEX  = OFF, STATISTICS_NORECOMPUTE  = OFF, IGNORE_DUP_KEY  = OFF, ALLOW_ROW_LOCKS  = ON, ALLOW_PAGE_LOCKS  = ON) ON [PRIMARY]
+) ON [PRIMARY] TEXTIMAGE_ON [PRIMARY]
+GO
+SET ANSI_PADDING OFF
+GO
+/****** Object:  Table [dbo].[Crashes_FunctionCalls]    Script Date: 27/01/2014 13:57:39 ******/
+SET ANSI_NULLS ON
+GO
+SET QUOTED_IDENTIFIER ON
+GO
+CREATE TABLE [dbo].[Crashes_FunctionCalls](
+	[CrashId] [int] NOT NULL,
+	[FunctionCallId] [int] NOT NULL,
+ CONSTRAINT [PK_Crashes_FunctionCalls] PRIMARY KEY CLUSTERED 
+(
+	[CrashId] ASC,
+	[FunctionCallId] ASC
+)WITH (PAD_INDEX  = OFF, STATISTICS_NORECOMPUTE  = OFF, IGNORE_DUP_KEY  = OFF, ALLOW_ROW_LOCKS  = ON, ALLOW_PAGE_LOCKS  = ON) ON [PRIMARY]
+) ON [PRIMARY]
+GO
+/****** Object:  Table [dbo].[Buggs]    Script Date: 27/01/2014 13:55:40 ******/
+SET ANSI_NULLS ON
+GO
+SET QUOTED_IDENTIFIER ON
+GO
+SET ANSI_PADDING ON
+GO
+CREATE TABLE [dbo].[Buggs](
+		[Id] [int] IDENTITY(1,1) NOT NULL,
 	[TTPID] [varchar](50) NULL,
 	[Title] [varchar](120) NULL,
 	[Summary] [varchar](250) NULL,
@@ -16,6 +189,138 @@
 	[Game] [varchar](120) NULL,
 	[BuildVersion] [varchar](128) NULL,
 	[CrashType] [smallint] NULL,
-	[SummaryV2] [varchar](512) NULL, CONSTRAINT [PK_Buggs] PRIMARY KEY CLUSTERED (	[Id] ASC)WITH (PAD_INDEX  = OFF, STATISTICS_NORECOMPUTE  = OFF, IGNORE_DUP_KEY  = OFF, ALLOW_ROW_LOCKS  = ON, ALLOW_PAGE_LOCKS  = ON) ON [PRIMARY]) ON [PRIMARY] TEXTIMAGE_ON [PRIMARY]GOSET ANSI_PADDING OFFGOEXEC sys.sp_addextendedproperty @name=N'MS_Description', @value=N'The error message, can be assertion message, ensure message or message from the fatal error. (FCrashDescription.ErrorMessage)' , @level0type=N'SCHEMA',@level0name=N'dbo', @level1type=N'TABLE',@level1name=N'Buggs', @level2type=N'COLUMN',@level2name=N'SummaryV2'GO/****** Object:  Table [dbo].[Buggs_Crashes]    Script Date: 27/01/2014 13:56:43 ******/SET ANSI_NULLS ONGOSET QUOTED_IDENTIFIER ONGOCREATE TABLE [dbo].[Buggs_Crashes](	[BuggId] [int] NOT NULL,	[CrashId] [int] NOT NULL, CONSTRAINT [PK_Buggs_Crashes] PRIMARY KEY CLUSTERED (	[BuggId] ASC,	[CrashId] ASC)WITH (PAD_INDEX  = OFF, STATISTICS_NORECOMPUTE  = OFF, IGNORE_DUP_KEY = OFF, ALLOW_ROW_LOCKS  = ON, ALLOW_PAGE_LOCKS  = ON) ON [PRIMARY]) ON [PRIMARY]GO/****** Object:  Table [dbo].[Buggs_Users]    Script Date: 27/01/2014 13:57:11 ******/SET ANSI_NULLS ONGOSET QUOTED_IDENTIFIER ONGOSET ANSI_PADDING ONGOCREATE TABLE [dbo].[Buggs_Users](	[BuggId] [int] NOT NULL,	[UserName] [varchar](64) NOT NULL, CONSTRAINT [PK_Buggs_Users] PRIMARY KEY CLUSTERED (	[BuggId] ASC,	[UserName] ASC)WITH (PAD_INDEX  = OFF, STATISTICS_NORECOMPUTE  = OFF, IGNORE_DUP_KEY  = OFF, ALLOW_ROW_LOCKS  = ON, ALLOW_PAGE_LOCKS  = ON) ON [PRIMARY]) ON [PRIMARY]GOSET ANSI_PADDING OFFGO/****** Object:  Table [dbo].[Buggs_UserGroups]    Script Date: 27/01/2014 13:56:58 ******/SET ANSI_NULLS ONGOSET QUOTED_IDENTIFIER ONGOCREATE TABLE [dbo].[Buggs_UserGroups](	[BuggId] [int] NOT NULL,	[UserGroupId] [int] NOT NULL, CONSTRAINT [PK_Buggs_UserGroups] PRIMARY KEY CLUSTERED (	[BuggId] ASC,	[UserGroupId] ASC)WITH (PAD_INDEX  = OFF, STATISTICS_NORECOMPUTE  = OFF, IGNORE_DUP_KEY  = OFF, ALLOW_ROW_LOCKS  = ON, ALLOW_PAGE_LOCKS  = ON) ON [PRIMARY]) ON [PRIMARY]GO/****** Object:  Table [dbo].[PIIMapping]    Script Date: 27/01/2014 13:58:11 ******/SET ANSI_NULLS ONGOSET QUOTED_IDENTIFIER ONGOSET ANSI_PADDING ONGOCREATE TABLE [dbo].[PIIMapping](	[ID] [int] IDENTITY(1,1) NOT NULL,	[MachineGUID] [varchar](64) NULL,	[UserName] [varchar](64) NULL,	[MachineName] [varchar](64) NULL, CONSTRAINT [PK_PIIMapping] PRIMARY KEY CLUSTERED (	[ID] ASC)WITH (PAD_INDEX  = OFF, STATISTICS_NORECOMPUTE  = OFF, IGNORE_DUP_KEY  = OFF, ALLOW_ROW_LOCKS  = ON, ALLOW_PAGE_LOCKS  = ON) ON [PRIMARY]) ON [PRIMARY]GOSET ANSI_PADDING OFFGOEXEC sys.sp_addextendedproperty @name =N'MS_Description', @value =N'A mapping of the anonymous WER machine id to a user name and machine name' , @level0type =N'SCHEMA',@level0name =N'dbo', @level1type =N'TABLE',@level1name =N'PIIMapping'GO/****** Object:  Table [dbo].[UsersMapping]    Script Date: 2014-10-24 13:46:25 ******/
-SET ANSI_NULLS ONGOSET QUOTED_IDENTIFIER ONGOCREATE TABLE [dbo].[UsersMapping](	[Id] [int] IDENTITY(1,1) NOT NULL,	[UserEmail] [nvarchar](128) NOT NULL,	[MachineId] [nchar](32) NULL,	[EpicAccountId] [nchar](32) NULL,	[UserName] [nvarchar](64) NULL, CONSTRAINT [PK_UsersMapping] PRIMARY KEY CLUSTERED (	[Id] ASC)WITH (PAD_INDEX = OFF, STATISTICS_NORECOMPUTE = OFF, IGNORE_DUP_KEY = OFF, ALLOW_ROW_LOCKS = ON, ALLOW_PAGE_LOCKS = ON) ON [PRIMARY]) ON [PRIMARY]GOEXEC sys.sp_addextendedproperty @name=N'MS_Description', @value=N'The unique ID used to identify the machine the crash occurred on.' , @level0type=N'SCHEMA',@level0name=N'dbo', @level1type=N'TABLE',@level1name=N'UsersMapping', @level2type=N'COLUMN',@level2name=N'MachineId'GOEXEC sys.sp_addextendedproperty @name=N'MS_Description', @value=N'The Epic account ID for the user who last used the Launcher.' , @level0type=N'SCHEMA',@level0name=N'dbo', @level1type=N'TABLE',@level1name=N'UsersMapping', @level2type=N'COLUMN',@level2name=N'EpicAccountId'GOEXEC sys.sp_addextendedproperty @name=N'MS_Description', @value=N'The name of the user that caused this crash.' , @level0type=N'SCHEMA',@level0name=N'dbo', @level1type=N'TABLE',@level1name=N'UsersMapping', @level2type=N'COLUMN',@level2name=N'UserName'GO
-/****** Object:  ForeignKey [FK_Buggs_Crashes_Crashes]    Script Date: 27/01/2014 13:56:43 ******/ALTER TABLE [dbo].[Buggs_Crashes]  WITH CHECK ADD  CONSTRAINT [FK_Buggs_Crashes_Crashes] FOREIGN KEY([CrashId])REFERENCES [dbo].[Crashes] ([Id])GOALTER TABLE [dbo].[Buggs_Crashes] CHECK CONSTRAINT [FK_Buggs_Crashes_Crashes]GO/****** Object:  ForeignKey [FK_Buggs_UserGroups_UserGroups]    Script Date: 27/01/2014 13:56:43 ******/ALTER TABLE [dbo].[Buggs_UserGroups]  WITH CHECK ADD  CONSTRAINT [FK_Buggs_UserGroups_UserGroups] FOREIGN KEY([UserGroupId])REFERENCES [dbo].[UserGroups] ([Id])GOALTER TABLE [dbo].[Buggs_UserGroups] CHECK CONSTRAINT [FK_Buggs_UserGroups_UserGroups]GO/****** Object:  ForeignKey [FK_Buggs_Users_Users]    Script Date: 27/01/2014 13:56:43 ******/ALTER TABLE [dbo].[Buggs_Users]  WITH CHECK ADD  CONSTRAINT [FK_Buggs_Users_Users] FOREIGN KEY([UserName])REFERENCES [dbo].[Users] ([UserName])GOALTER TABLE [dbo].[Buggs_Users] CHECK CONSTRAINT [FK_Buggs_Users_Users]GO/****** Object:  ForeignKey [FK_Crashes_Users]    Script Date: 27/01/2014 13:56:43 ******/ALTER TABLE [dbo].[Crashes]  WITH CHECK ADD  CONSTRAINT [FK_Crashes_Users] FOREIGN KEY([UserName])REFERENCES [dbo].[Users] ([UserName])GOALTER TABLE [dbo].[Crashes] CHECK CONSTRAINT [FK_Crashes_Users]GO/****** Object:  ForeignKey [FK_Crashes_FunctionCalls_Crashes]    Script Date: 27/01/2014 13:56:43 ******/ALTER TABLE [dbo].[Crashes_FunctionCalls]  WITH CHECK ADD  CONSTRAINT [FK_Crashes_FunctionCalls_Crashes] FOREIGN KEY([CrashId])REFERENCES [dbo].[Crashes] ([Id])GOALTER TABLE [dbo].[Crashes_FunctionCalls] CHECK CONSTRAINT [FK_Crashes_FunctionCalls_Crashes]GO/****** Object:  ForeignKey [FK_Crashes_FunctionCalls_FunctionCalls]    Script Date: 27/01/2014 13:56:43 ******/ALTER TABLE [dbo].[Crashes_FunctionCalls]  WITH CHECK ADD  CONSTRAINT [FK_Crashes_FunctionCalls_FunctionCalls] FOREIGN KEY([FunctionCallId])REFERENCES [dbo].[FunctionCalls] ([Id])GOALTER TABLE [dbo].[Crashes_FunctionCalls] CHECK CONSTRAINT [FK_Crashes_FunctionCalls_FunctionCalls]GO/****** Object:  DefaultValues    Script Date: 27/01/2014 13:56:43 ******/ALTER TABLE [dbo].[Crashes] ADD  CONSTRAINT [DF_Crashes_Processed]  DEFAULT ((0)) FOR [Processed]GOALTER TABLE [dbo].[Crashes] ADD  CONSTRAINT [DF_Crashes_HasDiagnosticsFile]  DEFAULT ((0)) FOR [HasDiagnosticsFile]GOALTER TABLE [dbo].[Crashes] ADD  CONSTRAINT [DF_Crashes_HasNewLogFile]  DEFAULT ((0)) FOR [HasNewLogFile]GOINSERT INTO [dbo].[UserGroups]		   ([Name])	 VALUES		   ('General')GOINSERT INTO [dbo].[UserGroups]		   ([Name])	 VALUES		   ('Tester')GOINSERT INTO [dbo].[UserGroups]		   ([Name])	 VALUES		   ('Coder')GOINSERT INTO [dbo].[UserGroups]		   ([Name])	 VALUES		   ('Automated')GOINSERT INTO [dbo].[UserGroups]		   ([Name])	 VALUES		   ('Undefined')GO
\ No newline at end of file
+	[SummaryV2] [varchar](512) NULL,
+ CONSTRAINT [PK_Buggs] PRIMARY KEY CLUSTERED 
+(
+	[Id] ASC
+)WITH (PAD_INDEX  = OFF, STATISTICS_NORECOMPUTE  = OFF, IGNORE_DUP_KEY  = OFF, ALLOW_ROW_LOCKS  = ON, ALLOW_PAGE_LOCKS  = ON) ON [PRIMARY]
+) ON [PRIMARY] TEXTIMAGE_ON [PRIMARY]
+GO
+SET ANSI_PADDING OFF
+GO
+
+EXEC sys.sp_addextendedproperty @name=N'MS_Description', @value=N'The error message, can be assertion message, ensure message or message from the fatal error. (FCrashDescription.ErrorMessage)' , @level0type=N'SCHEMA',@level0name=N'dbo', @level1type=N'TABLE',@level1name=N'Buggs', @level2type=N'COLUMN',@level2name=N'SummaryV2'
+GO
+
+/****** Object:  Table [dbo].[Buggs_Crashes]    Script Date: 27/01/2014 13:56:43 ******/
+SET ANSI_NULLS ON
+GO
+SET QUOTED_IDENTIFIER ON
+GO
+CREATE TABLE [dbo].[Buggs_Crashes](
+	[BuggId] [int] NOT NULL,
+	[CrashId] [int] NOT NULL,
+ CONSTRAINT [PK_Buggs_Crashes] PRIMARY KEY CLUSTERED 
+(
+	[BuggId] ASC,
+	[CrashId] ASC
+)WITH (PAD_INDEX  = OFF, STATISTICS_NORECOMPUTE  = OFF, IGNORE_DUP_KEY = OFF, ALLOW_ROW_LOCKS  = ON, ALLOW_PAGE_LOCKS  = ON) ON [PRIMARY]
+) ON [PRIMARY]
+GO
+/****** Object:  Table [dbo].[Buggs_Users]    Script Date: 27/01/2014 13:57:11 ******/
+SET ANSI_NULLS ON
+GO
+SET QUOTED_IDENTIFIER ON
+GO
+SET ANSI_PADDING ON
+GO
+CREATE TABLE [dbo].[Buggs_Users](
+	[BuggId] [int] NOT NULL,
+	[UserName] [varchar](64) NOT NULL,
+ CONSTRAINT [PK_Buggs_Users] PRIMARY KEY CLUSTERED 
+(
+	[BuggId] ASC,
+	[UserName] ASC
+)WITH (PAD_INDEX  = OFF, STATISTICS_NORECOMPUTE  = OFF, IGNORE_DUP_KEY  = OFF, ALLOW_ROW_LOCKS  = ON, ALLOW_PAGE_LOCKS  = ON) ON [PRIMARY]
+) ON [PRIMARY]
+GO
+SET ANSI_PADDING OFF
+GO
+/****** Object:  Table [dbo].[Buggs_UserGroups]    Script Date: 27/01/2014 13:56:58 ******/
+SET ANSI_NULLS ON
+GO
+SET QUOTED_IDENTIFIER ON
+GO
+CREATE TABLE [dbo].[Buggs_UserGroups](
+	[BuggId] [int] NOT NULL,
+	[UserGroupId] [int] NOT NULL,
+ CONSTRAINT [PK_Buggs_UserGroups] PRIMARY KEY CLUSTERED 
+(
+	[BuggId] ASC,
+	[UserGroupId] ASC
+)WITH (PAD_INDEX  = OFF, STATISTICS_NORECOMPUTE  = OFF, IGNORE_DUP_KEY  = OFF, ALLOW_ROW_LOCKS  = ON, ALLOW_PAGE_LOCKS  = ON) ON [PRIMARY]
+) ON [PRIMARY]
+GO
+
+
+
+/****** Object:  ForeignKey [FK_Buggs_Crashes_Crashes]    Script Date: 27/01/2014 13:56:43 ******/
+ALTER TABLE [dbo].[Buggs_Crashes]  WITH CHECK ADD  CONSTRAINT [FK_Buggs_Crashes_Crashes] FOREIGN KEY([CrashId])
+REFERENCES [dbo].[Crashes] ([Id])
+GO
+ALTER TABLE [dbo].[Buggs_Crashes] CHECK CONSTRAINT [FK_Buggs_Crashes_Crashes]
+GO
+/****** Object:  ForeignKey [FK_Buggs_UserGroups_UserGroups]    Script Date: 27/01/2014 13:56:43 ******/
+ALTER TABLE [dbo].[Buggs_UserGroups]  WITH CHECK ADD  CONSTRAINT [FK_Buggs_UserGroups_UserGroups] FOREIGN KEY([UserGroupId])
+REFERENCES [dbo].[UserGroups] ([Id])
+GO
+ALTER TABLE [dbo].[Buggs_UserGroups] CHECK CONSTRAINT [FK_Buggs_UserGroups_UserGroups]
+GO
+/****** Object:  ForeignKey [FK_Buggs_Users_Users]    Script Date: 27/01/2014 13:56:43 ******/
+ALTER TABLE [dbo].[Buggs_Users]  WITH CHECK ADD  CONSTRAINT [FK_Buggs_Users_Users] FOREIGN KEY([UserName])
+REFERENCES [dbo].[Users] ([UserName])
+GO
+ALTER TABLE [dbo].[Buggs_Users] CHECK CONSTRAINT [FK_Buggs_Users_Users]
+GO
+/****** Object:  ForeignKey [FK_Crashes_Users]    Script Date: 27/01/2014 13:56:43 ******/
+ALTER TABLE [dbo].[Crashes]  WITH CHECK ADD  CONSTRAINT [FK_Crashes_Users] FOREIGN KEY([UserName])
+REFERENCES [dbo].[Users] ([UserName])
+GO
+ALTER TABLE [dbo].[Crashes] CHECK CONSTRAINT [FK_Crashes_Users]
+GO
+/****** Object:  ForeignKey [FK_Crashes_FunctionCalls_Crashes]    Script Date: 27/01/2014 13:56:43 ******/
+ALTER TABLE [dbo].[Crashes_FunctionCalls]  WITH CHECK ADD  CONSTRAINT [FK_Crashes_FunctionCalls_Crashes] FOREIGN KEY([CrashId])
+REFERENCES [dbo].[Crashes] ([Id])
+GO
+ALTER TABLE [dbo].[Crashes_FunctionCalls] CHECK CONSTRAINT [FK_Crashes_FunctionCalls_Crashes]
+GO
+/****** Object:  ForeignKey [FK_Crashes_FunctionCalls_FunctionCalls]    Script Date: 27/01/2014 13:56:43 ******/
+ALTER TABLE [dbo].[Crashes_FunctionCalls]  WITH CHECK ADD  CONSTRAINT [FK_Crashes_FunctionCalls_FunctionCalls] FOREIGN KEY([FunctionCallId])
+REFERENCES [dbo].[FunctionCalls] ([Id])
+GO
+ALTER TABLE [dbo].[Crashes_FunctionCalls] CHECK CONSTRAINT [FK_Crashes_FunctionCalls_FunctionCalls]
+GO
+/****** Object:  DefaultValues    Script Date: 27/01/2014 13:56:43 ******/
+ALTER TABLE [dbo].[Crashes] ADD  CONSTRAINT [DF_Crashes_Processed]  DEFAULT ((0)) FOR [Processed]
+GO
+ALTER TABLE [dbo].[Crashes] ADD  CONSTRAINT [DF_Crashes_HasDiagnosticsFile]  DEFAULT ((0)) FOR [HasDiagnosticsFile]
+GO
+ALTER TABLE [dbo].[Crashes] ADD  CONSTRAINT [DF_Crashes_HasNewLogFile]  DEFAULT ((0)) FOR [HasNewLogFile]
+GO
+
+INSERT INTO [dbo].[UserGroups]
+		   ([Name])
+	 VALUES
+		   ('General')
+GO
+INSERT INTO [dbo].[UserGroups]
+		   ([Name])
+	 VALUES
+		   ('Tester')
+GO
+
+INSERT INTO [dbo].[UserGroups]
+		   ([Name])
+	 VALUES
+		   ('Coder')
+GO
+INSERT INTO [dbo].[UserGroups]
+		   ([Name])
+	 VALUES
+		   ('Automated')
+GO
+INSERT INTO [dbo].[UserGroups]
+		   ([Name])
+	 VALUES
+		   ('Undefined')
+GO
