CrashWebSite - High memory usage in dashboard view (fixed)

diff --git a/Engine/Source/Programs/CrashReporter/CrashReportWebSite/Controllers/DashboardController.cs b/Engine/Source/Programs/CrashReporter/CrashReportWebSite/Controllers/DashboardController.cs
--- a/Engine/Source/Programs/CrashReporter/CrashReportWebSite/Controllers/DashboardController.cs
+++ b/Engine/Source/Programs/CrashReporter/CrashReportWebSite/Controllers/DashboardController.cs
@@ -146,8 +146,11 @@ namespace Tools.CrashReporter.CrashReportWebSite.Controllers
 				DateTime AfewMonthsAgo = Today.AddMonths( -6 );
 
 				FAutoScopedLogTimer LogTimerSQL = new FAutoScopedLogTimer( "CrashesFilterByDate", "", "" );
-				IEnumerable<Crash> Crashes = FRepository.Get().Crashes.FilterByDate( FRepository.Get().Crashes.ListAll(), AfewMonthsAgo, Today );
-				var VMinimalCrashes = Crashes.Select( Crash => new { TimeOfCrash = Crash.TimeOfCrash.Value, UserID = Crash.UserNameId.Value } ).ToList();
+
+				IQueryable<Crash> CrashesInTimeFrame = FRepository.Get().Crashes.ListAll()
+					.Where( MyCrash => MyCrash.TimeOfCrash >= AfewMonthsAgo && MyCrash.TimeOfCrash <= Today.AddDays( 1 ) );
+				//IEnumerable<Crash> Crashes = FRepository.Get().Crashes.FilterByDate( FRepository.Get().Crashes.ListAll(), AfewMonthsAgo, Today );
+				var VMinimalCrashes = CrashesInTimeFrame.Select( Crash => new { TimeOfCrash = Crash.TimeOfCrash.Value, UserID = Crash.UserNameId.Value } ).ToList();
 
 				List<FCrashMinimal> MinimalCrashes = new List<FCrashMinimal>( VMinimalCrashes.Count );
 				foreach( var Anotype in VMinimalCrashes )
