CrashReportWWW - Misc UI tweaks (generated text is more visible)

diff --git a/Engine/Source/Programs/CrashReporter/CrashReportWebSite/Views/Buggs/Index.aspx b/Engine/Source/Programs/CrashReporter/CrashReportWebSite/Views/Buggs/Index.aspx
--- a/Engine/Source/Programs/CrashReporter/CrashReportWebSite/Views/Buggs/Index.aspx
+++ b/Engine/Source/Programs/CrashReporter/CrashReportWebSite/Views/Buggs/Index.aspx
@@ -52,6 +52,8 @@
 </asp:Content>
 
 <asp:Content ID="AboveMainContent" ContentPlaceHolderID="AboveMainContent" runat="server" >
+	<div style="clear:both;"><small style="color: lightgray;">Generated in <%=Model.GenerationTime%> second(s)</small><br /></div>
+
 <div id='SearchForm' style="clear:both;">
 <%using( Html.BeginForm( "", "Buggs", FormMethod.Get, new { id = "FilterBuggsForm" } ) )
 { %>
diff --git a/Engine/Source/Programs/CrashReporter/CrashReportWebSite/Views/Buggs/Show.aspx b/Engine/Source/Programs/CrashReporter/CrashReportWebSite/Views/Buggs/Show.aspx
--- a/Engine/Source/Programs/CrashReporter/CrashReportWebSite/Views/Buggs/Show.aspx
+++ b/Engine/Source/Programs/CrashReporter/CrashReportWebSite/Views/Buggs/Show.aspx
@@ -62,6 +62,7 @@
 </asp:Content>
 
 <asp:Content ID="AboveMainContent" ContentPlaceHolderID="AboveMainContent" runat="server">
+	<div style="clear:both;"><small style="color: lightgray;">Generated in <%=Model.GenerationTime%> second(s)</small><br /></div>
 </asp:Content>
 
 <asp:Content ID="MainContent" ContentPlaceHolderID="MainContent" runat="server">
@@ -280,6 +281,4 @@
 	</div>
 </div>
 
-	<small>Generated in <%=Model.GenerationTime%> second(s)</small>
-
 </asp:Content>
diff --git a/Engine/Source/Programs/CrashReporter/CrashReportWebSite/Views/Buggs/ViewBuggs.ascx b/Engine/Source/Programs/CrashReporter/CrashReportWebSite/Views/Buggs/ViewBuggs.ascx
--- a/Engine/Source/Programs/CrashReporter/CrashReportWebSite/Views/Buggs/ViewBuggs.ascx
+++ b/Engine/Source/Programs/CrashReporter/CrashReportWebSite/Views/Buggs/ViewBuggs.ascx
@@ -163,5 +163,3 @@
 		}
 		%>
 </table>
-
-<small>Generated in <%=Model.GenerationTime%> second(s)</small>
\ No newline at end of file
diff --git a/Engine/Source/Programs/CrashReporter/CrashReportWebSite/Views/Crashes/Index.aspx b/Engine/Source/Programs/CrashReporter/CrashReportWebSite/Views/Crashes/Index.aspx
--- a/Engine/Source/Programs/CrashReporter/CrashReportWebSite/Views/Crashes/Index.aspx
+++ b/Engine/Source/Programs/CrashReporter/CrashReportWebSite/Views/Crashes/Index.aspx
@@ -96,8 +96,10 @@ Crash Reports
 </asp:Content>
 
 <asp:Content ID="AboveMainContent"  ContentPlaceHolderID="AboveMainContent" runat="server" >
+	<div style="clear:both;"><small style="color: lightgray;">Generated in <%=Model.GenerationTime%> second(s)</small><br /></div>
+
 	<div id='SearchForm' style="clear:both">
-		
+
 	<% using( Html.BeginForm( "", "Crashes", FormMethod.Get, new { id = "FilterCrashesForm" } ) )
 	{ %>
 		<%=Html.HiddenFor( u => u.UserGroup )%>
diff --git a/Engine/Source/Programs/CrashReporter/CrashReportWebSite/Views/Crashes/Show.aspx b/Engine/Source/Programs/CrashReporter/CrashReportWebSite/Views/Crashes/Show.aspx
--- a/Engine/Source/Programs/CrashReporter/CrashReportWebSite/Views/Crashes/Show.aspx
+++ b/Engine/Source/Programs/CrashReporter/CrashReportWebSite/Views/Crashes/Show.aspx
@@ -59,6 +59,7 @@
 </asp:Content>
 
 <asp:Content ID="AboveMainContent" ContentPlaceHolderID="AboveMainContent" runat="server">
+	<div style="clear:both;"><small style="color: lightgray;">Generated in <%=Model.GenerationTime%> second(s)</small><br /></div>
 </asp:Content>
 
 <asp:Content ID="MainContent" ContentPlaceHolderID="MainContent" runat="server">
@@ -304,7 +305,4 @@
 		</div>
 	</div>
 </div>
-
-	<small>Generated in <%=Model.GenerationTime%> second(s)</small>
-
 </asp:Content>
diff --git a/Engine/Source/Programs/CrashReporter/CrashReportWebSite/Views/Crashes/ViewCrash.ascx b/Engine/Source/Programs/CrashReporter/CrashReportWebSite/Views/Crashes/ViewCrash.ascx
--- a/Engine/Source/Programs/CrashReporter/CrashReportWebSite/Views/Crashes/ViewCrash.ascx
+++ b/Engine/Source/Programs/CrashReporter/CrashReportWebSite/Views/Crashes/ViewCrash.ascx
@@ -194,5 +194,3 @@
 			} 
 		%>
 </table>
-
-<small>Generated in <%=Model.GenerationTime%> second(s)</small>
\ No newline at end of file
diff --git a/Engine/Source/Programs/CrashReporter/CrashReportWebSite/Views/Dashboard/Index.aspx b/Engine/Source/Programs/CrashReporter/CrashReportWebSite/Views/Dashboard/Index.aspx
--- a/Engine/Source/Programs/CrashReporter/CrashReportWebSite/Views/Dashboard/Index.aspx
+++ b/Engine/Source/Programs/CrashReporter/CrashReportWebSite/Views/Dashboard/Index.aspx
@@ -33,6 +33,7 @@ table
 </asp:Content>
 
 <asp:Content ID="AboveMainContent" ContentPlaceHolderID="AboveMainContent" runat="server" >
+	<div style="clear:both;"><small style="color: lightgray;">Generated in <%=Model.GenerationTime%> second(s)</small><br /></div>
 </asp:Content>
 
 <asp:Content ID="MainContent" ContentPlaceHolderID="MainContent" runat="server">
@@ -121,5 +122,5 @@ table
 				</td>
 			</tr>
 		</table>	
-	<small>Generated in <%=Model.GenerationTime%> second(s)</small>
+
 </asp:Content>
diff --git a/Engine/Source/Programs/CrashReporter/CrashReportWebSite/Views/Reports/Index.aspx b/Engine/Source/Programs/CrashReporter/CrashReportWebSite/Views/Reports/Index.aspx
--- a/Engine/Source/Programs/CrashReporter/CrashReportWebSite/Views/Reports/Index.aspx
+++ b/Engine/Source/Programs/CrashReporter/CrashReportWebSite/Views/Reports/Index.aspx
@@ -39,6 +39,7 @@
 </asp:Content>
 
 <asp:Content ID="AboveMainContent" ContentPlaceHolderID="AboveMainContent" runat="server" >
+<div style="clear:both;"><small style="color: lightgray;">Generated in <%=Model.GenerationTime%> second(s)</small><br /></div>
 
 <div id='SearchForm' style="clear:both;">
 <%using( Html.BeginForm( "", "Reports", FormMethod.Get, new { id = "ReporsForm" } ) )
diff --git a/Engine/Source/Programs/CrashReporter/CrashReportWebSite/Views/Reports/ViewReports.ascx b/Engine/Source/Programs/CrashReporter/CrashReportWebSite/Views/Reports/ViewReports.ascx
--- a/Engine/Source/Programs/CrashReporter/CrashReportWebSite/Views/Reports/ViewReports.ascx
+++ b/Engine/Source/Programs/CrashReporter/CrashReportWebSite/Views/Reports/ViewReports.ascx
@@ -211,7 +211,5 @@ table, table tr, table th, table td
 
 </table>
 
-<small>Generated in <%=Model.GenerationTime%> second(s)</small>
-
 </form>
 </body>
\ No newline at end of file
diff --git a/Engine/Source/Programs/CrashReporter/CrashReportWebSite/Views/Shared/Site.Master b/Engine/Source/Programs/CrashReporter/CrashReportWebSite/Views/Shared/Site.Master
--- a/Engine/Source/Programs/CrashReporter/CrashReportWebSite/Views/Shared/Site.Master
+++ b/Engine/Source/Programs/CrashReporter/CrashReportWebSite/Views/Shared/Site.Master
@@ -31,6 +31,10 @@
 				<br />
 			</div>
 			<div id="utils">
+				&nbsp;|&nbsp;
+				<span><a href="/Crashes">Crashes</a></span>
+				&nbsp;|&nbsp;
+				<span><a href="/Buggs">Buggs</a></span>
 				&nbsp;|&nbsp;
 				<span><a href="/Reports/Index/">Reports</a></span>
 				&nbsp;|&nbsp;
