CrashReportWebsite - Fixes searching in the crashes view

diff --git a/Engine/Source/Programs/CrashReporter/CrashReportWebSite/Controllers/BuggsController.cs b/Engine/Source/Programs/CrashReporter/CrashReportWebSite/Controllers/BuggsController.cs
--- a/Engine/Source/Programs/CrashReporter/CrashReportWebSite/Controllers/BuggsController.cs
+++ b/Engine/Source/Programs/CrashReporter/CrashReportWebSite/Controllers/BuggsController.cs
@@ -31,8 +31,10 @@ namespace Tools.CrashReporter.CrashReportWebSite.Controllers
 		{
 			using( FAutoScopedLogTimer LogTimer = new FAutoScopedLogTimer( this.GetType().ToString() ) )
 			{
+				BuggRepository Buggs = new BuggRepository();
+
 				FormHelper FormData = new FormHelper( Request, BuggsForm, "CrashesInTimeFrameGroup" );
-				BuggsViewModel Results = FRepository.Get().Buggs.GetResults( FormData );
+				BuggsViewModel Results = Buggs.GetResults( FormData );
 				Results.GenerationTime = LogTimer.GetElapsedSeconds().ToString( "F2" );
 				return View( "Index", Results );
 			}
@@ -48,7 +50,7 @@ namespace Tools.CrashReporter.CrashReportWebSite.Controllers
 		{
 			using( FAutoScopedLogTimer LogTimer = new FAutoScopedLogTimer( this.GetType().ToString() + "(BuggId=" + id + ")" ) )
 			{
-				var Buggs = FRepository.Get().Buggs;
+				BuggRepository Buggs = new BuggRepository();
 
 				// Set the display properties based on the radio buttons
 				bool DisplayModuleNames = false;
diff --git a/Engine/Source/Programs/CrashReporter/CrashReportWebSite/Controllers/CrashesController.cs b/Engine/Source/Programs/CrashReporter/CrashReportWebSite/Controllers/CrashesController.cs
--- a/Engine/Source/Programs/CrashReporter/CrashReportWebSite/Controllers/CrashesController.cs
+++ b/Engine/Source/Programs/CrashReporter/CrashReportWebSite/Controllers/CrashesController.cs
@@ -33,7 +33,7 @@ namespace Tools.CrashReporter.CrashReportWebSite.Controllers
 		{
 			using( FAutoScopedLogTimer LogTimer = new FAutoScopedLogTimer( this.GetType().ToString() ) )
 			{
-				var Crashes = FRepository.Get().Crashes;
+				CrashRepository Crashes = new CrashRepository();
 
 				// Handle any edits made in the Set form fields
 				foreach( var Entry in CrashesForm )
@@ -61,7 +61,7 @@ namespace Tools.CrashReporter.CrashReportWebSite.Controllers
 						}
 					}
 
-					FRepository.Get().SubmitChanges();
+					Crashes.SubmitChanges();
 				}
 
 				// <STATUS>
@@ -87,8 +87,9 @@ namespace Tools.CrashReporter.CrashReportWebSite.Controllers
 		{
 			using( FAutoScopedLogTimer LogTimer = new FAutoScopedLogTimer( this.GetType().ToString() + "(CrashId=" + id + ")" ) )
 			{
+				CrashRepository Crashes = new CrashRepository();
+
 				CallStackContainer CurrentCallStack = null;
-				var Crashes = FRepository.Get().Crashes;
 
 				// Update the selected crash based on the form contents
 				Crash CurrentCrash = Crashes.GetCrash( id );
@@ -138,7 +139,7 @@ namespace Tools.CrashReporter.CrashReportWebSite.Controllers
 
 				// Populate the crash with the correct user data
 				Crashes.PopulateUserInfo( CurrentCrash );
-				FRepository.Get().SubmitChanges();
+				Crashes.SubmitChanges();
 
 				var Model = new CrashViewModel { Crash = CurrentCrash, CallStack = CurrentCallStack };
 				Model.GenerationTime = LogTimer.GetElapsedSeconds().ToString( "F2" );
@@ -155,6 +156,8 @@ namespace Tools.CrashReporter.CrashReportWebSite.Controllers
 		{
 			using( FAutoScopedLogTimer LogTimer = new FAutoScopedLogTimer( this.GetType().ToString() + "(NewCrashId=" + id + ")" ) )
 			{
+				CrashRepository Crashes = new CrashRepository();
+
 				CrashReporterResult NewCrashResult = new CrashReporterResult();
 				NewCrashResult.ID = -1;
 
@@ -164,7 +167,7 @@ namespace Tools.CrashReporter.CrashReportWebSite.Controllers
 					{
 						string Result = Reader.ReadToEnd();
 						CrashDescription NewCrash = XmlHandler.FromXmlString<CrashDescription>( Result );
-						NewCrashResult.ID = FRepository.Get().Crashes.AddNewCrash( NewCrash );
+						NewCrashResult.ID = Crashes.AddNewCrash( NewCrash );
 						NewCrashResult.bSuccess = true;
 					}
 				}
diff --git a/Engine/Source/Programs/CrashReporter/CrashReportWebSite/Controllers/DashboardController.cs b/Engine/Source/Programs/CrashReporter/CrashReportWebSite/Controllers/DashboardController.cs
--- a/Engine/Source/Programs/CrashReporter/CrashReportWebSite/Controllers/DashboardController.cs
+++ b/Engine/Source/Programs/CrashReporter/CrashReportWebSite/Controllers/DashboardController.cs
@@ -56,6 +56,9 @@ namespace Tools.CrashReporter.CrashReportWebSite.Controllers
 		/// <summary>Fake id for all user groups</summary>
 		public static readonly int AllUserGroupId = -1;
 
+		CrashRepository _Crashes = new CrashRepository();
+		BuggRepository _Buggs = new BuggRepository();
+
 		/// <summary>
 		/// Return a dictionary of crashes per group grouped by week.
 		/// </summary>
@@ -68,7 +71,7 @@ namespace Tools.CrashReporter.CrashReportWebSite.Controllers
 			{
 				Dictionary<DateTime, int> Results = new Dictionary<DateTime, int>();
 
-				var UsersIds = FRepository.Get().GetUserIdsFromUserGroupId( UserGroupId );
+				var UsersIds = FRepository.Get( _Crashes ).GetUserIdsFromUserGroupId( UserGroupId );
 
 				// Trim crashes to user group.
 				if( UserGroupId != DashboardController.AllUserGroupId )
@@ -107,7 +110,7 @@ namespace Tools.CrashReporter.CrashReportWebSite.Controllers
 			{
 				Dictionary<DateTime, int> Results = new Dictionary<DateTime, int>();
 
-				var UsersIds = FRepository.Get().GetUserIdsFromUserGroupId( UserGroupId );
+				var UsersIds = FRepository.Get( _Crashes ).GetUserIdsFromUserGroupId( UserGroupId );
 
 				// Trim crashes to user group.
 				if( UserGroupId != DashboardController.AllUserGroupId )
@@ -147,7 +150,7 @@ namespace Tools.CrashReporter.CrashReportWebSite.Controllers
 
 				FAutoScopedLogTimer LogTimerSQL = new FAutoScopedLogTimer( "CrashesFilterByDate", "", "" );
 
-				IQueryable<Crash> CrashesInTimeFrame = FRepository.Get().Crashes.ListAll()
+				IQueryable<Crash> CrashesInTimeFrame = _Crashes.ListAll()
 					.Where( MyCrash => MyCrash.TimeOfCrash >= AfewMonthsAgo && MyCrash.TimeOfCrash <= Today.AddDays( 1 ) );
 				//IEnumerable<Crash> Crashes = FRepository.Get().Crashes.FilterByDate( FRepository.Get().Crashes.ListAll(), AfewMonthsAgo, Today );
 				var VMinimalCrashes = CrashesInTimeFrame.Select( Crash => new { TimeOfCrash = Crash.TimeOfCrash.Value, UserID = Crash.UserNameId.Value } ).ToList();
@@ -159,11 +162,11 @@ namespace Tools.CrashReporter.CrashReportWebSite.Controllers
 				}
 				LogTimerSQL.Dispose();
 
-				int GeneralUserGroupId = FRepository.Get().FindOrAddGroup( "General" );
-				int CoderUserGroupId = FRepository.Get().FindOrAddGroup( "Coder" );
-				int EngineQAUserGroupId = FRepository.Get().FindOrAddGroup( "EngineQA" );
-				int GameQAUserGroupId = FRepository.Get().FindOrAddGroup( "GameQA" );
-				int AnonymousUserGroupId = FRepository.Get().FindOrAddGroup( "Anonymous" );
+				int GeneralUserGroupId = FRepository.Get( _Crashes ).FindOrAddGroup( "General" );
+				int CoderUserGroupId = FRepository.Get( _Crashes ).FindOrAddGroup( "Coder" );
+				int EngineQAUserGroupId = FRepository.Get( _Crashes ).FindOrAddGroup( "EngineQA" );
+				int GameQAUserGroupId = FRepository.Get( _Crashes ).FindOrAddGroup( "GameQA" );
+				int AnonymousUserGroupId = FRepository.Get( _Crashes ).FindOrAddGroup( "Anonymous" );
 
 				Dictionary<DateTime, int> GeneralResults = GetWeeklyCountsByGroup( MinimalCrashes, GeneralUserGroupId );
 				Dictionary<DateTime, int> CoderResults = GetWeeklyCountsByGroup( MinimalCrashes, CoderUserGroupId );
@@ -180,7 +183,7 @@ namespace Tools.CrashReporter.CrashReportWebSite.Controllers
 				Dictionary<DateTime, int> DailyAllResults = GetDailyCountsByGroup( MinimalCrashes, AllUserGroupId );
 
 				// Get daily buggs stats.
-				List<Bugg> Buggs = FRepository.Get().Buggs.ListAll().Where( Bugg => Bugg.TimeOfFirstCrash >= AfewMonthsAgo ).ToList();
+				List<Bugg> Buggs = _Buggs.ListAll().Where( Bugg => Bugg.TimeOfFirstCrash >= AfewMonthsAgo ).ToList();
 
 				Dictionary<DateTime, int> BuggDailyAllResults  =
 				(
diff --git a/Engine/Source/Programs/CrashReporter/CrashReportWebSite/Controllers/ReportsController.cs b/Engine/Source/Programs/CrashReporter/CrashReportWebSite/Controllers/ReportsController.cs
--- a/Engine/Source/Programs/CrashReporter/CrashReportWebSite/Controllers/ReportsController.cs
+++ b/Engine/Source/Programs/CrashReporter/CrashReportWebSite/Controllers/ReportsController.cs
@@ -32,6 +32,9 @@ namespace Tools.CrashReporter.CrashReportWebSite.Controllers
 		/// <returns>A view to display the filtered Buggs.</returns>
 		public ReportsViewModel GetResults( FormHelper FormData, int BuggIDToBeAddedToJira )
 		{
+			BuggRepository BuggsRepo = new BuggRepository();
+			CrashRepository CrashRepo = new CrashRepository();
+
 			// @TODO yrx 2015-02-17 BuggIDToBeAddedToJira replace with List<int> based on check box and Submit?
 			// Enumerate JIRA projects if needed.
 			// https://jira.ol.epicgames.net//rest/api/2/project
@@ -43,15 +46,15 @@ namespace Tools.CrashReporter.CrashReportWebSite.Controllers
 			{
 				string AnonumousGroup = "Anonymous";
 				//List<String> Users = FRepository.Get().GetUserNamesFromGroupName( AnonumousGroup );
-				int AnonymousGroupID = FRepository.Get().FindOrAddGroup( AnonumousGroup );
-				HashSet<int> AnonumousIDs = FRepository.Get().GetUserIdsFromUserGroup( AnonumousGroup );
+				int AnonymousGroupID = FRepository.Get( BuggsRepo ).FindOrAddGroup( AnonumousGroup );
+				HashSet<int> AnonumousIDs = FRepository.Get( BuggsRepo ).GetUserIdsFromUserGroup( AnonumousGroup );
 				int AnonymousID = AnonumousIDs.First();
-				HashSet<string> UserNamesForUserGroup = FRepository.Get().GetUserNamesFromGroupName( AnonumousGroup );
+				HashSet<string> UserNamesForUserGroup = FRepository.Get( BuggsRepo ).GetUserNamesFromGroupName( AnonumousGroup );
 
 				//FormData.DateFrom = DateTime.Now.AddDays( -1 );
 
-				var Crashes = FRepository.Get().Crashes
-					.FilterByDate( FRepository.Get().Crashes.ListAll(), FormData.DateFrom, FormData.DateTo.AddDays( 1 ) )
+				var Crashes = CrashRepo
+					.FilterByDate( CrashRepo.ListAll(), FormData.DateFrom, FormData.DateTo.AddDays( 1 ) )
 					// Only crashes and asserts
 					.Where( Crash => Crash.CrashType == 1 || Crash.CrashType == 2 )
 					// Only anonymous user
@@ -130,7 +133,7 @@ namespace Tools.CrashReporter.CrashReportWebSite.Controllers
 				// Total # of AFFECTED USERS (Anonymous) in timeframe
 				int TotalAffectedUsers = UniqueMachines.Count;
 
-				var RealBuggs = FRepository.Get().Context.Buggs.Where( Bugg => PatternAndCount.Keys.Contains( Bugg.Pattern ) ).ToList();
+				var RealBuggs = BuggsRepo.Context.Buggs.Where( Bugg => PatternAndCount.Keys.Contains( Bugg.Pattern ) ).ToList();
 
 				// Build search string.
 				List<string> FoundJiras = new List<string>();
diff --git a/Engine/Source/Programs/CrashReporter/CrashReportWebSite/Controllers/UsersController.cs b/Engine/Source/Programs/CrashReporter/CrashReportWebSite/Controllers/UsersController.cs
--- a/Engine/Source/Programs/CrashReporter/CrashReportWebSite/Controllers/UsersController.cs
+++ b/Engine/Source/Programs/CrashReporter/CrashReportWebSite/Controllers/UsersController.cs
@@ -24,67 +24,23 @@ namespace Tools.CrashReporter.CrashReportWebSite.Controllers
 		/// <returns>A view to show a list of users in the current user group.</returns>
 		public ActionResult Index( FormCollection UserForms, string UserGroup )
 		{
+			CrashRepository Crashes = new CrashRepository();
+
 			// Examine an incoming form for a new user group
 			foreach( var FormInstance in UserForms )
 			{
 				string UserName = FormInstance.ToString();
 				string NewUserGroup = UserForms[UserName];
-				FRepository.Get().SetUserGroup( UserName, NewUserGroup );
+				FRepository.Get( Crashes ).SetUserGroup( UserName, NewUserGroup );
 			}
 
 			UsersViewModel Model = new UsersViewModel();
 
 			Model.UserGroup = UserGroup;
-			Model.Users = FRepository.Get().GetUserNamesFromGroupName( UserGroup );
-			Model.GroupCounts = GetCountsByGroup();
+			Model.Users = FRepository.Get( Crashes ).GetUserNamesFromGroupName( UserGroup );
+			Model.GroupCounts = FRepository.Get( Crashes ).GetCountsByGroup();
 
 			return View( "Index", Model );
 		}
-
-		/// <summary>
-		/// Gets a container of crash counts per user group for all crashes.
-		/// </summary>
-		/// <returns>A dictionary of user group names, and the count of crashes for each group.</returns>
-		public Dictionary<string, int> GetCountsByGroup()
-		{
-			// @TODO yrx 2014-11-06 Optimize?
-			using( FAutoScopedLogTimer LogTimer = new FAutoScopedLogTimer( this.GetType().ToString() + " SQL OPT" ) )
-			{
-				Dictionary<string, int> Results = new Dictionary<string, int>();
-				var Context = FRepository.Get().Context;
-
-				try
-				{
-					var GroupCounts =
-					(
-						from UserDetail in Context.Users
-						join UserGroupDetail in Context.UserGroups on UserDetail.UserGroupId equals UserGroupDetail.Id
-						group UserDetail by UserGroupDetail.Name into GroupCount
-						select new { Key = GroupCount.Key, Count = GroupCount.Count() }
-					);
-
-					foreach( var GroupCount in GroupCounts )
-					{
-						Results.Add( GroupCount.Key, GroupCount.Count );
-					}
-
-					// Add in all groups, even though there are no crashes associated
-					IEnumerable<string> UserGroups = ( from UserGroupDetail in Context.UserGroups select UserGroupDetail.Name );
-					foreach( string UserGroupName in UserGroups )
-					{
-						if( !Results.Keys.Contains( UserGroupName ) )
-						{
-							Results[UserGroupName] = 0;
-						}
-					}
-				}
-				catch( Exception Ex )
-				{
-					Debug.WriteLine( "Exception in GetCountsByGroup: " + Ex.ToString() );
-				}
-
-				return Results;
-			}
-		}
 	}
 }
diff --git a/Engine/Source/Programs/CrashReporter/CrashReportWebSite/Models/BuggRepository.cs b/Engine/Source/Programs/CrashReporter/CrashReportWebSite/Models/BuggRepository.cs
--- a/Engine/Source/Programs/CrashReporter/CrashReportWebSite/Models/BuggRepository.cs
+++ b/Engine/Source/Programs/CrashReporter/CrashReportWebSite/Models/BuggRepository.cs
@@ -17,8 +17,45 @@ namespace Tools.CrashReporter.CrashReportWebSite.Models
 	/// <summary>
 	/// The class to handle the processing of bucketed crashes a.k.a. Buggs.
 	/// </summary>
-	public class BuggRepository
+	public class BuggRepository : IDisposable
 	{
+		/// <summary>
+		/// Context
+		/// </summary>
+		public CrashReportDataContext Context;
+
+		/// <summary>
+		/// The default constructor.
+		/// </summary>
+		public BuggRepository()
+		{
+			Context = new CrashReportDataContext();
+		}
+
+		/// <summary> Submits enqueue changes to the database. </summary>
+		public void SubmitChanges()
+		{
+			Context.SubmitChanges();
+		}
+
+		/// <summary>
+		/// Implementing Dispose.
+		/// </summary>
+		public void Dispose()
+		{
+			Dispose( true );
+			GC.SuppressFinalize( this );
+		}
+
+		/// <summary>
+		/// Disposes the resources.
+		/// </summary>
+		/// <param name="Disposing">true if the Dispose call is from user code, and not system code.</param>
+		protected virtual void Dispose( bool Disposing )
+		{
+			Context.Dispose();
+		}
+	
 		private const string DefaultUserGroup = "General";
 
 		/// <summary>
@@ -30,8 +67,6 @@ namespace Tools.CrashReporter.CrashReportWebSite.Models
 		{
 			try
 			{
-				var Context = FRepository.Get().Context;
-
 				string Query = "UPDATE Crashes SET Status = {0} WHERE Id IN ( SELECT CrashId FROM Buggs_Crashes WHERE BuggId = {1} )";
 				Context.ExecuteCommand( Query, Status, BuggId );
 
@@ -54,8 +89,6 @@ namespace Tools.CrashReporter.CrashReportWebSite.Models
 		{
 			try
 			{
-				var Context = FRepository.Get().Context;
-
 				string Query = "UPDATE Crashes SET FixedChangeList = {0} WHERE Id IN ( SELECT CrashId FROM Buggs_Crashes WHERE BuggId = {1} )";
 				Context.ExecuteCommand( Query, FixedChangeList, BuggId );
 
@@ -77,8 +110,6 @@ namespace Tools.CrashReporter.CrashReportWebSite.Models
 		{
 			try
 			{
-				var Context = FRepository.Get().Context;
-
 				using( FAutoScopedLogTimer LogTimer = new FAutoScopedLogTimer( "SetJIRAForBuggAndCrashes (" + BuggId + ")" ) )
 				{
 					string Query = "UPDATE Crashes SET TTPID = {0} WHERE Id IN ( SELECT CrashId FROM Buggs_Crashes WHERE BuggId = {1} )";
@@ -103,7 +134,6 @@ namespace Tools.CrashReporter.CrashReportWebSite.Models
 		{
 			using( FAutoScopedLogTimer LogTimer = new FAutoScopedLogTimer( this.GetType().ToString() + "(" + Id + ")" ) )
 			{
-				var Context = FRepository.Get().Context;
 				Bugg Result = null;
 
 				try
@@ -132,8 +162,6 @@ namespace Tools.CrashReporter.CrashReportWebSite.Models
 		{
 			using( FAutoScopedLogTimer LogTimer = new FAutoScopedLogTimer( this.GetType().ToString() ) )
 			{
-				var Context = FRepository.Get().Context;
-
 				try
 				{
 					return Context.FunctionCalls.Where( FunctionCallInstance => FunctionCallInstance.Call.Contains( FunctionCallName ) ).Select( X => X.Id ).ToList();
@@ -154,7 +182,7 @@ namespace Tools.CrashReporter.CrashReportWebSite.Models
 		/// <returns>A list of callstack lines.</returns>
 		public List<string> GetFunctionCalls( string Pattern )
 		{
-			CachedDataService CachedResults = new CachedDataService( HttpContext.Current.Cache );
+			CachedDataService CachedResults = new CachedDataService( HttpContext.Current.Cache, this );
 			List<string> FunctionCalls = CachedResults.GetFunctionCalls( Pattern );
 			return FunctionCalls;
 		}
@@ -170,7 +198,6 @@ namespace Tools.CrashReporter.CrashReportWebSite.Models
 			using( FAutoScopedLogTimer LogTimer = new FAutoScopedLogTimer( this.GetType().ToString() + "(Ids.Count=" + Ids.Count + ")" ) )
 			{
 				List<string> FunctionCalls = new List<string>();
-				var Context = FRepository.Get().Context;
 				try
 				{
 					List<FunctionCall> Funcs = Context.FunctionCalls.Where( FuncCall => Ids.Contains( FuncCall.Id ) ).ToList();
@@ -208,8 +235,6 @@ namespace Tools.CrashReporter.CrashReportWebSite.Models
 				List<int> UserNameIds = new List<int>();
 				int CrashCount = 0;
 
-				var Context = FRepository.Get().Context;
-
 				// Set or return min date max date while we're iterating through the crashes
 				bool bHasChanges = false;
 				try
@@ -269,7 +294,7 @@ namespace Tools.CrashReporter.CrashReportWebSite.Models
 
 					if( bHasChanges )
 					{
-						Context.SubmitChanges();
+						SubmitChanges();
 					}
 				}
 				catch( Exception Ex )
@@ -290,8 +315,6 @@ namespace Tools.CrashReporter.CrashReportWebSite.Models
 		{
 			try
 			{
-				var Context = FRepository.Get().Context;
-
 				// Make sure we don't already have this relationship
 				if( Context.Buggs_Crashes.Where( BuggInstance => BuggInstance.CrashId == CurrentCrash.Id && BuggInstance.BuggId == Bugg.Id ).Count() < 1 )
 				{
@@ -321,8 +344,6 @@ namespace Tools.CrashReporter.CrashReportWebSite.Models
 		{
 			try
 			{
-				var Context = FRepository.Get().Context;
-
 				int BuggUserCount = Context.Buggs_Users.Where( BuggUserInstance => BuggUserInstance.BuggId == Bugg.Id && BuggUserInstance.UserNameId == UserNameId ).Count();
 				if( BuggUserCount < 1 )
 				{
@@ -348,7 +369,7 @@ namespace Tools.CrashReporter.CrashReportWebSite.Models
 		/// <returns>A container of all known Buggs.</returns>
 		public IQueryable<Bugg> ListAll()
 		{
-			return FRepository.Get().Context.Buggs.AsQueryable();
+			return Context.Buggs.AsQueryable();
 		}
 
 		/// <summary>
@@ -504,10 +525,9 @@ namespace Tools.CrashReporter.CrashReportWebSite.Models
 			using( FAutoScopedLogTimer LogTimer = new FAutoScopedLogTimer( this.GetType().ToString() + " SQL OPT" ) )
 			{
 				Dictionary<string, int> Results = new Dictionary<string, int>();
-				var Context = FRepository.Get().Context;
 				try
 				{
-					/*Results =
+					Results =
 					(
 						from BuggDetail in Buggs
 						join BuggsUserGroupDetail in Context.Buggs_UserGroups on BuggDetail.Id equals BuggsUserGroupDetail.BuggId
@@ -524,8 +544,9 @@ namespace Tools.CrashReporter.CrashReportWebSite.Models
 						{
 							Results[UserGroupName] = 0;
 						}
-					}*/
+					}
 
+					/*
 					//
 					var UsersIDsAndGroupIDs = Context.Users.Select( User => new { UserId = User.Id, UserGroupId = User.UserGroupId } ).ToList();
 					var UserGroupArray = Context.UserGroups.ToList();
@@ -552,8 +573,8 @@ namespace Tools.CrashReporter.CrashReportWebSite.Models
 						(
 							from BuggCrash in Context.Buggs_Crashes
 							where BuggCrash.BuggId == Bugg.Id
-							select BuggCrash.Crash.UserNameId.Value
-						).AsEnumerable();
+							select BuggCrash.Crash.UserNameId.GetValueOrDefault()
+						).ToList();
 
 						UserNameIds.UnionWith( CrashList );
 					}
@@ -562,7 +583,7 @@ namespace Tools.CrashReporter.CrashReportWebSite.Models
 					{
 						string UserGroupName = UserIdToGroupName[UserId];
 						Results[UserGroupName]++;
-					}
+					}*/
 				}
 				catch( Exception Ex )
 				{
@@ -626,7 +647,6 @@ namespace Tools.CrashReporter.CrashReportWebSite.Models
 			{
 				IQueryable<Bugg> NewSetOfBuggs = null;
 				IQueryable<Bugg> SetOfBuggsQueryable = SetOfBuggs.AsQueryable();
-				var Context = FRepository.Get().Context;
 
 				try
 				{
@@ -682,10 +702,8 @@ namespace Tools.CrashReporter.CrashReportWebSite.Models
 			{
 				try
 				{
-					var Context = FRepository.Get().Context;
-
 					// Get the group id and grab all buggs for the specified group.
-					HashSet<string> UserNamesForUserGroup = FRepository.Get().GetUserNamesFromGroupName( GroupName );
+					HashSet<string> UserNamesForUserGroup = FRepository.Get( this ).GetUserNamesFromGroupName( GroupName );
 
 					// Simplified query.
 					var BuggIdToCountMapGroup = new Dictionary<int, int>();
diff --git a/Engine/Source/Programs/CrashReporter/CrashReportWebSite/Models/CachedDataService.cs b/Engine/Source/Programs/CrashReporter/CrashReportWebSite/Models/CachedDataService.cs
--- a/Engine/Source/Programs/CrashReporter/CrashReportWebSite/Models/CachedDataService.cs
+++ b/Engine/Source/Programs/CrashReporter/CrashReportWebSite/Models/CachedDataService.cs
@@ -10,8 +10,11 @@ namespace Tools.CrashReporter.CrashReportWebSite.Models
 	/// <summary>
 	/// A class to handle caching of database results on the web server.
 	/// </summary>
-	public class CachedDataService 
+	public class CachedDataService : IDisposable
 	{
+		private CrashRepository Crashes;
+		private BuggRepository Buggs;
+
 		private Cache CacheInstance;
 
 		private const string CacheKeyPrefix = "_CachedDataService_";
@@ -23,9 +26,41 @@ namespace Tools.CrashReporter.CrashReportWebSite.Models
 		/// Link the Http context cache to a crash repository.
 		/// </summary>
 		/// <param name="InCache">The current Http context cache.</param>
-		public CachedDataService( Cache InCache )
+		/// <param name="InCrashRepository">The repository to associate the cache with.</param>
+		public CachedDataService( Cache InCache, CrashRepository InCrashRepository )
 		{
 			CacheInstance = InCache;
+			Crashes = InCrashRepository;
+		}
+
+		/// <summary>
+		/// Link the Http context cache to a Bugg repository.
+		/// </summary>
+		/// <param name="InCache">The current Http context cache.</param>
+		/// <param name="InBuggRepository">The repository to associate the cache with.</param>
+		public CachedDataService( Cache InCache, BuggRepository InBuggRepository )
+		{
+			CacheInstance = InCache;
+			Buggs = InBuggRepository;
+		}
+
+		/// <summary>
+		/// Implementing Dispose.
+		/// </summary>
+		public void Dispose()
+		{
+			Dispose( true );
+			GC.SuppressFinalize( this );
+		}
+
+		/// <summary>
+		/// Disposes the resources.
+		/// </summary>
+		/// <param name="Disposing">true if the Dispose call is from user code, and not system code.</param>
+		protected virtual void Dispose( bool Disposing )
+		{
+			Crashes.Dispose();
+			Buggs.Dispose();
 		}
 
 		/// <summary>
@@ -84,7 +119,7 @@ namespace Tools.CrashReporter.CrashReportWebSite.Models
 					IdList.Remove( 64334 ); // 64334	UE4Editor_Core!FOutputDevice::Logf__VA()
 
 
-					FunctionCalls = FRepository.Get().Buggs.GetFunctionCalls( IdList );
+					FunctionCalls = Buggs.GetFunctionCalls( IdList );
 					CacheInstance.Insert( Key, FunctionCalls );
 				}
 
diff --git a/Engine/Source/Programs/CrashReporter/CrashReportWebSite/Models/CrashReport.cs b/Engine/Source/Programs/CrashReporter/CrashReportWebSite/Models/CrashReport.cs
--- a/Engine/Source/Programs/CrashReporter/CrashReportWebSite/Models/CrashReport.cs
+++ b/Engine/Source/Programs/CrashReporter/CrashReportWebSite/Models/CrashReport.cs
@@ -170,7 +170,8 @@ namespace Tools.CrashReporter.CrashReportWebSite.Models
 				if( !string.IsNullOrEmpty( Key ) )
 				{
 					TTPID = Key;
-					FRepository.Get().Buggs.SetJIRAForBuggAndCrashes( Key, Id );
+					BuggRepository Buggs = new BuggRepository();
+					Buggs.SetJIRAForBuggAndCrashes( Key, Id );
 				}
 
 			}
@@ -237,10 +238,10 @@ namespace Tools.CrashReporter.CrashReportWebSite.Models
 		/// <returns></returns>
 		public List<Crash> GetCrashes()
 		{
-			var Context = FRepository.Get().Context;
+			CrashRepository CrashRepo = new CrashRepository();
 			var CrashList =
 			(
-					from BuggCrash in Context.Buggs_Crashes
+					from BuggCrash in CrashRepo.Context.Buggs_Crashes
 					where BuggCrash.BuggId == Id
 					select BuggCrash.Crash
 			).AsEnumerable().ToList();
@@ -313,7 +314,8 @@ namespace Tools.CrashReporter.CrashReportWebSite.Models
 		{
 			using( FAutoScopedLogTimer LogTimer = new FAutoScopedLogTimer( this.GetType().ToString() + "(Id=" + this.Id + ")" ) )
 			{
-				List<string> Results = FRepository.Get().Buggs.GetFunctionCalls( Pattern );
+				BuggRepository Buggs = new BuggRepository();
+				List<string> Results = Buggs.GetFunctionCalls( Pattern );
 				return Results;
 			}
 		}
@@ -456,7 +458,8 @@ namespace Tools.CrashReporter.CrashReportWebSite.Models
 		/// <returns>A formatted callstack.</returns>
 		public CallStackContainer GetCallStack()
 		{
-			return FRepository.Get().Crashes.GetCallStack( this );
+			CrashRepository Crashes = new CrashRepository();
+			return Crashes.GetCallStack( this );
 		}
 
 
@@ -467,8 +470,10 @@ namespace Tools.CrashReporter.CrashReportWebSite.Models
 		{
 			using( FAutoScopedLogTimer LogTimer = new FAutoScopedLogTimer( this.GetType().ToString() + "(Crash=" + Id + ")" ) )
 			{
+				CrashRepository Crashes = new CrashRepository();
+
 				List<string> PatternList = new List<string>();
-				var FunctionCalls = FRepository.Get().Context.FunctionCalls;
+				var FunctionCalls = Crashes.Context.FunctionCalls;
 
 				// Get an array of callstack items
 				CallStackContainer CallStack = new CallStackContainer( this );
@@ -495,7 +500,7 @@ namespace Tools.CrashReporter.CrashReportWebSite.Models
 								FunctionCalls.InsertOnSubmit( CurrentFunctionCall );
 							}
 
-							FRepository.Get().SubmitChanges();
+							Crashes.SubmitChanges();
 
 							PatternList.Add( CurrentFunctionCall.Id.ToString() );
 						}
@@ -506,7 +511,7 @@ namespace Tools.CrashReporter.CrashReportWebSite.Models
 						//CrashInstance.Pattern += "+";
 
 
-						FRepository.Get().SubmitChanges();
+						Crashes.SubmitChanges();
 					}
 					catch( Exception Ex )
 					{
diff --git a/Engine/Source/Programs/CrashReporter/CrashReportWebSite/Models/CrashRepository.cs b/Engine/Source/Programs/CrashReporter/CrashReportWebSite/Models/CrashRepository.cs
--- a/Engine/Source/Programs/CrashReporter/CrashReportWebSite/Models/CrashRepository.cs
+++ b/Engine/Source/Programs/CrashReporter/CrashReportWebSite/Models/CrashRepository.cs
@@ -21,8 +21,45 @@ namespace Tools.CrashReporter.CrashReportWebSite.Models
 	/// <summary>
 	/// A model to talk to the database.
 	/// </summary>
-	public class CrashRepository
-	{	
+	public class CrashRepository : IDisposable
+	{
+		/// <summary>
+		/// Context
+		/// </summary>
+		public CrashReportDataContext Context;
+
+		/// <summary>
+		/// The default constructor.
+		/// </summary>
+		public CrashRepository()
+		{
+			Context = new CrashReportDataContext();
+		}
+
+		/// <summary> Submits enqueue changes to the database. </summary>
+		public void SubmitChanges()
+		{
+			Context.SubmitChanges();
+		}
+
+		/// <summary>
+		/// Implementing Dispose.
+		/// </summary>
+		public void Dispose()
+		{
+			Dispose( true );
+			GC.SuppressFinalize( this );
+		}
+
+		/// <summary>
+		/// Disposes the resources.
+		/// </summary>
+		/// <param name="Disposing">true if the Dispose call is from user code, and not system code.</param>
+		protected virtual void Dispose( bool Disposing )
+		{
+			Context.Dispose();
+		}
+		
 		/// <summary>
 		/// Gets the crash from an id.
 		/// </summary>
@@ -32,7 +69,6 @@ namespace Tools.CrashReporter.CrashReportWebSite.Models
 		{
 			using( FAutoScopedLogTimer LogTimer = new FAutoScopedLogTimer( this.GetType().ToString() + "(CrashId" + Id + ")" ) )
 			{
-				var Context = FRepository.Get().Context;
 				try
 				{
 					IQueryable<Crash> Crashes =
@@ -62,8 +98,6 @@ namespace Tools.CrashReporter.CrashReportWebSite.Models
 		{
 			using( FAutoScopedLogTimer LogTimer = new FAutoScopedLogTimer( this.GetType().ToString() + "(CrashId=" + CrashInstance.Id + ")" ) )
 			{
-				var Context = FRepository.Get().Context;
-
 				try
 				{
 					int UserGroupId = CrashInstance.User.UserGroupId;
@@ -83,7 +117,7 @@ namespace Tools.CrashReporter.CrashReportWebSite.Models
 		/// <returns>A container of all known crashes.</returns>
 		public IQueryable<Crash> ListAll()
 		{
-			return FRepository.Get().Context.Crashes.AsQueryable();
+			return Context.Crashes.AsQueryable();
 		}
 
 		/// <summary>
@@ -101,18 +135,30 @@ namespace Tools.CrashReporter.CrashReportWebSite.Models
 				try
 				{
 					string[] Terms = Query.Split( "-, ;+".ToCharArray(), StringSplitOptions.RemoveEmptyEntries );
-					string TermsToUse = "";
+
+					// Iterate over all crashes. 
+					// Brute force search.
 					foreach( string Term in Terms )
 					{
-						if( !TermsToUse.Contains( Term ) )
-						{
-							TermsToUse = TermsToUse + "+" + Term;
-						}
+						Results = Results.Where( X =>
+
+							( !string.IsNullOrEmpty( X.ChangeListVersion ) ? X.ChangeListVersion.ToLower().Contains( Term ) : false ) ||
+							( !string.IsNullOrEmpty( X.PlatformName ) ? X.PlatformName.ToLower().Contains( Term ) : false ) ||
+							( !string.IsNullOrEmpty( X.Summary ) ? X.Summary.ToLower().Contains( Term ) : false ) ||
+							( !string.IsNullOrEmpty( X.Description ) ? X.Description.ToLower().Contains( Term ) : false ) ||
+							( !string.IsNullOrEmpty( X.RawCallStack ) ? X.RawCallStack.ToLower().Contains( Term ) : false ) ||
+							( !string.IsNullOrEmpty( X.CommandLine ) ? X.CommandLine.ToLower().Contains( Term ) : false ) ||
+							( !string.IsNullOrEmpty( X.ComputerName ) ? X.ComputerName.ToLower().Contains( Term ) : false ) ||
+
+							( !string.IsNullOrEmpty( X.Module ) ? X.Module.ToLower().Contains( Term ) : false ) ||
+
+							( !string.IsNullOrEmpty( X.BaseDir ) ? X.BaseDir.ToLower().Contains( Term ) : false ) ||
+							( !string.IsNullOrEmpty( X.TTPID ) ? X.TTPID.ToLower().Contains( Term ) : false ) 
+						);
 					}
 
-					// Search the results by the search terms using IQueryable search
-					var CrashesQueryable = (IQueryable<Crash>)IntermediateQueryable.Search( TermsToUse.Split( "+".ToCharArray() ) );
-					Crashes = CrashesQueryable.AsEnumerable();
+
+					Crashes = Results;
 				}
 				catch( Exception Ex )
 				{
@@ -145,7 +191,6 @@ namespace Tools.CrashReporter.CrashReportWebSite.Models
 			using( FAutoScopedLogTimer LogTimer = new FAutoScopedLogTimer( this.GetType().ToString() ) )
 			{
 				UsersMapping UniqueUser = null;
-				var Context = FRepository.Get().Context;
 
 				IEnumerable<Crash> Results = null;
 				int Skip = ( FormData.Page - 1 ) * FormData.PageSize;
@@ -270,14 +315,14 @@ namespace Tools.CrashReporter.CrashReportWebSite.Models
 				int UserGroupId;
 				if( !string.IsNullOrEmpty( FormData.UserGroup ) )
 				{
-					UserGroupId = FRepository.Get().FindOrAddGroup( FormData.UserGroup );
+					UserGroupId = FRepository.Get( this ).FindOrAddGroup( FormData.UserGroup );
 				}
 				else
 				{
 					UserGroupId = 1;
 				}
 
-				HashSet<int> UserIdsForGroup = FRepository.Get().GetUserIdsFromUserGroupId( UserGroupId );
+				HashSet<int> UserIdsForGroup = FRepository.Get( this ).GetUserIdsFromUserGroupId( UserGroupId );
 
 				using( FScopedLogTimer LogTimer3 = new FScopedLogTimer( "CrashRepository.Results.Users" ) )
 				{
@@ -346,7 +391,6 @@ namespace Tools.CrashReporter.CrashReportWebSite.Models
 			using( FAutoScopedLogTimer LogTimer = new FAutoScopedLogTimer( this.GetType().ToString() ) )
 			{
 				Dictionary<string, int> Results = new Dictionary<string, int>();
-				var Context = FRepository.Get().Context;
 
 				try
 				{
@@ -503,7 +547,7 @@ namespace Tools.CrashReporter.CrashReportWebSite.Models
 		/// <returns>A class containing the parsed callstack.</returns>
 		public CallStackContainer GetCallStack( Crash CrashInstance )
 		{
-			CachedDataService CachedResults = new CachedDataService( HttpContext.Current.Cache );
+			CachedDataService CachedResults = new CachedDataService( HttpContext.Current.Cache, this );
 			return CachedResults.GetCallStack( CrashInstance );
 		}	
 
@@ -537,20 +581,20 @@ namespace Tools.CrashReporter.CrashReportWebSite.Models
 			// Valid MachineID and UserName, updated crash from non-UE4 release
 			if(!string.IsNullOrEmpty(NewCrashInfo.UserName))
 			{
-				NewCrash.UserNameId = FRepository.Get().FindOrAddUser( NewCrashInfo.UserName );
+				NewCrash.UserNameId = FRepository.Get( this ).FindOrAddUser( NewCrashInfo.UserName );
 			}
 			// Valid MachineID and EpicAccountId, updated crash from UE4 release
 			else if(!string.IsNullOrEmpty( NewCrashInfo.EpicAccountId ))
 			{
 				NewCrash.EpicAccountId = NewCrashInfo.EpicAccountId;
-				NewCrash.UserNameId = FRepository.Get().FindOrAddUser( UserNameAnonymous );
+				NewCrash.UserNameId = FRepository.Get( this ).FindOrAddUser( UserNameAnonymous );
 			}
 			// Crash from an older version.
 			else
 			{
 				// MachineGuid for older crashes is obsolete, so ignore it.
 				//NewCrash.ComputerName = NewCrashInfo.MachineGuid;
-				NewCrash.UserNameId = FRepository.Get().FindOrAddUser
+				NewCrash.UserNameId = FRepository.Get( this ).FindOrAddUser
 				(
 					!string.IsNullOrEmpty( NewCrashInfo.UserName ) ? NewCrashInfo.UserName : UserNameAnonymous
 				);
@@ -635,9 +679,8 @@ namespace Tools.CrashReporter.CrashReportWebSite.Models
 			//NewCrash.HasMetaData = true;
 
 			// Add the crash to the database
-			var Context = FRepository.Get().Context;
 			Context.Crashes.InsertOnSubmit( NewCrash );
-			Context.SubmitChanges();
+			SubmitChanges();
 
 			NewID = NewCrash.Id;
 
diff --git a/Engine/Source/Programs/CrashReporter/CrashReportWebSite/Models/Repository.cs b/Engine/Source/Programs/CrashReporter/CrashReportWebSite/Models/Repository.cs
--- a/Engine/Source/Programs/CrashReporter/CrashReportWebSite/Models/Repository.cs
+++ b/Engine/Source/Programs/CrashReporter/CrashReportWebSite/Models/Repository.cs
@@ -11,29 +11,35 @@ using Tools.DotNETCommon;
 namespace Tools.CrashReporter.CrashReportWebSite.Models
 {
 	/// <summary>
-	/// Globally accessible repository for crashes and buggs.
-	/// Defines helper methods etc.
+	/// Globally accessible repository for crashes and buggs. 
+	/// Requires a new instance of crashes or buggs, or both.
+	/// Defines helper methods for the context.
 	/// </summary>
 	public class FRepository : IDisposable
 	{
-		private BuggRepository BuggsRepo = new BuggRepository();
-		private CrashRepository CrashesRepo = new CrashRepository();
-		private CrashReportDataContext DataContext = new CrashReportDataContext();
-
-		private static FRepository Instance = new FRepository();
+		private BuggRepository _Buggs;
+		private CrashRepository _Crashes;
 
 		/// <summary>
-		/// Accesses the singleton.
+		/// Accesses the instance.
 		/// </summary>
-		public static FRepository Get()
+		public static FRepository Get( BuggRepository Buggs )
 		{
-			return Instance;
+			return new FRepository() { _Buggs = Buggs };
 		}
-
-		/// <summary> Submits enqueue changes to the database. </summary>
-		public void SubmitChanges()
+		/// <summary>
+		/// Accesses the instance.
+		/// </summary>
+		public static FRepository Get( CrashRepository Crashes )
 		{
-			Context.SubmitChanges();
+			return new FRepository() { _Crashes = Crashes };
+		}
+		/// <summary>
+		/// Accesses the instance.
+		/// </summary>
+		public static FRepository Get( CrashRepository Crashes, BuggRepository Buggs )
+		{
+			return new FRepository() { _Crashes = Crashes, _Buggs = Buggs };
 		}
 
 		/// <summary>
@@ -51,7 +57,14 @@ namespace Tools.CrashReporter.CrashReportWebSite.Models
 		/// <param name="Disposing">true if the Dispose call is from user code, and not system code.</param>
 		protected virtual void Dispose( bool Disposing )
 		{
-			Context.Dispose();
+			if( Crashes != null )
+			{
+				Crashes.Dispose();
+			}
+			if( Buggs != null )
+			{
+				Buggs.Dispose();
+			}
 		}
 
 		/// <summary>
@@ -61,7 +74,7 @@ namespace Tools.CrashReporter.CrashReportWebSite.Models
 		{
 			get
 			{
-				return BuggsRepo;
+				return _Buggs;
 			}
 		}
 
@@ -72,7 +85,7 @@ namespace Tools.CrashReporter.CrashReportWebSite.Models
 		{
 			get
 			{
-				return CrashesRepo;
+				return _Crashes;
 			}
 		}
 
@@ -83,7 +96,7 @@ namespace Tools.CrashReporter.CrashReportWebSite.Models
 		{
 			get
 			{
-				return DataContext;
+				return Crashes != null ? Crashes.Context : (Buggs != null ? Buggs.Context : null);
 			}
 		}
 
@@ -166,7 +179,7 @@ namespace Tools.CrashReporter.CrashReportWebSite.Models
 					NewUser.UserGroupId = UserGroupId;
 					Context.Users.InsertOnSubmit( NewUser );
 
-					SubmitChanges();
+					Context.SubmitChanges();
 					UserNameId = NewUser.Id;
 				}
 				else
@@ -251,5 +264,50 @@ namespace Tools.CrashReporter.CrashReportWebSite.Models
 				FLogger.WriteException( "AdSetUserGroupdUser: " + Ex.ToString() );
 			}
 		}
+
+		/// <summary>
+		/// Gets a container of crash counts per user group for all crashes.
+		/// </summary>
+		/// <returns>A dictionary of user group names, and the count of crashes for each group.</returns>
+		public Dictionary<string, int> GetCountsByGroup()
+		{
+			// @TODO yrx 2014-11-06 Optimize?
+			using( FAutoScopedLogTimer LogTimer = new FAutoScopedLogTimer( this.GetType().ToString() + " SQL OPT" ) )
+			{
+				Dictionary<string, int> Results = new Dictionary<string, int>();
+
+				try
+				{
+					var GroupCounts =
+					(
+						from UserDetail in Context.Users
+						join UserGroupDetail in Context.UserGroups on UserDetail.UserGroupId equals UserGroupDetail.Id
+						group UserDetail by UserGroupDetail.Name into GroupCount
+						select new { Key = GroupCount.Key, Count = GroupCount.Count() }
+					);
+
+					foreach( var GroupCount in GroupCounts )
+					{
+						Results.Add( GroupCount.Key, GroupCount.Count );
+					}
+
+					// Add in all groups, even though there are no crashes associated
+					IEnumerable<string> UserGroups = ( from UserGroupDetail in Context.UserGroups select UserGroupDetail.Name );
+					foreach( string UserGroupName in UserGroups )
+					{
+						if( !Results.Keys.Contains( UserGroupName ) )
+						{
+							Results[UserGroupName] = 0;
+						}
+					}
+				}
+				catch( Exception Ex )
+				{
+					Debug.WriteLine( "Exception in GetCountsByGroup: " + Ex.ToString() );
+				}
+
+				return Results;
+			}
+		}
 	}
 }
