CrashReport - LinqPad queries

diff --git a/Engine/Source/Programs/CrashReporter/CrashReportWebSite/Queries/FunctionCalls.linq b/Engine/Source/Programs/CrashReporter/CrashReportWebSite/Queries/FunctionCalls.linq
new file mode 100644
--- /dev/null
+++ b/Engine/Source/Programs/CrashReporter/CrashReportWebSite/Queries/FunctionCalls.linq
@@ -0,0 +1,25 @@
+<Query Kind="Program">
+  <Connection>
+    <ID>e40cf789-b18a-4b65-bae6-a0f5440c1ce2</ID>
+    <Persist>true</Persist>
+    <Server>db-09</Server>
+    <Database>CrashReport</Database>
+    <IsProduction>true</IsProduction>
+  </Connection>
+  <Output>DataGrids</Output>
+</Query>
+
+void Main()
+{
+	/*IdList.Remove( 63039 ); // 63039	KERNELBASE!RaiseException()
+IdList.Remove( 63138 ); // 63138	UE4Editor_Core!FDebug::EnsureFailed()
+IdList.Remove( 63137 ); // 63137	UE4Editor_Core!NewReportEnsure()
+IdList.Remove( 63149 ); // 63149	UE4Editor_Core!FOutputDeviceWindowsError::Serialize()
+IdList.Remove( 63151 ); // 63151	UE4Editor_Core!FDebug::AssertFailed()
+IdList.Remove( 63445 ); // 63445	UE4Editor_Core!FDebug::EnsureNotFalseFormatted()
+IdList.Remove( 64334 ); // 64334	UE4Editor_Core!FOutputDevice::Logf__VA()*/
+
+
+}
+
+// Define other methods and classes here
diff --git a/Engine/Source/Programs/CrashReporter/CrashReportWebSite/Queries/Last3Linux.linq b/Engine/Source/Programs/CrashReporter/CrashReportWebSite/Queries/Last3Linux.linq
new file mode 100644
--- /dev/null
+++ b/Engine/Source/Programs/CrashReporter/CrashReportWebSite/Queries/Last3Linux.linq
@@ -0,0 +1,19 @@
+<Query Kind="Expression">
+  <Connection>
+    <ID>e40cf789-b18a-4b65-bae6-a0f5440c1ce2</ID>
+    <Persist>true</Persist>
+    <Server>db-09</Server>
+    <Database>CrashReport</Database>
+    <IsProduction>true</IsProduction>
+  </Connection>
+  <Output>DataGrids</Output>
+</Query>
+
+Crashes
+.Where (c =>c.TimeOfCrash > DateTime.UtcNow.AddDays(-7) )
+.Where (c =>c.PlatformName.Contains("Linux") )
+.Take(3)
+.Select (c => new {
+c.RawCallStack,
+c.TimeOfCrash
+})
\ No newline at end of file
diff --git a/Engine/Source/Programs/CrashReporter/CrashReportWebSite/Queries/Last3Mac.linq b/Engine/Source/Programs/CrashReporter/CrashReportWebSite/Queries/Last3Mac.linq
new file mode 100644
--- /dev/null
+++ b/Engine/Source/Programs/CrashReporter/CrashReportWebSite/Queries/Last3Mac.linq
@@ -0,0 +1,19 @@
+<Query Kind="Expression">
+  <Connection>
+    <ID>e40cf789-b18a-4b65-bae6-a0f5440c1ce2</ID>
+    <Persist>true</Persist>
+    <Server>db-09</Server>
+    <Database>CrashReport</Database>
+    <IsProduction>true</IsProduction>
+  </Connection>
+  <Output>DataGrids</Output>
+</Query>
+
+Crashes
+.Where (c =>c.TimeOfCrash > DateTime.UtcNow.AddDays(-1) )
+.Where (c =>c.PlatformName.Contains("Mac") )
+.Take(3)
+.Select (c => new {
+c.RawCallStack,
+c.TimeOfCrash
+})
\ No newline at end of file
diff --git a/Engine/Source/Programs/CrashReporter/CrashReportWebSite/Queries/Last3Windows.linq b/Engine/Source/Programs/CrashReporter/CrashReportWebSite/Queries/Last3Windows.linq
new file mode 100644
--- /dev/null
+++ b/Engine/Source/Programs/CrashReporter/CrashReportWebSite/Queries/Last3Windows.linq
@@ -0,0 +1,20 @@
+<Query Kind="Expression">
+  <Connection>
+    <ID>e40cf789-b18a-4b65-bae6-a0f5440c1ce2</ID>
+    <Persist>true</Persist>
+    <Server>db-09</Server>
+    <Database>CrashReport</Database>
+    <IsProduction>true</IsProduction>
+  </Connection>
+  <Output>DataGrids</Output>
+</Query>
+
+Crashes
+.Where (c =>c.TimeOfCrash > DateTime.UtcNow.AddDays(-1) )
+.Where (c =>c.PlatformName.Contains("Windows") )
+
+.Take(3)
+.Select (c => new {
+c.RawCallStack,
+c.TimeOfCrash
+})
\ No newline at end of file
diff --git a/Engine/Source/Programs/CrashReporter/CrashReportWebSite/Queries/PerfTest.linq b/Engine/Source/Programs/CrashReporter/CrashReportWebSite/Queries/PerfTest.linq
new file mode 100644
--- /dev/null
+++ b/Engine/Source/Programs/CrashReporter/CrashReportWebSite/Queries/PerfTest.linq
@@ -0,0 +1,34 @@
+<Query Kind="Program">
+  <Connection>
+    <ID>e40cf789-b18a-4b65-bae6-a0f5440c1ce2</ID>
+    <Persist>true</Persist>
+    <Server>db-09</Server>
+    <Database>CrashReport</Database>
+    <IsProduction>true</IsProduction>
+  </Connection>
+</Query>
+
+void Main()
+{
+	int UserGroupId = 1;
+
+	var UsersForGroup = Users.Where (User => User.UserGroupId==UserGroupId ).Select (User => User.Id ).ToList();
+	
+	IQueryable<Crashes> Results = Crashes.Where( X => X.TimeOfCrash.Value > new DateTime(2015,3,3) );
+	
+	Results =
+	(
+		from CrashDetail in Results
+		where UsersForGroup.Contains( CrashDetail.UserNameId.Value)
+		select CrashDetail
+	);
+		
+	var ResultList = Results.ToList();	
+	
+	foreach (var It in Results)
+	{
+		Debug.WriteLine( string.Format( "{0}:{1}", It.UserName, It.UserNameId.Value) ); 
+	}
+}
+
+// Define other methods and classes here
diff --git a/Engine/Source/Programs/CrashReporter/CrashReportWebSite/Queries/Remove-NoPattern.linq b/Engine/Source/Programs/CrashReporter/CrashReportWebSite/Queries/Remove-NoPattern.linq
new file mode 100644
--- /dev/null
+++ b/Engine/Source/Programs/CrashReporter/CrashReportWebSite/Queries/Remove-NoPattern.linq
@@ -0,0 +1,15 @@
+<Query Kind="SQL">
+  <Connection>
+    <ID>e40cf789-b18a-4b65-bae6-a0f5440c1ce2</ID>
+    <Persist>true</Persist>
+    <Server>db-09</Server>
+    <Database>CrashReport</Database>
+    <IsProduction>true</IsProduction>
+  </Connection>
+  <Output>DataGrids</Output>
+</Query>
+
+use crashreport
+delete from crashes
+where pattern is null and timeofcrash > '2015-01-01'
+go
\ No newline at end of file
diff --git a/Engine/Source/Programs/CrashReporter/CrashReportWebSite/Queries/Search.linq b/Engine/Source/Programs/CrashReporter/CrashReportWebSite/Queries/Search.linq
new file mode 100644
--- /dev/null
+++ b/Engine/Source/Programs/CrashReporter/CrashReportWebSite/Queries/Search.linq
@@ -0,0 +1,64 @@
+<Query Kind="Expression">
+  <Connection>
+    <ID>e40cf789-b18a-4b65-bae6-a0f5440c1ce2</ID>
+    <Persist>true</Persist>
+    <Server>db-09</Server>
+    <Database>CrashReport</Database>
+    <IsProduction>true</IsProduction>
+  </Connection>
+  <Output>DataGrids</Output>
+</Query>
+
+//string "GetBestShadowTransform" = "";
+
+Crashes
+.Where( X => X.TimeOfCrash > DateTime.UtcNow.AddDays(-7) &&(
+X.GameName.Contains( "GetBestShadowTransform" ) ||
+							X.ChangeListVersion.Contains( "GetBestShadowTransform" ) ||
+							X.PlatformName.Contains( "GetBestShadowTransform" ) ||
+							X.EngineMode.Contains( "GetBestShadowTransform" ) ||
+							X.Description.Contains( "GetBestShadowTransform" ) ||
+							X.RawCallStack.Contains( "GetBestShadowTransform" ) ||
+							X.CommandLine.Contains( "GetBestShadowTransform" ) ||
+							X.ComputerName.Contains( "GetBestShadowTransform" ) ||
+							X.FixedChangeList.Contains( "GetBestShadowTransform" ) ||
+							X.Module.Contains( "GetBestShadowTransform" ) ||
+							X.BuildVersion.Contains( "GetBestShadowTransform" ) ||
+							X.BaseDir.Contains( "GetBestShadowTransform" ) ||
+							X.TTPID.Contains( "GetBestShadowTransform" ) ||
+							X.Branch.Contains( "GetBestShadowTransform" ) 
+							))
+.Select (X => new 
+{
+	X.Summary,
+	X.GameName,
+	X.ChangeListVersion,
+	X.PlatformName,
+	X.EngineMode,
+	X.Description,
+	X.RawCallStack,
+	X.CommandLine,
+	X.ComputerName,
+	X.FixedChangeList,
+	X.Module,
+	X.BuildVersion,
+	X.BaseDir,
+	X.TTPID,
+	X.Branch
+})
+/*
+	X.GameName.Contains( "GetBestShadowTransform" ) ||
+							X.ChangeListVersion.Contains( "GetBestShadowTransform" ) ||
+							X.PlatformName.Contains( "GetBestShadowTransform" ) ||
+							X.EngineMode.Contains( "GetBestShadowTransform" ) ||
+							X.Description.Contains( "GetBestShadowTransform" ) ||
+							X.RawCallStack.Contains( "GetBestShadowTransform" ) ||
+							X.CommandLine.Contains( "GetBestShadowTransform" ) ||
+							X.ComputerName.Contains( "GetBestShadowTransform" ) ||
+							X.FixedChangeList.Contains( "GetBestShadowTransform" ) ||
+							X.Module.Contains( "GetBestShadowTransform" ) ||
+							X.BuildVersion.Contains( "GetBestShadowTransform" ) ||
+							X.BaseDir.Contains( "GetBestShadowTransform" ) ||
+							X.TTPID.Contains( "GetBestShadowTransform" ) ||
+							X.Branch.Contains( "GetBestShadowTransform" ) 
+*/
\ No newline at end of file
diff --git a/Engine/Source/Programs/CrashReporter/CrashReportWebSite/Queries/Today-Buggs.linq b/Engine/Source/Programs/CrashReporter/CrashReportWebSite/Queries/Today-Buggs.linq
new file mode 100644
--- /dev/null
+++ b/Engine/Source/Programs/CrashReporter/CrashReportWebSite/Queries/Today-Buggs.linq
@@ -0,0 +1,33 @@
+<Query Kind="Expression">
+  <Connection>
+    <ID>e40cf789-b18a-4b65-bae6-a0f5440c1ce2</ID>
+    <Persist>true</Persist>
+    <Server>db-09</Server>
+    <Database>CrashReport</Database>
+    <IsProduction>true</IsProduction>
+  </Connection>
+  <Output>DataGrids</Output>
+</Query>
+
+Buggs
+.Where (c => c.TimeOfLastCrash > DateTime.UtcNow.AddDays(-1))
+.Select (c => new 
+{
+c.TTPID,
+//c.Title,
+//c.Summary,
+//c.Priority,
+//c.Type,
+c.Pattern,//.ToString().Split(new char[]{'+'},StringSplitOptions.RemoveEmptyEntries).Length,
+c.NumberOfCrashes,
+//c.TimeOfFirstCrash,
+c.TimeOfLastCrash,
+//c.Status,
+c.FixedChangeList,
+//c.Description,
+//c.ReproSteps,
+//c.Game,
+c.BuildVersion,
+c.CrashType,
+//c.SummaryV2
+})
\ No newline at end of file
diff --git a/Engine/Source/Programs/CrashReporter/CrashReportWebSite/Queries/Today-Crashes.linq b/Engine/Source/Programs/CrashReporter/CrashReportWebSite/Queries/Today-Crashes.linq
new file mode 100644
--- /dev/null
+++ b/Engine/Source/Programs/CrashReporter/CrashReportWebSite/Queries/Today-Crashes.linq
@@ -0,0 +1,65 @@
+<Query Kind="Expression">
+  <Connection>
+    <ID>e40cf789-b18a-4b65-bae6-a0f5440c1ce2</ID>
+    <Persist>true</Persist>
+    <Server>db-09</Server>
+    <Database>CrashReport</Database>
+    <IsProduction>true</IsProduction>
+  </Connection>
+  <Output>DataGrids</Output>
+</Query>
+
+Crashes
+.Where (c => c.TimeOfCrash > DateTime.UtcNow.AddDays(-0).AddHours(-3))
+.OrderBy (c => c.Id)
+.Select (c => new 
+{
+c.Id,
+//c.Title,
+c.Summary,
+c.GameName,
+//c.Status,
+c.TimeOfCrash,
+c.ChangeListVersion,
+c.PlatformName,
+//c.EngineMode,
+//c.Description,
+c.RawCallStack,
+c.Pattern,
+//c.CommandLine,
+c.ComputerName,
+//c.Selected,
+//c.FixedChangeList,
+//c.LanguageExt,
+c.Module,
+c.BuildVersion,
+c.BaseDir,
+//c.Version,
+//c.UserName,
+//c.TTPID,
+//c.AutoReporterID,
+//c.Processed,
+c.Branch,
+c.CrashType,
+//c.UserNameId,
+//c.SourceContext,
+c.EpicAccountId,
+//c.EngineVersion
+})
+
+/*
+	X.GameName.Contains( Term ) ||
+							X.ChangeListVersion.Contains( Term ) ||
+							X.PlatformName.Contains( Term ) ||
+							X.EngineMode.Contains( Term ) ||
+							X.Description.Contains( Term ) ||
+							X.RawCallStack.Contains( Term ) ||
+							X.CommandLine.Contains( Term ) ||
+							X.ComputerName.Contains( Term ) ||
+							X.FixedChangeList.Contains( Term ) ||
+							X.Module.Contains( Term ) ||
+							X.BuildVersion.Contains( Term ) ||
+							X.BaseDir.Contains( Term ) ||
+							X.TTPID.Contains( Term ) ||
+							X.Branch.Contains( Term ) 
+*/
\ No newline at end of file
