CrashWWW - Removed from the callstack following entries:

"RaiseException",
"FDebug::EnsureFailed",
"Error::Serialize",
"FDebug::AssertFailed",
"FDebug::EnsureNotFalseFormatted",
"FOutputDevice::Logf",
"FMsg::Logf",
"ReportCrash",
"NewReportEnsure",

diff --git a/Engine/Source/Programs/CrashReporter/CrashReportWebSite/Models/CachedDataService.cs b/Engine/Source/Programs/CrashReporter/CrashReportWebSite/Models/CachedDataService.cs
--- a/Engine/Source/Programs/CrashReporter/CrashReportWebSite/Models/CachedDataService.cs
+++ b/Engine/Source/Programs/CrashReporter/CrashReportWebSite/Models/CachedDataService.cs
@@ -109,16 +109,6 @@ namespace Tools.CrashReporter.CrashReportWebSite.Models
 						}
 					}
 
-					// These function calls don't add any new information.
-					IdList.Remove( 63039 ); // 63039	KERNELBASE!RaiseException()
-					IdList.Remove( 63138 ); // 63138	UE4Editor_Core!FDebug::EnsureFailed()
-					IdList.Remove( 63137 ); // 63137	UE4Editor_Core!NewReportEnsure()
-					IdList.Remove( 63149 ); // 63149	UE4Editor_Core!FOutputDeviceWindowsError::Serialize()
-					IdList.Remove( 63151 ); // 63151	UE4Editor_Core!FDebug::AssertFailed()
-					IdList.Remove( 63445 ); // 63445	UE4Editor_Core!FDebug::EnsureNotFalseFormatted()
-					IdList.Remove( 64334 ); // 64334	UE4Editor_Core!FOutputDevice::Logf__VA()
-
-
 					FunctionCalls = Buggs.GetFunctionCalls( IdList );
 					CacheInstance.Insert( Key, FunctionCalls );
 				}
diff --git a/Engine/Source/Programs/CrashReporter/CrashReportWebSite/Models/CrashDataFormater.cs b/Engine/Source/Programs/CrashReporter/CrashReportWebSite/Models/CrashDataFormater.cs
--- a/Engine/Source/Programs/CrashReporter/CrashReportWebSite/Models/CrashDataFormater.cs
+++ b/Engine/Source/Programs/CrashReporter/CrashReportWebSite/Models/CrashDataFormater.cs
@@ -10,6 +10,19 @@ using System.Diagnostics;
 
 namespace Tools.CrashReporter.CrashReportWebSite.Models
 {
+	class CustomFuncComparer : IEqualityComparer<string>
+	{
+		public bool Equals( string x, string y )
+		{
+			return y.IndexOf( x, StringComparison.InvariantCultureIgnoreCase ) != -1;
+		}
+
+		public int GetHashCode( string obj )
+		{
+			return obj.GetHashCode();
+		}
+	}
+
 	/// <summary>
 	/// A class representing the line of a callstack.
 	/// </summary>
@@ -161,19 +174,6 @@ namespace Tools.CrashReporter.CrashReportWebSite.Models
 			{
 				// Default to the first module. This is valid for crashes
 				ModuleName = CallStackEntries[0].ModuleName;
-
-				// If the callstack starts in KERNELBASE, that means we asserted and need to find the first non-core module
-				if( CallStackEntries[0].ModuleName.ToUpper() == "KERNELBASE" )
-				{
-					for( int CallStackEntryIndex = 1; CallStackEntryIndex < CallStackEntries.Count; CallStackEntryIndex++ )
-					{
-						if( CallStackEntries[CallStackEntryIndex].ModuleName.ToUpper() != "UE4_CORE" )
-						{
-							ModuleName = CallStackEntries[CallStackEntryIndex].ModuleName;
-							break;
-						}
-					}
-				}
 			}
 
 			return ModuleName;
@@ -384,7 +384,25 @@ namespace Tools.CrashReporter.CrashReportWebSite.Models
 					ModuleName = CurrentLine.Substring( MacModuleStart + 3, MacModuleEnd - MacModuleStart - 3 ).Trim();
 				}
 
-				CallStackEntries.Add( new CallStackEntry( CurrentLine, ModuleName, FilePath, FuncName, LineNumber ) );
+				// Remove callstack entries that match any of these functions.
+				var FuncsToRemove = new HashSet<string>( new string[] 
+				{ 
+					"RaiseException", 
+					"FDebug::EnsureFailed", 		
+					"Error::Serialize",
+					"FDebug::AssertFailed",
+					"FDebug::EnsureNotFalseFormatted",
+					"FOutputDevice::Logf",
+					"FMsg::Logf",
+					"ReportCrash",
+					"NewReportEnsure",
+				} );
+
+				bool Contains = FuncsToRemove.Contains( FuncName, new CustomFuncComparer() );
+				if( !Contains )
+				{
+					CallStackEntries.Add( new CallStackEntry( CurrentLine, ModuleName, FilePath, FuncName, LineNumber ) );
+				}
 			}
 		}
 
