CrashReportWWW - Fixed crash/bugg summary not showing all data

diff --git a/Engine/Source/Programs/CrashReporter/CrashReportWebSite/Views/Buggs/Show.aspx b/Engine/Source/Programs/CrashReporter/CrashReportWebSite/Views/Buggs/Show.aspx
--- a/Engine/Source/Programs/CrashReporter/CrashReportWebSite/Views/Buggs/Show.aspx
+++ b/Engine/Source/Programs/CrashReporter/CrashReportWebSite/Views/Buggs/Show.aspx
@@ -189,7 +189,7 @@
 					<br />
 					<h3>Error Message</h3>
 					<div id='ErrorMessage'>
-						<%=Model.Bugg.Summary %>
+						<%=Html.DisplayFor( m => m.Bugg.Summary )%>
 					</div>
 				</div>
 				<br />
diff --git a/Engine/Source/Programs/CrashReporter/CrashReportWebSite/Views/Crashes/Show.aspx b/Engine/Source/Programs/CrashReporter/CrashReportWebSite/Views/Crashes/Show.aspx
--- a/Engine/Source/Programs/CrashReporter/CrashReportWebSite/Views/Crashes/Show.aspx
+++ b/Engine/Source/Programs/CrashReporter/CrashReportWebSite/Views/Crashes/Show.aspx
@@ -225,7 +225,7 @@
 					<br />
 					<h3>Error Message</h3>
 					<div id='ErrorMessage'>
-						<%=Model.Crash.Summary %>
+						<%=Html.DisplayFor( m => m.Crash.Summary )%>
 					</div>
 				</div>
 				<br />
diff --git a/Engine/Source/Programs/CrashReporter/CrashReportWebSite/Views/Crashes/ViewCrash.ascx b/Engine/Source/Programs/CrashReporter/CrashReportWebSite/Views/Crashes/ViewCrash.ascx
--- a/Engine/Source/Programs/CrashReporter/CrashReportWebSite/Views/Crashes/ViewCrash.ascx
+++ b/Engine/Source/Programs/CrashReporter/CrashReportWebSite/Views/Crashes/ViewCrash.ascx
@@ -180,7 +180,7 @@
 									<td class="Jira"> <span><a href="https://jira.ol.epicgames.net/browse/<%=CurrentCrash.TTPID%>" target="_blank"><%=CurrentCrash.TTPID%></a></span>  </td>
 									<td class="Branch"><%=CurrentCrash.Branch%>&nbsp;</td>
 									<td class="Description"><span class="TableData"><%=CurrentCrash.Description%>&nbsp;</span></td>
-									<td class="Summary"><%=CurrentCrash.Summary%></td>
+									<td class="Summary"><%=Html.Encode(CurrentCrash.Summary)%></td>
 									<td class="ChangeListVersion"><%=CurrentCrash.ChangeListVersion%></td>
 									<td class="Computer"><%=CurrentCrash.ComputerName%></td>
 									<td class="Platform"><%=CurrentCrash.PlatformName%></td>
