CrashReportLinq - Added RepairDatabase program

diff --git a/Engine/Source/Programs/CrashReporter/CrashReportWebSite/Queries/Fix-RepairDatabase.linq b/Engine/Source/Programs/CrashReporter/CrashReportWebSite/Queries/Fix-RepairDatabase.linq
new file mode 100644
--- /dev/null
+++ b/Engine/Source/Programs/CrashReporter/CrashReportWebSite/Queries/Fix-RepairDatabase.linq
@@ -0,0 +1,257 @@
+<Query Kind="Program">
+  <Connection>
+    <ID>e40cf789-b18a-4b65-bae6-a0f5440c1ce2</ID>
+    <Persist>true</Persist>
+    <Server>db-09</Server>
+    <Database>CrashReport</Database>
+    <IsProduction>true</IsProduction>
+  </Connection>
+  <Output>DataGrids</Output>
+</Query>
+
+// Copyright 1998-2015 Epic Games, Inc. All Rights Reserved.
+
+// Removes all crashes without pattern.
+// Removes all crashes without branch
+// Crashes with invalid callstack are removed
+//	First pattern is set to null to disconnect from a bugg
+//  Crash is removed in the next iteration of the script
+
+static string InvalidFuncsStr ="63039+63164+63165+63391+63392+63164+63165+63391+63392+63039+63039+63164+69428+69429+69430+103713+70829+63164+79166+69432+69433+69434+69435+69436+63164+69428+98671+69430+70883+69432+69433+69434+69435+69436+63164+69428";
+static List<string> InvalidModules = new List<string>(){"uxtheme!","user32!","ntdll!","kernel32!","KERNELBASE!","ntdll!"};
+
+HashSet<int> GetInvalidFuncIds()
+{
+	HashSet<int> Result = new HashSet<int>();
+	var SplitString = InvalidFuncsStr.Split("+".ToCharArray(), StringSplitOptions.RemoveEmptyEntries );
+	foreach (var Id in SplitString)
+	{
+		int IntId = int.Parse(Id);
+		Result.Add( IntId);
+	}
+	return Result;
+}
+
+static Stopwatch Timer = Stopwatch.StartNew();
+static DateTime Date = new DateTime(2015,01,01);
+static TimeSpan Tick = TimeSpan.FromDays(3);
+const int NUM_OPS_PER_BATCH = 64;
+
+void WriteLine( string Line )
+{
+	Debug.WriteLine( string.Format( "{0,5}: {1}", (long)Timer.Elapsed.TotalSeconds, Line ) );
+}
+
+void DeleteNoPatternBatch( List<Crashes> NoPatterBatch )
+{
+	Crashes.DeleteAllOnSubmit( NoPatterBatch );
+	Crashes.Context.SubmitChanges();
+	WriteLine( string.Format( "NoPattern deleted: {0}", NoPatterBatch.Count ) );
+}
+
+void DeleteBuggCrash( List<Buggs_Crashes> BuggCrashBatch )
+{
+	Buggs_Crashes.DeleteAllOnSubmit( BuggCrashBatch );
+	Buggs_Crashes.Context.SubmitChanges();
+	WriteLine( string.Format( "BuggCrashBatch deleted: {0}", BuggCrashBatch.Count ) );
+}
+
+// Deletes all crashes without pattern.
+void DeleteNoPatternAll()
+{
+	var NoPattern = Crashes
+	.Where( c => c.Pattern == null )
+	.Where (c => c.TimeOfCrash> Date)
+	.ToList();
+	
+	var CrashesIds = new HashSet<int>( NoPattern.Select (np => np.Id).ToList() );
+	
+	int NumBuggCrass = Buggs_Crashes.Count (bc => bc.BuggId>0);
+	
+	// Unlink from the bugg.	
+	List<Buggs_Crashes> ToRemove = new List<Buggs_Crashes>();
+	for (int n = 0; n < NoPattern.Count; n ++)
+	{
+		var Crash = NoPattern[n];
+	
+		var BC = Buggs_Crashes.Where (bc => bc.CrashId==Crash.Id).FirstOrDefault();
+		if(BC!=null)
+		{
+			ToRemove.Add(BC);
+		}
+		
+		if (ToRemove.Count == NUM_OPS_PER_BATCH)
+		{
+			DeleteBuggCrash(ToRemove);
+			ToRemove.Clear();
+		}
+	}
+	
+	DeleteBuggCrash(ToRemove);
+	ToRemove.Clear();
+	
+	WriteLine( string.Format( "NoPattern: {0}", NoPattern.Count ) );
+	
+	var NoPatternBatch = new List<Crashes>(NUM_OPS_PER_BATCH);
+	for (int n = 0; n < NoPattern.Count; n ++)
+	{
+		NoPatternBatch.Add(NoPattern[n]);
+		
+		if (NoPatternBatch.Count == NUM_OPS_PER_BATCH)
+		{
+			DeleteNoPatternBatch(NoPatternBatch);
+			NoPatternBatch.Clear();
+		}
+	}
+	
+	// Last batch
+	DeleteNoPatternBatch(NoPatternBatch);
+}
+
+// Sets patterns for all crashes without branch, will be deleted after running the script again.
+void SetNoPatternForNoBranchAll()
+{
+	var NoBranch = Crashes
+	.Where( c => c.Branch == null )
+	.Where (c => c.TimeOfCrash> Date)
+	.ToList();
+	
+	WriteLine( string.Format( "NoBranch: {0}", NoBranch.Count ) );
+	
+	var NoBranchBatch = new List<Crashes>(NUM_OPS_PER_BATCH);
+	for (int n = 0; n < NoBranch.Count; n ++)
+	{
+		NoBranchBatch.Add(NoBranch[n]);		
+		NoBranch[n].Pattern = null;
+		NoBranch[n].Branch = "INVALID";
+		
+		if (n % NUM_OPS_PER_BATCH == 0)
+		{
+			SumbitChanges(NoBranchBatch);
+			NoBranchBatch.Clear();
+		}
+	}
+	
+	// Last batch
+	SumbitChanges(NoBranchBatch);
+}
+
+void SumbitChanges( List<Crashes> CrashesBatch )
+{
+	//Crashes.InsertAllOnSubmit( CrashesBatch );
+	Crashes.Context.SubmitChanges();
+	WriteLine( string.Format( "SubmitBatchChanges: {0}", CrashesBatch.Count ) );
+}
+
+class CustomFuncComparer : IEqualityComparer<string>
+{
+    public bool Equals(string x, string y)
+    {
+        return y.IndexOf( x, StringComparison.InvariantCultureIgnoreCase ) != -1;
+    }
+
+    public int GetHashCode(string obj)
+    {
+        return obj.GetHashCode();
+    }
+}
+
+void SetNoPatternForInvalidPattern()
+{
+	//HashSet<int> InvalidFuncs = GetInvalidFuncIds();
+	var AllFunctionNames = FunctionCalls
+	//.Where (fc => InvalidModules.ToList().Contains(fc.Call, new CustomFuncComparer()))
+	.ToList();
+	
+	var InvalidFuncs = new List<FunctionCalls>();
+	foreach( var Func in AllFunctionNames )
+	{
+		bool bInvalid = InvalidModules.Contains(Func.Call, new CustomFuncComparer() );
+		if (bInvalid)
+		{
+			InvalidFuncs.Add(Func);
+		}
+	}
+	
+	var InvalidFuncIds = new HashSet<int>();
+	foreach (var Func in InvalidFuncs)
+	{
+		InvalidFuncIds.Add(Func.Id);
+	}
+
+	InvalidFuncs.OrderBy (i => i.Call);
+	InvalidFuncIds.OrderBy (ifi => ifi);
+	
+	DateTime StartDate = Date;
+	DateTime EndDate = Date.Add( Tick );
+	while( true )
+	{
+		var AllCrashes = Crashes
+		.Where (c => c.TimeOfCrash > StartDate && c.TimeOfCrash <= EndDate)
+		.Where (c => c.Pattern != null)
+		.ToList();
+		
+		if( AllCrashes.Count == 0 )
+		{
+			break;
+		}
+		
+		WriteLine( string.Format( "AllCrashes: {0} {1} -> {2}", AllCrashes.Count, StartDate.ToString("yyyy-MM-dd"), EndDate.ToString("yyyy-MM-dd") ) );
+		StartDate = EndDate;
+	 	EndDate = EndDate.Add( Tick );
+	
+		var InvalidCrashes = new List<Crashes>();
+		int CrashIndex = 0;
+		foreach (var Crash in AllCrashes)
+		{
+			
+			HashSet<int> PatternIds = new HashSet<int>();
+			var SplitString = Crash.Pattern.Split("+".ToCharArray(), StringSplitOptions.RemoveEmptyEntries );
+			foreach (var Id in SplitString)
+			{
+				int IntId = int.Parse(Id);
+				PatternIds.Add( IntId);
+			}
+		
+			int NumInvalidInPattern = 0;
+			foreach (var Id in PatternIds)
+			{
+				if (InvalidFuncIds.Contains(Id))
+				{
+					NumInvalidInPattern++;
+				}
+			}
+		
+			bool bInvalid = NumInvalidInPattern>0 && (PatternIds.Count - NumInvalidInPattern) < 4;
+			if (bInvalid)
+			{
+				CrashIndex++;
+				InvalidCrashes.Add(Crash);
+				Crash.Pattern = null;
+			}
+			
+//			if (CrashIndex % NUM_OPS_PER_BATCH == 0 && InvalidCrashes.Count > 0)
+//			{
+//				WriteLine( string.Format( "InvalidCrashes: {0}", InvalidCrashes.Count ) );
+//				SumbitChanges(InvalidCrashes);
+//				InvalidCrashes.Clear();
+//			}
+		}
+		
+		if (InvalidCrashes.Count > 0)
+		{
+			WriteLine( string.Format( "InvalidCrashes: {0}", InvalidCrashes.Count ) );
+			SumbitChanges(InvalidCrashes);
+			InvalidCrashes.Clear();
+		}
+	}
+}
+
+void Main()
+{	
+	SetNoPatternForNoBranchAll();
+	SetNoPatternForInvalidPattern();
+	DeleteNoPatternAll();
+}
+
+// Define other methods and classes here
\ No newline at end of file
diff --git a/Engine/Source/Programs/CrashReporter/CrashReportWebSite/Queries/FunctionCalls.linq b/Engine/Source/Programs/CrashReporter/CrashReportWebSite/Queries/FunctionCalls.linq
--- a/Engine/Source/Programs/CrashReporter/CrashReportWebSite/Queries/FunctionCalls.linq
+++ b/Engine/Source/Programs/CrashReporter/CrashReportWebSite/Queries/FunctionCalls.linq
@@ -9,9 +9,9 @@
   <Output>DataGrids</Output>
 </Query>
 
-void Main()
-{
-	/*IdList.Remove( 63039 ); // 63039	KERNELBASE!RaiseException()
+// Copyright 1998-2015 Epic Games, Inc. All Rights Reserved.
+
+/*IdList.Remove( 63039 ); // 63039	KERNELBASE!RaiseException()
 IdList.Remove( 63138 ); // 63138	UE4Editor_Core!FDebug::EnsureFailed()
 IdList.Remove( 63137 ); // 63137	UE4Editor_Core!NewReportEnsure()
 IdList.Remove( 63149 ); // 63149	UE4Editor_Core!FOutputDeviceWindowsError::Serialize()
@@ -19,7 +19,87 @@ IdList.Remove( 63151 ); // 63151	UE4Editor_Core!FDebug::AssertFailed()
 IdList.Remove( 63445 ); // 63445	UE4Editor_Core!FDebug::EnsureNotFalseFormatted()
 IdList.Remove( 64334 ); // 64334	UE4Editor_Core!FOutputDevice::Logf__VA()*/
 
+class CustomFuncComparer : IEqualityComparer<string>
+{
+    public bool Equals(string x, string y)
+    {
+        return y.IndexOf( x, StringComparison.InvariantCultureIgnoreCase ) != -1;
+    }
+
+    public int GetHashCode(string obj)
+    {
+        return obj.GetHashCode();
+    }
+}
+
+string CrashTypeToString( Crashes MyCrash )
+{
+	if(MyCrash.CrashType==1)
+	{
+		return "Crash";
+	}
+	else if(MyCrash.CrashType==2)
+	{
+		return "Assert";
+	}
+	else if(MyCrash.CrashType==3)
+	{
+		return "Ensure";
+	}
+	return "NULL";
+}
 
+void Main()
+{
+	var LatestCrashes = Crashes
+		.Where (c => c.TimeOfCrash > DateTime.UtcNow.AddHours(-1))
+		.OrderBy (c => c.Id);
+	
+	Dictionary<int, string> IdToCallstack = new Dictionary<int,string>();
+	Dictionary<int, string> IdToFunc = new Dictionary<int,string>();
+	
+	int Num = LatestCrashes.Count();
+	
+	int Cx=0;
+	foreach (var Crash in LatestCrashes)
+	{
+		Cx++;
+		string[] Ids = Crash.Pattern.Split( "+".ToCharArray(), StringSplitOptions.RemoveEmptyEntries );
+		List<int> IdList = new List<int>();
+		foreach( string id in Ids )
+		{
+			int i;
+			if( int.TryParse( id, out i ) )
+			{
+				IdList.Add( i );
+			}
+		}
+		var CallstackEntries = FunctionCalls.Where (fc => IdList.Contains(fc.Id)).ToList();
+		
+		List<string> CallstackFuncs = new List<string>();
+		// Order by Ids
+		foreach( int Id in IdList )
+		{
+			var Found = CallstackEntries.Find( FC => FC.Id == Id );
+			CallstackFuncs.Add( Found.Call );
+		}
+		
+		Debug.WriteLine( string.Format( "Crash: {0,4}/{1,4}:{2}", Cx, Num, Crash.Id ));
+		Debug.WriteLine( string.Format( " Module: {0} [{1}]", Crash.Module, CrashTypeToString(Crash))  );
+		int Nx=0;
+		foreach (var Entry in CallstackFuncs)
+		{
+			Nx++;
+			if (Nx > 3)
+			{
+				break;
+			}
+			Debug.WriteLine( string.Format( "    {0,2}: {1}", Nx, Entry ) );
+		}
+		Debug.WriteLine( "" );
+	}	
 }
 
 // Define other methods and classes here
+
+// Define other methods and classes here
\ No newline at end of file
diff --git a/Engine/Source/Programs/CrashReporter/CrashReportWebSite/Queries/Last3Linux.linq b/Engine/Source/Programs/CrashReporter/CrashReportWebSite/Queries/Last3Linux.linq
--- a/Engine/Source/Programs/CrashReporter/CrashReportWebSite/Queries/Last3Linux.linq
+++ b/Engine/Source/Programs/CrashReporter/CrashReportWebSite/Queries/Last3Linux.linq
@@ -9,9 +9,12 @@
   <Output>DataGrids</Output>
 </Query>
 
+// Copyright 1998-2015 Epic Games, Inc. All Rights Reserved.
+
 Crashes
 .Where (c =>c.TimeOfCrash > DateTime.UtcNow.AddDays(-7) )
 .Where (c =>c.PlatformName.Contains("Linux") )
+.OrderByDescending(c => c.TimeOfCrash)
 .Take(3)
 .Select (c => new {
 c.RawCallStack,
diff --git a/Engine/Source/Programs/CrashReporter/CrashReportWebSite/Queries/Last3Mac.linq b/Engine/Source/Programs/CrashReporter/CrashReportWebSite/Queries/Last3Mac.linq
--- a/Engine/Source/Programs/CrashReporter/CrashReportWebSite/Queries/Last3Mac.linq
+++ b/Engine/Source/Programs/CrashReporter/CrashReportWebSite/Queries/Last3Mac.linq
@@ -9,9 +9,12 @@
   <Output>DataGrids</Output>
 </Query>
 
+// Copyright 1998-2015 Epic Games, Inc. All Rights Reserved.
+
 Crashes
 .Where (c =>c.TimeOfCrash > DateTime.UtcNow.AddDays(-1) )
 .Where (c =>c.PlatformName.Contains("Mac") )
+.OrderByDescending(c => c.TimeOfCrash)
 .Take(3)
 .Select (c => new {
 c.RawCallStack,
diff --git a/Engine/Source/Programs/CrashReporter/CrashReportWebSite/Queries/Last3Windows.linq b/Engine/Source/Programs/CrashReporter/CrashReportWebSite/Queries/Last3Windows.linq
--- a/Engine/Source/Programs/CrashReporter/CrashReportWebSite/Queries/Last3Windows.linq
+++ b/Engine/Source/Programs/CrashReporter/CrashReportWebSite/Queries/Last3Windows.linq
@@ -9,10 +9,12 @@
   <Output>DataGrids</Output>
 </Query>
 
+// Copyright 1998-2015 Epic Games, Inc. All Rights Reserved.
+
 Crashes
 .Where (c =>c.TimeOfCrash > DateTime.UtcNow.AddDays(-1) )
 .Where (c =>c.PlatformName.Contains("Windows") )
-
+.OrderByDescending(c => c.TimeOfCrash)
 .Take(3)
 .Select (c => new {
 c.RawCallStack,
diff --git a/Engine/Source/Programs/CrashReporter/CrashReportWebSite/Queries/PerfTest.linq b/Engine/Source/Programs/CrashReporter/CrashReportWebSite/Queries/PerfTest.linq
--- a/Engine/Source/Programs/CrashReporter/CrashReportWebSite/Queries/PerfTest.linq
+++ b/Engine/Source/Programs/CrashReporter/CrashReportWebSite/Queries/PerfTest.linq
@@ -8,6 +8,8 @@
   </Connection>
 </Query>
 
+// Copyright 1998-2015 Epic Games, Inc. All Rights Reserved.
+
 void Main()
 {
 	int UserGroupId = 1;
@@ -31,4 +33,4 @@ void Main()
 	}
 }
 
-// Define other methods and classes here
+// Define other methods and classes here
\ No newline at end of file
diff --git a/Engine/Source/Programs/CrashReporter/CrashReportWebSite/Queries/Remove-NoPattern.linq b/Engine/Source/Programs/CrashReporter/CrashReportWebSite/Queries/Remove-NoPattern.linq
--- a/Engine/Source/Programs/CrashReporter/CrashReportWebSite/Queries/Remove-NoPattern.linq
+++ b/Engine/Source/Programs/CrashReporter/CrashReportWebSite/Queries/Remove-NoPattern.linq
@@ -9,7 +9,21 @@
   <Output>DataGrids</Output>
 </Query>
 
+
+
+
+--use crashreport
+--delete from crashes
+--where branch is null and timeofcrash > '2015-01-01'
+--go
+
+
 use crashreport
 delete from crashes
 where pattern is null and timeofcrash > '2015-01-01'
-go
\ No newline at end of file
+go
+
+--use crashreport
+--delete from crashes
+--where branch is null and timeofcrash > '2015-01-01'
+--go
\ No newline at end of file
diff --git a/Engine/Source/Programs/CrashReporter/CrashReportWebSite/Queries/Search.linq b/Engine/Source/Programs/CrashReporter/CrashReportWebSite/Queries/Search.linq
--- a/Engine/Source/Programs/CrashReporter/CrashReportWebSite/Queries/Search.linq
+++ b/Engine/Source/Programs/CrashReporter/CrashReportWebSite/Queries/Search.linq
@@ -9,6 +9,8 @@
   <Output>DataGrids</Output>
 </Query>
 
+// Copyright 1998-2015 Epic Games, Inc. All Rights Reserved.
+
 //string "GetBestShadowTransform" = "";
 
 Crashes
diff --git a/Engine/Source/Programs/CrashReporter/CrashReportWebSite/Queries/Today-Buggs.linq b/Engine/Source/Programs/CrashReporter/CrashReportWebSite/Queries/Today-Buggs.linq
--- a/Engine/Source/Programs/CrashReporter/CrashReportWebSite/Queries/Today-Buggs.linq
+++ b/Engine/Source/Programs/CrashReporter/CrashReportWebSite/Queries/Today-Buggs.linq
@@ -9,10 +9,14 @@
   <Output>DataGrids</Output>
 </Query>
 
+// Copyright 1998-2015 Epic Games, Inc. All Rights Reserved.
+
 Buggs
-.Where (c => c.TimeOfLastCrash > DateTime.UtcNow.AddDays(-1))
+.Where (c => c.TimeOfFirstCrash > DateTime.UtcNow.AddDays(-1))
+.OrderByDescending (c => c.NumberOfCrashes)
 .Select (c => new 
 {
+c.Id,
 c.TTPID,
 //c.Title,
 //c.Summary,
@@ -20,7 +24,7 @@ c.TTPID,
 //c.Type,
 c.Pattern,//.ToString().Split(new char[]{'+'},StringSplitOptions.RemoveEmptyEntries).Length,
 c.NumberOfCrashes,
-//c.TimeOfFirstCrash,
+c.TimeOfFirstCrash,
 c.TimeOfLastCrash,
 //c.Status,
 c.FixedChangeList,
diff --git a/Engine/Source/Programs/CrashReporter/CrashReportWebSite/Queries/Today-Crashes.linq b/Engine/Source/Programs/CrashReporter/CrashReportWebSite/Queries/Today-Crashes.linq
--- a/Engine/Source/Programs/CrashReporter/CrashReportWebSite/Queries/Today-Crashes.linq
+++ b/Engine/Source/Programs/CrashReporter/CrashReportWebSite/Queries/Today-Crashes.linq
@@ -9,8 +9,10 @@
   <Output>DataGrids</Output>
 </Query>
 
+// Copyright 1998-2015 Epic Games, Inc. All Rights Reserved.
+
 Crashes
-.Where (c => c.TimeOfCrash > DateTime.UtcNow.AddDays(-0).AddHours(-3))
+.Where (c => c.TimeOfCrash > DateTime.UtcNow.AddDays(-1))
 .OrderBy (c => c.Id)
 .Select (c => new 
 {
