CrashReportWWW - Bugg's error message is using the latest crash's summary, removed unused buggs properties

diff --git a/Engine/Source/Programs/CrashReporter/CrashReportWebSite/Controllers/BuggsController.cs b/Engine/Source/Programs/CrashReporter/CrashReportWebSite/Controllers/BuggsController.cs
--- a/Engine/Source/Programs/CrashReporter/CrashReportWebSite/Controllers/BuggsController.cs
+++ b/Engine/Source/Programs/CrashReporter/CrashReportWebSite/Controllers/BuggsController.cs
@@ -221,11 +221,6 @@ namespace Tools.CrashReporter.CrashReportWebSite.Controllers
 						Buggs.SetJIRAForBuggAndCrashes( NewBugg.TTPID, id );
 					}
 
-					if( !string.IsNullOrEmpty( BuggsForm["Description"] ) )
-					{
-						NewBugg.Description = BuggsForm["Description"];
-					}
-
 					// <STATUS>
 				}
 
@@ -257,6 +252,8 @@ namespace Tools.CrashReporter.CrashReportWebSite.Controllers
 
 						Model.SourceContext = NewCrash.SourceContext;
 					}
+
+					Model.Bugg.LatestCrashSummary = NewCrash.Summary;
 				}
 
 				Model.GenerationTime = LogTimer.GetElapsedSeconds().ToString( "F2" );
diff --git a/Engine/Source/Programs/CrashReporter/CrashReportWebSite/Models/CrashReport.cs b/Engine/Source/Programs/CrashReporter/CrashReportWebSite/Models/CrashReport.cs
--- a/Engine/Source/Programs/CrashReporter/CrashReportWebSite/Models/CrashReport.cs
+++ b/Engine/Source/Programs/CrashReporter/CrashReportWebSite/Models/CrashReport.cs
@@ -113,6 +113,9 @@ namespace Tools.CrashReporter.CrashReportWebSite.Models
 		/// <summary></summary>
 		public string JiraFixCL { get; set; }
 
+		/// <summary></summary>
+		public string LatestCrashSummary { get; set; }
+
 		/// <summary>The CL of the oldest build when this bugg is occurring</summary>
 		int ToJiraFirstCLAffected { get; set; }
 
@@ -409,9 +412,9 @@ namespace Tools.CrashReporter.CrashReportWebSite.Models
 				var CrashList =
 				(
 						from BuggCrash in CrashRepo.Context.Buggs_Crashes
-						where BuggCrash.BuggId == Id
+						where BuggCrash.BuggId == Id && BuggCrash.Crash.TimeOfCrash > DateTime.UtcNow.AddMonths(-6)
 						select BuggCrash.Crash
-				).ToList();
+				).OrderByDescending( c => c.TimeOfCrash ).ToList();
 				return CrashList;
 			}
 		}
@@ -512,6 +515,12 @@ namespace Tools.CrashReporter.CrashReportWebSite.Models
 	/// </summary>
 	public partial class Crash
 	{
+		/// <summary> Helper method, display this Bugg as a human readable string. Debugging purpose. </summary>
+		public override string ToString()
+		{
+			return string.Format( "Id={0} TOC={1} BV={2}-{3}+{4} G={5} U={6}", Id, TimeOfCrash, this.BuildVersion, this.BuiltFromCL, this.Branch, this.GameName, this.UserName );
+		}
+
 		/// <summary>A formatted callstack.</summary>
 		public CallStackContainer CallStackContainer { get; set; }
 
diff --git a/Engine/Source/Programs/CrashReporter/CrashReportWebSite/Models/CrashReport.dbml b/Engine/Source/Programs/CrashReporter/CrashReportWebSite/Models/CrashReport.dbml
--- a/Engine/Source/Programs/CrashReporter/CrashReportWebSite/Models/CrashReport.dbml
+++ b/Engine/Source/Programs/CrashReporter/CrashReportWebSite/Models/CrashReport.dbml
@@ -4,9 +4,6 @@
     <Type Name="Bugg">
       <Column Name="Id" Type="System.Int32" DbType="Int NOT NULL IDENTITY" IsPrimaryKey="true" IsDbGenerated="true" CanBeNull="false" />
       <Column Name="TTPID" Type="System.String" DbType="VarChar(50)" CanBeNull="true" />
-      <Column Name="Title" Type="System.String" DbType="VarChar(120)" CanBeNull="true" />
-      <Column Name="SummaryV2" Member="Summary" Type="System.String" DbType="VarChar(512)" CanBeNull="true" />
-      <Column Name="Priority" Type="System.Int32" DbType="Int" CanBeNull="true" />
       <Column Name="Pattern" Type="System.String" DbType="VarChar(800) NOT NULL" CanBeNull="false" />
       <Column Name="NumberOfCrashes" Type="System.Int32" DbType="Int" CanBeNull="true" />
       <Column Name="NumberOfUsers" Type="System.Int32" DbType="Int" CanBeNull="true" />
@@ -14,10 +11,7 @@
       <Column Name="TimeOfLastCrash" Type="System.DateTime" DbType="DateTime" CanBeNull="true" />
       <Column Name="Status" Type="System.String" DbType="VarChar(64)" CanBeNull="true" />
       <Column Name="FixedChangeList" Type="System.String" DbType="VarChar(50)" CanBeNull="true" />
-      <Column Name="Description" Type="System.String" DbType="Text" CanBeNull="true" UpdateCheck="Never" />
-      <Column Name="ReproSteps" Type="System.String" DbType="Text" CanBeNull="true" UpdateCheck="Never" />
-      <Column Name="Game" Type="System.String" DbType="VarChar(120)" CanBeNull="true" />
-      <Column Name="BuildVersion" Type="System.String" DbType="VarChar(128)" CanBeNull="true" />
+      <Column Name="BuildVersion" Type="System.String" DbType="VarChar(64)" CanBeNull="true" />
       <Column Name="CrashType" Type="System.Int16" DbType="SmallInt" CanBeNull="true" />
       <Association Name="Bugg_Buggs_Crash" Member="Buggs_Crashes" ThisKey="Id" OtherKey="BuggId" Type="Buggs_Crash" />
     </Type>
diff --git a/Engine/Source/Programs/CrashReporter/CrashReportWebSite/Models/CrashReport.dbml.layout b/Engine/Source/Programs/CrashReporter/CrashReportWebSite/Models/CrashReport.dbml.layout
--- a/Engine/Source/Programs/CrashReporter/CrashReportWebSite/Models/CrashReport.dbml.layout
+++ b/Engine/Source/Programs/CrashReporter/CrashReportWebSite/Models/CrashReport.dbml.layout
@@ -2,10 +2,10 @@
 <ordesignerObjectsDiagram dslVersion="1.0.0.0" absoluteBounds="0, 0, 13.75, 9.75" name="CrashReport">
   <DataContextMoniker Name="/CrashReportDataContext" />
   <nestedChildShapes>
-    <classShape Id="a37f1cdb-93d7-48ef-8da7-2b7ca09282b0" absoluteBounds="11.5, 1.25, 2, 4.0785139973958326">
+    <classShape Id="a37f1cdb-93d7-48ef-8da7-2b7ca09282b0" absoluteBounds="11.5, 1.25, 2, 2.9247054036458326">
       <DataClassMoniker Name="/CrashReportDataContext/Bugg" />
       <nestedChildShapes>
-        <elementListCompartment Id="7672a5dd-a23f-4509-99c2-fa6215910daf" absoluteBounds="11.515, 1.7100000000000009, 1.9700000000000002, 3.518513997395833" name="DataPropertiesCompartment" titleTextColor="Black" itemTextColor="Black" />
+        <elementListCompartment Id="7672a5dd-a23f-4509-99c2-fa6215910daf" absoluteBounds="11.515, 1.7100000000000009, 1.9700000000000002, 2.364705403645833" name="DataPropertiesCompartment" titleTextColor="Black" itemTextColor="Black" />
       </nestedChildShapes>
     </classShape>
     <classShape Id="01dddf4f-058f-43b4-b6c0-b90e7acfe0b9" absoluteBounds="3.375, 3.875, 2, 1.3862939453124987">
@@ -71,7 +71,7 @@
         <classShapeMoniker Id="8a60fc4a-a797-417f-87ed-25b22bd5e420" />
       </nodes>
     </associationConnector>
-    <associationConnector edgePoints="[(11.5 : 2.47199625651042); (10.875 : 2.47199625651042)]" fixedFrom="NotFixed" fixedTo="NotFixed">
+    <associationConnector edgePoints="[(11.5 : 2.47199625651042); (10.875 : 2.47199625651042)]" fixedFrom="Algorithm" fixedTo="Algorithm">
       <AssociationMoniker Name="/CrashReportDataContext/Bugg/Bugg_Buggs_Crash" />
       <nodes>
         <classShapeMoniker Id="a37f1cdb-93d7-48ef-8da7-2b7ca09282b0" />
diff --git a/Engine/Source/Programs/CrashReporter/CrashReportWebSite/Views/Buggs/Show.aspx b/Engine/Source/Programs/CrashReporter/CrashReportWebSite/Views/Buggs/Show.aspx
--- a/Engine/Source/Programs/CrashReporter/CrashReportWebSite/Views/Buggs/Show.aspx
+++ b/Engine/Source/Programs/CrashReporter/CrashReportWebSite/Views/Buggs/Show.aspx
@@ -15,18 +15,6 @@
 	<script type="text/javascript">
 		$(document).ready(function ()
 		{
-			$("#EditDescription").click(function ()
-			{
-				$("#CrashDescription").css("display", "none");
-				$("#ShowCrashDescription input").css("display", "block");
-				$("#EditDescription").css("display", "none");
-				$("#SaveDescription").css("display", "inline");
-			});
-
-			$("#SaveDescription").click(function ()
-			{
-				$('#EditCrashDescriptionForm').submit();
-			});
 			// Zebra stripes
 			$("#CrashesTable tr:nth-child(even)").css("background-color", "#C3CAD0");
 			$("#CrashesTable tr:nth-child(odd)").css("background-color", "#eeeeee");
@@ -183,13 +171,13 @@
 
 	<div id="CallStackContainer" >
 
-		<% if( !string.IsNullOrEmpty( Model.Bugg.Summary ) ) 
+		<% if( !string.IsNullOrEmpty( Model.Bugg.LatestCrashSummary ) ) 
 			{ %>
 				<div id='ShowErrorMessage'>
 					<br />
 					<h3>Error Message</h3>
 					<div id='ErrorMessage'>
-						<%=Html.DisplayFor( m => m.Bugg.Summary )%>
+						<%=Html.DisplayFor( m => m.Bugg.LatestCrashSummary )%>
 					</div>
 				</div>
 				<br />
@@ -244,16 +232,6 @@
 				<pre><%=Html.DisplayFor( x => Model.SourceContext )%></pre>
 			</div>
 		</div>
-		<%using( Html.BeginForm( "Show", "Buggs", FormMethod.Post, new { id = "EditCrashDescriptionForm"} ) )
-			{ %>
-			<div id='ShowCrashDescription' class='CrashViewTextBox'>
-				<br />
-				<h3>Description <span class='EditButton' id='EditDescription'>Edit</span> <span class='EditButton' id='SaveDescription'>Save</span></h3>
-				<div id='InputCrashDescription'>
-					<%=Html.TextBox( "Description", Model.Bugg.Description )%> 
-				</div>
-			</div>
-		<%} %>
 
 		<div class='CrashViewTextBoxRight'>
 			<h3>Crashes</h3>
