CrashReportLinq - Updated RepairDatabase program

diff --git a/Engine/Source/Programs/CrashReporter/CrashReportWebSite/Queries/Fix-RepairDatabase.linq b/Engine/Source/Programs/CrashReporter/CrashReportWebSite/Queries/Fix-RepairDatabase.linq
--- a/Engine/Source/Programs/CrashReporter/CrashReportWebSite/Queries/Fix-RepairDatabase.linq
+++ b/Engine/Source/Programs/CrashReporter/CrashReportWebSite/Queries/Fix-RepairDatabase.linq
@@ -18,7 +18,7 @@
 //  Crash is removed in the next iteration of the script
 
 static string InvalidFuncsStr ="63039+63164+63165+63391+63392+63164+63165+63391+63392+63039+63039+63164+69428+69429+69430+103713+70829+63164+79166+69432+69433+69434+69435+69436+63164+69428+98671+69430+70883+69432+69433+69434+69435+69436+63164+69428";
-static List<string> InvalidModules = new List<string>(){"uxtheme!","user32!","ntdll!","kernel32!","KERNELBASE!","ntdll!"};
+static List<string> InvalidModules = new List<string>(){"uxtheme!","user32!","ntdll!","kernel32!","KERNELBASE!","ntdll!","wer!","<Unknown>"};
 
 HashSet<int> GetInvalidFuncIds()
 {
@@ -33,27 +33,64 @@ HashSet<int> GetInvalidFuncIds()
 }
 
 static Stopwatch Timer = Stopwatch.StartNew();
-static DateTime Date = new DateTime(2015,01,01);
+static DateTime Date = DateTime.UtcNow.AddMonths(-2);
 static TimeSpan Tick = TimeSpan.FromDays(3);
-const int NUM_OPS_PER_BATCH = 64;
+const int NUM_OPS_PER_BATCH = 32;
 
 void WriteLine( string Line )
 {
 	Debug.WriteLine( string.Format( "{0,5}: {1}", (long)Timer.Elapsed.TotalSeconds, Line ) );
 }
 
-void DeleteNoPatternBatch( List<Crashes> NoPatterBatch )
+int NoPatternTotal = 0;
+
+void DeleteNoPatternBatch( List<Crashes> NoPatternBatch )
 {
-	Crashes.DeleteAllOnSubmit( NoPatterBatch );
-	Crashes.Context.SubmitChanges();
-	WriteLine( string.Format( "NoPattern deleted: {0}", NoPatterBatch.Count ) );
+	if (NoPatternBatch.Count>0)
+	{
+		Crashes.DeleteAllOnSubmit( NoPatternBatch );	
+		Crashes.Context.SubmitChanges();
+		NoPatternTotal += NoPatternBatch.Count;
+		WriteLine( string.Format( "NoPattern deleted: {0}", NoPatternTotal ) );
+		NoPatternBatch.Clear();
+	}
 }
 
+int BuggCrashTotal = 0;
+
 void DeleteBuggCrash( List<Buggs_Crashes> BuggCrashBatch )
 {
-	Buggs_Crashes.DeleteAllOnSubmit( BuggCrashBatch );
-	Buggs_Crashes.Context.SubmitChanges();
-	WriteLine( string.Format( "BuggCrashBatch deleted: {0}", BuggCrashBatch.Count ) );
+	try
+	{
+		if (BuggCrashBatch.Count>0)
+		{
+			Buggs_Crashes.DeleteAllOnSubmit( BuggCrashBatch );
+			Buggs_Crashes.Context.SubmitChanges();
+			BuggCrashTotal += BuggCrashBatch.Count;
+			WriteLine( string.Format( "BuggCrashBatch deleted: {0}", BuggCrashTotal ) );
+			BuggCrashBatch.Clear();
+		}
+	}
+	catch( Exception E )
+	{
+		WriteLine( E.ToString() );
+	}
+}
+
+
+int BuggsTotal = 0;
+
+void DeleteBuggs( List<int> BuggsBatch )
+{
+	if (BuggsBatch.Count>0)
+	{
+		var Batch = Buggs.Where (b => BuggsBatch.Contains(b.Id));
+		BuggsTotal += Batch.Count (b => b.Id>0);
+		Buggs.DeleteAllOnSubmit( Batch );
+		Buggs.Context.SubmitChanges();
+		WriteLine( string.Format( "DeleteBuggs deleted: {0}", BuggsTotal ) );
+		BuggsBatch.Clear();
+	}
 }
 
 // Deletes all crashes without pattern.
@@ -70,6 +107,7 @@ void DeleteNoPatternAll()
 	
 	// Unlink from the bugg.	
 	List<Buggs_Crashes> ToRemove = new List<Buggs_Crashes>();
+	HashSet<int> BuggIds = new HashSet<int>();
 	for (int n = 0; n < NoPattern.Count; n ++)
 	{
 		var Crash = NoPattern[n];
@@ -78,17 +116,35 @@ void DeleteNoPatternAll()
 		if(BC!=null)
 		{
 			ToRemove.Add(BC);
+			BuggIds.Add(BC.BuggId);
 		}
 		
 		if (ToRemove.Count == NUM_OPS_PER_BATCH)
 		{
 			DeleteBuggCrash(ToRemove);
-			ToRemove.Clear();
 		}
 	}
 	
 	DeleteBuggCrash(ToRemove);
-	ToRemove.Clear();
+	
+	List<int> BuggsToRemove = new List<int>();
+	foreach (var BuggId in BuggIds)
+	{
+		var CrashesForBugg = Buggs_Crashes.Where (bc => bc.BuggId==BuggId).Select (bc => bc.CrashId);
+		int CrashCount = CrashesForBugg.Count (cfb => cfb>0);
+		if(CrashCount==0)
+		{
+			// Remove the bugg.
+			BuggsToRemove.Add(BuggId);
+		}
+		
+		if (BuggsToRemove.Count > NUM_OPS_PER_BATCH)
+		{
+			DeleteBuggs( BuggsToRemove );
+		}
+	}
+	
+	DeleteBuggs( BuggsToRemove );
 	
 	WriteLine( string.Format( "NoPattern: {0}", NoPattern.Count ) );
 	
@@ -100,7 +156,6 @@ void DeleteNoPatternAll()
 		if (NoPatternBatch.Count == NUM_OPS_PER_BATCH)
 		{
 			DeleteNoPatternBatch(NoPatternBatch);
-			NoPatternBatch.Clear();
 		}
 	}
 	
@@ -128,7 +183,6 @@ void SetNoPatternForNoBranchAll()
 		if (n % NUM_OPS_PER_BATCH == 0)
 		{
 			SumbitChanges(NoBranchBatch);
-			NoBranchBatch.Clear();
 		}
 	}
 	
@@ -147,7 +201,7 @@ class CustomFuncComparer : IEqualityComparer<string>
 {
     public bool Equals(string x, string y)
     {
-        return y.IndexOf( x, StringComparison.InvariantCultureIgnoreCase ) != -1;
+        return y.StartsWith( x, StringComparison.InvariantCultureIgnoreCase );
     }
 
     public int GetHashCode(string obj)
@@ -206,6 +260,11 @@ void SetNoPatternForInvalidPattern()
 		{
 			
 			HashSet<int> PatternIds = new HashSet<int>();
+			if(Crash.Pattern==null)
+			{
+				continue;
+			}
+			
 			var SplitString = Crash.Pattern.Split("+".ToCharArray(), StringSplitOptions.RemoveEmptyEntries );
 			foreach (var Id in SplitString)
 			{
@@ -229,13 +288,6 @@ void SetNoPatternForInvalidPattern()
 				InvalidCrashes.Add(Crash);
 				Crash.Pattern = null;
 			}
-			
-//			if (CrashIndex % NUM_OPS_PER_BATCH == 0 && InvalidCrashes.Count > 0)
-//			{
-//				WriteLine( string.Format( "InvalidCrashes: {0}", InvalidCrashes.Count ) );
-//				SumbitChanges(InvalidCrashes);
-//				InvalidCrashes.Clear();
-//			}
 		}
 		
 		if (InvalidCrashes.Count > 0)
@@ -249,6 +301,12 @@ void SetNoPatternForInvalidPattern()
 
 void Main()
 {	
+	int NumCrashes = Crashes.Count (c => c.Id>0);
+	int MaxCrashId = Crashes.OrderByDescending (c => c.TimeOfCrash).Take(1).FirstOrDefault().Id;
+	int Diff = MaxCrashId-NumCrashes;
+	WriteLine( string.Format( "NumCrashes: {0}", NumCrashes ) );
+	WriteLine( string.Format( "MaxCrashId: {0}", MaxCrashId ) );
+	WriteLine( string.Format( "Diff:       {0}", Diff ) );
 	SetNoPatternForNoBranchAll();
 	SetNoPatternForInvalidPattern();
 	DeleteNoPatternAll();
