CrashReportWWW - Fix for crash report view crashing due to changes to jira's project

diff --git a/Engine/Source/Programs/CrashReporter/CrashReportWebSite/Controllers/ReportsController.cs b/Engine/Source/Programs/CrashReporter/CrashReportWebSite/Controllers/ReportsController.cs
--- a/Engine/Source/Programs/CrashReporter/CrashReportWebSite/Controllers/ReportsController.cs
+++ b/Engine/Source/Programs/CrashReporter/CrashReportWebSite/Controllers/ReportsController.cs
@@ -137,7 +137,7 @@ namespace Tools.CrashReporter.CrashReportWebSite.Controllers
 				var RealBuggs = BuggsRepo.Context.Buggs.Where( Bugg => PatternAndCount.Keys.Contains( Bugg.Pattern ) ).ToList();
 
 				// Build search string.
-				List<string> FoundJiras = new List<string>();
+				HashSet<string> FoundJiras = new HashSet<string>();
 				Dictionary<string, List<Bugg>> JiraIDtoBugg = new Dictionary<string, List<Bugg>>();
 
 				List<Bugg> Buggs = new List<Bugg>( 100 );
@@ -205,6 +205,8 @@ namespace Tools.CrashReporter.CrashReportWebSite.Controllers
 
 				if( JC.CanBeUsed() )
 				{
+					var BuggsCopy = new List<Bugg>( Buggs );
+
 					// Grab the data form JIRA.
 					string JiraSearchQuery = string.Join( " OR ", FoundJiras );
 
@@ -251,21 +253,40 @@ namespace Tools.CrashReporter.CrashReportWebSite.Controllers
 
 							int FixCL = Jira.Value["customfield_11200"] != null ? (int)(decimal)Jira.Value["customfield_11200"] : 0;
 
-							var BuggsForJira = JiraIDtoBugg[JiraID];
+							List<Bugg> BuggsForJira;
+							JiraIDtoBugg.TryGetValue( JiraID, out BuggsForJira );
+
+							//var BuggsForJira = JiraIDtoBugg[JiraID];
 
-							foreach( Bugg Bugg in BuggsForJira )
+							if( BuggsForJira != null )
 							{
-								Bugg.JiraSummary = Summary;
-								Bugg.JiraComponentsText = ComponentsText;
-								Bugg.JiraResolution = Resolution;
-								Bugg.JiraFixVersionsText = FixVersionsText;
-								if( FixCL != 0 )
+								foreach( Bugg Bugg in BuggsForJira )
 								{
-									Bugg.JiraFixCL = FixCL.ToString();
+									Bugg.JiraSummary = Summary;
+									Bugg.JiraComponentsText = ComponentsText;
+									Bugg.JiraResolution = Resolution;
+									Bugg.JiraFixVersionsText = FixVersionsText;
+									if( FixCL != 0 )
+									{
+										Bugg.JiraFixCL = FixCL.ToString();
+									}
+
+									BuggsCopy.Remove( Bugg );
 								}
 							}
 						}
 					}
+
+					// If there are buggs, we need to update the summary to indicate an error.
+					// Usually caused when bugg's project has changed.
+					foreach( var Bugg in BuggsCopy.Where( b => !string.IsNullOrEmpty( b.TTPID ) ) )
+					{
+						Bugg.JiraSummary = "JIRA MISMATCH";
+						Bugg.JiraComponentsText = "JIRA MISMATCH";
+						Bugg.JiraResolution = "JIRA MISMATCH";
+						Bugg.JiraFixVersionsText = "JIRA MISMATCH";
+						Bugg.JiraFixCL = "JIRA MISMATCH";
+					}
 				}
 
 				return new ReportsViewModel
@@ -281,7 +302,7 @@ namespace Tools.CrashReporter.CrashReportWebSite.Controllers
 			}
 		}
 
-		private void AddBuggJiraMapping( Bugg NewBugg, ref List<string> FoundJiras, ref Dictionary<string, List<Bugg>> JiraIDtoBugg )
+		private void AddBuggJiraMapping( Bugg NewBugg, ref HashSet<string> FoundJiras, ref Dictionary<string, List<Bugg>> JiraIDtoBugg )
 		{
 			string JiraID = NewBugg.TTPID;
 			FoundJiras.Add( "key = " + JiraID );
