Converted to text from utf-16

diff --git a/Engine/Source/Programs/CrashReporter/CrashReportWebSite/Models/Query-CleanKernelBaseEntries.sql b/Engine/Source/Programs/CrashReporter/CrashReportWebSite/Models/Query-CleanKernelBaseEntries.sql
new file mode 100644
--- /dev/null
+++ b/Engine/Source/Programs/CrashReporter/CrashReportWebSite/Models/Query-CleanKernelBaseEntries.sql
@@ -0,0 +1,33 @@
+USE [CrashReport]
+
+UPDATE Buggs 
+SET [SummaryV2] = '' 
+WHERE [SummaryV2] = 'KERNELBASE.dll'
+
+UPDATE Buggs 
+SET [SummaryV2] = '' 
+WHERE [SummaryV2] = 'UE4Editor-UnrealEd.dll'
+
+UPDATE Crashes 
+SET [Summary] = '' 
+WHERE [Summary] = 'KERNELBASE.dll'
+
+UPDATE Crashes 
+SET [Summary] = '' 
+WHERE [Summary] = 'UE4Editor-UnrealEd.dll'
+
+UPDATE Crashes 
+SET [Summary] = ''
+WHERE [Summary] = 'AssertLog='
+
+UPDATE Crashes 
+SET [Summary] = ''
+WHERE [Summary] = 'Fatal error!' 
+
+UPDATE Crashes 
+SET [Description] = ''
+WHERE [Description] = 'No comment provided' 
+
+
+
+--OUTPUT $action, Inserted.*, Deleted.*
\ No newline at end of file
diff --git a/Engine/Source/Programs/CrashReporter/CrashReportWebSite/Models/Query-UserName.sql b/Engine/Source/Programs/CrashReporter/CrashReportWebSite/Models/Query-UserName.sql
new file mode 100644
--- /dev/null
+++ b/Engine/Source/Programs/CrashReporter/CrashReportWebSite/Models/Query-UserName.sql
@@ -0,0 +1,11 @@
+SELECT [Id]
+      ,[ComputerName]
+      ,[BuildVersion]
+      ,[UserName]
+      ,[Branch]
+--    ,[UserNameId]
+      ,[EpicAccountId]
+FROM [CrashReport].[dbo].[Crashes]
+--where [TimeOfCrash] > CAST(GETDATE() AS DATE)
+where username != 'anonymous' and buildversion = '4.5.0.0' and branch = 'UE4-Releases/4.5'
+group by ComputerName
\ No newline at end of file
