CrashReportWWW - Updated SQL scripts

diff --git a/Engine/Source/Programs/CrashReporter/CrashReportWebSite/Models/CreateBugg-stored-procedure.sql b/Engine/Source/Programs/CrashReporter/CrashReportWebSite/Models/CreateBugg-stored-procedure.sql
--- a/Engine/Source/Programs/CrashReporter/CrashReportWebSite/Models/CreateBugg-stored-procedure.sql
+++ b/Engine/Source/Programs/CrashReporter/CrashReportWebSite/Models/CreateBugg-stored-procedure.sql
@@ -1,186 +1,258 @@
-USE [CrashReport]
+﻿USE [CrashReport]
+GO
 
-GO/****** Object:  StoredProcedure [dbo].[UpdateCrashesByPattern]    Script Date: 2014-11-06 16:40:01 ******/
+/****** Object:  StoredProcedure [dbo].[UpdateCrashesByPattern]    Script Date: 2015-03-20 15:19:10 ******/
+SET ANSI_NULLS ON
+GO
 
-SET ANSI_NULLS ONGO
+SET QUOTED_IDENTIFIER ON
+GO
 
-SET QUOTED_IDENTIFIER ONGO
+CREATE PROCEDURE [dbo].[UpdateCrashesByPattern]
+AS
 
-
+BEGIN
+	-- SET NOCOUNT ON added to prevent extra result sets from
 
-ALTER PROCEDURE [dbo].[UpdateCrashesByPattern]AS
+	-- interfering with SELECT statements.
+	SET NOCOUNT ON;
 
-BEGIN	-- SET NOCOUNT ON added to prevent extra result sets from
+	
+	-- Update crash user names
 
-	-- interfering with SELECT statements.	SET NOCOUNT ON;
+	UPDATE [dbo].[Crashes]
+	SET [UserName] = CASE WHEN c.[UserName] is NULL THEN u.[UserName] ELSE c.[UserName] END
 
-		-- Update crash user names
+	FROM [dbo].[Crashes] c
+	LEFT JOIN [dbo].[Users] u on (u.[id] = c.[UserNameId]);
 
-	UPDATE [dbo].[Crashes]	SET [UserName] = CASE WHEN c.[UserName] is NULL THEN u.[UserName] ELSE c.[UserName] END
+	
+	--Create Buggs
 
-	FROM [dbo].[Crashes] c	LEFT JOIN [dbo].[Users] u on (u.[id] = c.[UserNameId]);
+	MERGE Buggs AS B
+	USING 
 
-		--Create Buggs
+	(
+		SELECT 
+		
+			   [TTPID]
+			  ,[Pattern]
+			  ,[NumberOfCrashes]
+			  ,[NumberOfUsers]
+			  
+			  ,[TimeOfFirstCrash]
+			  ,[TimeOfLastCrash]
+			  
+			  ,[BuildVersion]
+			  ,[CrashType]
 
-	MERGE Buggs AS B	USING 
+			  ,[Status]
+			  ,[FixedChangeList]
+			  
+		FROM
 
-	(		SELECT [TTPID]
+		(
+			SELECT COUNT(1) as NumberOfCrashes
 
-			  ,[Pattern]			  ,[NumberOfCrashes]
+				  , max([TimeOfCrash]) as TimeOfLastCrash
+				  , min([TimeOfCrash]) as TimeOfFirstCrash
 
-			  ,[NumberOfUsers]			  ,[TimeOfFirstCrash]
+				  , max([BuildVersion]) as BuildVersion
 
-			  ,[TimeOfLastCrash]			  ,[BuildVersion]
+				  , max([CrashType]) as CrashType
+				  , count(distinct [UserName]) as NumberOfUsers
 
-			  ,[Summary]			  ,[CrashType]
+				  , Max([Status]) as [Status]
+				  , Max(TTPID) as TTPID
 
-			  ,[Status]			  ,[FixedChangeList]
+				  , Max(FixedChangeList) as FixedChangeList
+				  , Pattern
 
-			  ,[Game]		FROM
+			  FROM [dbo].[Crashes] c 
+			  WHERE
 
-		(			SELECT COUNT(1) as NumberOfCrashes
+				Pattern is not NULL AND Pattern not like ''
+			  group by Pattern
 
-				  , max([TimeOfCrash]) as TimeOfLastCrash				  , min([TimeOfCrash]) as TimeOfFirstCrash
+		  ) as s
+		 Where s.NumberOfCrashes > 1
 
-				  , max([BuildVersion]) as BuildVersion				  , max([Summary]) as Summary
+	) AS C 
+	ON (B.Pattern = C.Pattern)
 
-				  , max([CrashType]) as CrashType				  , count(distinct [UserName]) as NumberOfUsers
+	WHEN NOT MATCHED BY TARGET
+		THEN INSERT (
 
-				  , Count(distinct GameName) as GameNameCount				  , MIN(GameName) as Game
+				[TTPID]
+			  ,[Pattern]
+			  ,[NumberOfCrashes]
+			  ,[NumberOfUsers]
 
-				  , Max([Status]) as [Status]				  , Max(TTPID) as TTPID
+			  ,[TimeOfFirstCrash]
+			  ,[TimeOfLastCrash]
 
-				  , Max(FixedChangeList) as FixedChangeList				  , Pattern
+			  ,[BuildVersion]
+			  ,[CrashType]
+			  
+			  ,[Status]
+		  	  ,[FixedChangeList]
 
-			  FROM [dbo].[Crashes] c 			  WHERE
+			  ) 
+		Values(
 
-				Pattern is not NULL AND Pattern not like ''			  group by Pattern
+				C.[TTPID]
+			  , C.[Pattern]
+			  , C.[NumberOfCrashes]
+			  , C.[NumberOfUsers]
 
-		  ) as s		 Where s.NumberOfCrashes > 1
+			  , C.[TimeOfFirstCrash]
+			  , C.[TimeOfLastCrash]
 
-	) AS C 	ON (B.Pattern = C.Pattern)
+			  , C.[BuildVersion]
+			  , C.[CrashType]
+			  
+			  , C.[Status]
+			  , C.[FixedChangeList]
 
-	WHEN NOT MATCHED BY TARGET		THEN INSERT (
+			  ) 
 
-				[TTPID]			  ,[Pattern]
 
-			  ,[NumberOfCrashes]			  ,[NumberOfUsers]
+	WHEN MATCHED 
+		THEN UPDATE SET B.[NumberOfCrashes] = C.[NumberOfCrashes]
 
-			  ,[TimeOfFirstCrash]			  ,[TimeOfLastCrash]
+						, B.[TimeOfLastCrash] = C.[TimeOfLastCrash]
+						, B.[BuildVersion] = C.[BuildVersion]
 
-			  ,[BuildVersion]			  ,[SummaryV2] -- This a new summary
+						, B.[CrashType] = C.[CrashType]
 
-			  ,[CrashType]			  ,[Status]
+						, B.[NumberOfUsers] = C.[NumberOfUsers]
 
-			  ,[FixedChangeList]			  ,[Game]
 
-			  ) 		Values(
+	OUTPUT $action, Inserted.*, Deleted.*;
+		
 
-				C.[TTPID]			  , C.[Pattern]
+		
+		
 
-			  , C.[NumberOfCrashes]			  , C.[NumberOfUsers]
+	/****** Join Buggs and Crashes  ******/
+	MERGE dbo.Buggs_Crashes BC 
 
-			  , C.[TimeOfFirstCrash]			  , C.[TimeOfLastCrash]
+	USING 
+	( 
 
-			  , C.[BuildVersion]			  , C.[Summary]
+		SELECT  b.Id as BuggId, c.[Id] as CrashId
+		FROM [dbo].[Crashes] c
 
-			  , C.[CrashType]			  , C.[Status]
+		Join [dbo].Buggs b on (b.Pattern = c.Pattern)
+		group by b.Id, c.Id
 
-			  , C.[FixedChangeList]			  , C.[Game]
+	 ) AS C
+	 ON BC.BuggId = C.BuggId AND BC.CrashId = C.CrashId
 
-			  ) 
+	 WHEN NOT MATCHED BY TARGET
+		THEN INSERT 
 
-	WHEN MATCHED 		THEN UPDATE SET B.[NumberOfCrashes] = C.[NumberOfCrashes]
+			(BuggId, CrashId)
+		VALUES 
 
-						, B.[TimeOfLastCrash] = C.[TimeOfLastCrash]						, B.[BuildVersion] = C.[BuildVersion]
+			(C.BuggId, C.CrashId)	
+	 WHEN MATCHED
 
-						, B.[SummaryV2] = C.[Summary] -- This a new summary						, B.[CrashType] = C.[CrashType]
+		THEN UPDATE SET
+			BC.BuggId = C.BuggId, 
 
-						, B.[NumberOfUsers] = C.[NumberOfUsers]
+			BC.CrashId = C.CrashId
+	OUTPUT $action, Inserted.*, Deleted.*;
 
-	OUTPUT $action, Inserted.*, Deleted.*;		
+	
+	
 
-				
+	/****** Join Buggs_Users and Crashes  ******/
+	MERGE dbo.Buggs_Users AS BU
 
-	/****** Join Buggs and Crashes  ******/	MERGE dbo.Buggs_Crashes BC 
+	USING
+		(
 
-	USING 	( 
+		  SELECT  b.Id as BuggId, c.[UserName] as UserName
+		  FROM [dbo].[Crashes] c
 
-		SELECT  b.Id as BuggId, c.[Id] as CrashId		FROM [dbo].[Crashes] c
+		  Join [dbo].Buggs b on (b.Pattern = c.Pattern)
+		  Group by b.Id, c.UserName
 
-		Join [dbo].Buggs b on (b.Pattern = c.Pattern)		group by b.Id, c.Id
+		) AS C
+	ON BU.BuggId = C.BuggId AND BU.UserName = C.UserName
 
-	 ) AS C	 ON BC.BuggId = C.BuggId AND BC.CrashId = C.CrashId
+	 WHEN NOT MATCHED BY TARGET
+		THEN INSERT 
 
-	 WHEN NOT MATCHED BY TARGET		THEN INSERT 
+			(BuggId, UserName) 
+		VALUES 
 
-			(BuggId, CrashId)		VALUES 
+			(C.BuggId, C.UserName)	
+	 WHEN MATCHED
 
-			(C.BuggId, C.CrashId)		 WHEN MATCHED
+		THEN UPDATE SET
+			BU.BuggId = C.BuggId, 
 
-		THEN UPDATE SET			BC.BuggId = C.BuggId, 
+			BU.UserName = C.UserName
+	OUTPUT $action, Inserted.*, Deleted.*;
 
-			BC.CrashId = C.CrashId	OUTPUT $action, Inserted.*, Deleted.*;
+	
+		  
 
-		
+		
+	  
 
-	/****** Join Buggs_Users and Crashes  ******/	MERGE dbo.Buggs_Users AS BU
+	/****** Join Buggs and UserGroups  ******/
+	MERGE dbo.Buggs_UserGroups AS BUG
 
-	USING		(
+	USING
+	(
 
-		  SELECT  b.Id as BuggId, c.[UserName] as UserName		  FROM [dbo].[Crashes] c
+	  Select BuggId, UserGroupId
+	  From [dbo].[Buggs_Users] bu
 
-		  Join [dbo].Buggs b on (b.Pattern = c.Pattern)		  Group by b.Id, c.UserName
+	  Join dbo.Users u on (u.UserName = bu.UserName)
+	  Group by BuggId, UserGroupId
 
-		) AS C	ON BU.BuggId = C.BuggId AND BU.UserName = C.UserName
+	 ) AS BU
+	 ON  BUG.BuggId = BU.BuggId AND BUG.UserGroupId = BU.UserGroupId
 
-	 WHEN NOT MATCHED BY TARGET		THEN INSERT 
+	 WHEN NOT MATCHED BY TARGET
+		THEN INSERT 
 
-			(BuggId, UserName) 		VALUES 
+			(BuggId, UserGroupId) 
+		VALUES 
 
-			(C.BuggId, C.UserName)		 WHEN MATCHED
+			(BU.BuggId, BU.UserGroupId)	
+	 WHEN MATCHED
 
-		THEN UPDATE SET			BU.BuggId = C.BuggId, 
+		THEN UPDATE SET
+			BUG.BuggId = BU.BuggId, 
 
-			BU.UserName = C.UserName	OUTPUT $action, Inserted.*, Deleted.*;
+			BUG.UserGroupId = BU.UserGroupId
+	OUTPUT $action, Inserted.*, Deleted.*;
 
-			  
+	 
+	/*** Update Crashes to match the Bugg Status, TTPID, and Fixed Change List ****/
 
-			  
+	-- This handles new crashes that enter the system; it has the side effect of making the Bugg authoritative in this case but that's how they should be used.
+	update C
 
-	/****** Join Buggs and UserGroups  ******/	MERGE dbo.Buggs_UserGroups AS BUG
+	SET 
+		C.TTPID = B.TTPID,
 
-	USING	(
+		C.FixedChangeList = B.FixedChangeList,
+		C.[Status] = B.[Status]
 
-	  Select BuggId, UserGroupId	  From [dbo].[Buggs_Users] bu
+	From 
+		dbo.Crashes as C
 
-	  Join dbo.Users u on (u.UserName = bu.UserName)	  Group by BuggId, UserGroupId
+		Join Buggs_Crashes as BC ON (c.Id = BC.CrashId)
+		Join Buggs as B ON (B.Id =BC.BuggId AND C.Id = BC.CrashId) 
 
-	 ) AS BU	 ON  BUG.BuggId = BU.BuggId AND BUG.UserGroupId = BU.UserGroupId
+	 
+END
+GO
 
-	 WHEN NOT MATCHED BY TARGET		THEN INSERT 
-
-			(BuggId, UserGroupId) 		VALUES 
-
-			(BU.BuggId, BU.UserGroupId)		 WHEN MATCHED
-
-		THEN UPDATE SET			BUG.BuggId = BU.BuggId, 
-
-			BUG.UserGroupId = BU.UserGroupId	OUTPUT $action, Inserted.*, Deleted.*;
-
-	 	/*** Update Crashes to match the Bugg Status, TTPID, and Fixed Change List ****/
-
-	-- This handles new crashes that enter the system; it has the side effect of making the Bugg authoritative in this case but that's how they should be used.	update C
-
-	SET 		C.TTPID = B.TTPID,
-
-		C.FixedChangeList = B.FixedChangeList,		C.[Status] = B.[Status]
-
-	From 		dbo.Crashes as C
-
-		Join Buggs_Crashes as BC ON (c.Id = BC.CrashId)		Join Buggs as B ON (B.Id =BC.BuggId AND C.Id = BC.CrashId) 
-
-	 END
-
-
 
diff --git a/Engine/Source/Programs/CrashReporter/CrashReportWebSite/Models/CreateDatabase.sql b/Engine/Source/Programs/CrashReporter/CrashReportWebSite/Models/CreateDatabase.sql
--- a/Engine/Source/Programs/CrashReporter/CrashReportWebSite/Models/CreateDatabase.sql
+++ b/Engine/Source/Programs/CrashReporter/CrashReportWebSite/Models/CreateDatabase.sql
@@ -171,29 +171,21 @@ GO
 SET ANSI_PADDING ON
 GO
 CREATE TABLE [dbo].[Buggs](
-		[Id] [int] IDENTITY(1,1) NOT NULL,
+	[Id] [int] IDENTITY(1,1) NOT NULL,
 	[TTPID] [varchar](50) NULL,
-	[Title] [varchar](120) NULL,
-	[Summary] [varchar](250) NULL,
-	[Priority] [int] NULL,
 	[Pattern] [varchar](800) NOT NULL,
-	[Type] [varchar](50) NULL,
 	[NumberOfCrashes] [int] NULL,
 	[NumberOfUsers] [int] NULL,
 	[TimeOfFirstCrash] [datetime] NULL,
 	[TimeOfLastCrash] [datetime] NULL,
 	[Status] [varchar](64) NULL,
 	[FixedChangeList] [varchar](50) NULL,
-	[Description] [text] NULL,
-	[ReproSteps] [text] NULL,
-	[Game] [varchar](120) NULL,
-	[BuildVersion] [varchar](128) NULL,
 	[CrashType] [smallint] NULL,
-	[SummaryV2] [varchar](512) NULL,
+	[BuildVersion] [varchar](64) NULL,
  CONSTRAINT [PK_Buggs] PRIMARY KEY CLUSTERED 
 (
 	[Id] ASC
-)WITH (PAD_INDEX  = OFF, STATISTICS_NORECOMPUTE  = OFF, IGNORE_DUP_KEY  = OFF, ALLOW_ROW_LOCKS  = ON, ALLOW_PAGE_LOCKS  = ON) ON [PRIMARY]
+)WITH (PAD_INDEX = OFF, STATISTICS_NORECOMPUTE = OFF, IGNORE_DUP_KEY = OFF, ALLOW_ROW_LOCKS = ON, ALLOW_PAGE_LOCKS = ON) ON [PRIMARY]
 ) ON [PRIMARY] TEXTIMAGE_ON [PRIMARY]
 GO
 SET ANSI_PADDING OFF
