UE-12159 - Crash Report Client fails to restart project if the path has a space

diff --git a/Engine/Source/Programs/CrashReporter/CrashReportClient/Private/CrashDescription.cpp b/Engine/Source/Programs/CrashReporter/CrashReportClient/Private/CrashDescription.cpp
--- a/Engine/Source/Programs/CrashReporter/CrashReportClient/Private/CrashDescription.cpp
+++ b/Engine/Source/Programs/CrashReporter/CrashReportClient/Private/CrashDescription.cpp
@@ -22,6 +22,16 @@ FCrashDescription::FCrashDescription() :
 	InitializeIDs();
 }
 
+/** Unescapes a specified XML string, naive implementation. */
+FString UnescapeXMLString( const FString& Text )
+{
+	
+	return Text
+		.Replace( TEXT( "&amp;" ), TEXT( "&" ) )
+		.Replace( TEXT( "&quot;" ), TEXT( "\"" ) )
+		.Replace( TEXT( "&lt;" ), TEXT( "<" ) )
+		.Replace( TEXT( "&gt;" ), TEXT( ">" ) );
+}
 
 FCrashDescription::FCrashDescription( FString WERXMLFilepath ) :
 	// WER XML files are forced to be in the first version
@@ -88,7 +98,7 @@ FCrashDescription::FCrashDescription( FString WERXMLFilepath ) :
 
 				if( ParsedParameters8.Num() > 1 )
 				{
-					CommandLine = ParsedParameters8[1];
+					CommandLine = UnescapeXMLString( ParsedParameters8[1] );
 				}
 
 				if( ParsedParameters8.Num() > 2 )
diff --git a/Engine/Source/Programs/CrashReporter/CrashReportClient/Private/Windows/WindowsErrorReport.cpp b/Engine/Source/Programs/CrashReporter/CrashReportClient/Private/Windows/WindowsErrorReport.cpp
--- a/Engine/Source/Programs/CrashReporter/CrashReportClient/Private/Windows/WindowsErrorReport.cpp
+++ b/Engine/Source/Programs/CrashReporter/CrashReportClient/Private/Windows/WindowsErrorReport.cpp
@@ -43,14 +43,8 @@ struct FWindowsReportParser
 				TArray<FString> SeparatedParameters;
 				StringLine.ParseIntoArray( SeparatedParameters, Marker, true );
 
-				FString MatchedValue;
-				const bool bFound = FParse::Value( *StringLine, Marker, MatchedValue );
-
-				if( bFound )
-				{
-					Result = MatchedValue;
-					break;
-				}
+				Result = SeparatedParameters[SeparatedParameters.Num()-1];
+				break;
 			}
 		}
 
