CrashReportWWW - UECORE-108 / Crash Reporter: The Reports page is no longer sorting by # of Occurrences

diff --git a/Engine/Source/Programs/CrashReporter/CrashReportWebSite/Controllers/ReportsController.cs b/Engine/Source/Programs/CrashReporter/CrashReportWebSite/Controllers/ReportsController.cs
--- a/Engine/Source/Programs/CrashReporter/CrashReportWebSite/Controllers/ReportsController.cs
+++ b/Engine/Source/Programs/CrashReporter/CrashReportWebSite/Controllers/ReportsController.cs
@@ -289,6 +289,8 @@ namespace Tools.CrashReporter.CrashReportWebSite.Controllers
 					}
 				}
 
+				Buggs = Buggs.OrderByDescending( b => b.NumberOfCrashes ).ToList();
+
 				return new ReportsViewModel
 				{
 					Buggs = Buggs,
diff --git a/Engine/Source/Programs/CrashReporter/CrashReportWebSite/Models/CrashReport.cs b/Engine/Source/Programs/CrashReporter/CrashReportWebSite/Models/CrashReport.cs
--- a/Engine/Source/Programs/CrashReporter/CrashReportWebSite/Models/CrashReport.cs
+++ b/Engine/Source/Programs/CrashReporter/CrashReportWebSite/Models/CrashReport.cs
@@ -138,7 +138,7 @@ namespace Tools.CrashReporter.CrashReportWebSite.Models
 		/// <summary> Helper method, display this Bugg as a human readable string. Debugging purpose. </summary>
 		public override string ToString()
 		{
-			return string.Format( "Id={0} NOC={1} TLC={2} BV={3} Pattern={4}", Id, NumberOfCrashes, TimeOfLastCrash, BuildVersion, Pattern != null ? Pattern.Length : 0 );
+			return string.Format( "Id={0} NOC={1} TLC={2} BV={3} NOM={4} Pattern={5}", Id, NumberOfCrashes, TimeOfLastCrash, BuildVersion, NumberOfUniqueMachines, Pattern != null ? Pattern.Length : 0 );
 		}
 
 		/// <summary>
