CrashReportWWW - Removed unused code

diff --git a/Engine/Source/Programs/CrashReporter/CrashReportWebSite/Models/CrashReport.cs b/Engine/Source/Programs/CrashReporter/CrashReportWebSite/Models/CrashReport.cs
--- a/Engine/Source/Programs/CrashReporter/CrashReportWebSite/Models/CrashReport.cs
+++ b/Engine/Source/Programs/CrashReporter/CrashReportWebSite/Models/CrashReport.cs
@@ -254,14 +254,6 @@ namespace Tools.CrashReporter.CrashReportWebSite.Models
 				this.ToJiraSummary = "No valid callstack found"; 
 			}
 
-// 			Crash NewCrash = CrashesForBugg.FirstOrDefault();
-// 			if( NewCrash != null )
-// 			{
-// 				CallStackContainer CallStack = NewCrash.GetCallStack();
-// 				this.ToJiraSummary = CallStack.CallStackEntries.Count > 1 ? CallStack.CallStackEntries[0].FunctionName : "No valid callstack found";
-// 				this.ToJiraFunctionCalls = CallStack.GetFunctionCalls();
-// 			}
-
 			// ToJiraVersions
 			this.ToJiraVersions = new List<object>();
 			foreach( var Version in this.AffectedMajorVersions )
