CrashReportLinq - Updated RepairDatabase program

diff --git a/Engine/Source/Programs/CrashReporter/CrashReportWebSite/Queries/Fix-RepairDatabase.linq b/Engine/Source/Programs/CrashReporter/CrashReportWebSite/Queries/Fix-RepairDatabase.linq
--- a/Engine/Source/Programs/CrashReporter/CrashReportWebSite/Queries/Fix-RepairDatabase.linq
+++ b/Engine/Source/Programs/CrashReporter/CrashReportWebSite/Queries/Fix-RepairDatabase.linq
@@ -33,7 +33,7 @@ HashSet<int> GetInvalidFuncIds()
 }
 
 static Stopwatch Timer = Stopwatch.StartNew();
-static DateTime Date = DateTime.UtcNow.AddMonths(-2);
+static DateTime Date = new DateTime(2010,05,26);//DateTime.UtcNow.AddMonths(-2000);
 static TimeSpan Tick = TimeSpan.FromDays(3);
 const int NUM_OPS_PER_BATCH = 32;
 
@@ -96,66 +96,79 @@ void DeleteBuggs( List<int> BuggsBatch )
 // Deletes all crashes without pattern.
 void DeleteNoPatternAll()
 {
-	var NoPattern = Crashes
+	var NoPatternCount = Crashes
 	.Where( c => c.Pattern == null )
-	.Where (c => c.TimeOfCrash> Date)
-	.ToList();
+	.Where( c => c.TimeOfCrash> Date )
+	.Count();
+	WriteLine( string.Format( "NoPatternCount: {0}", NoPatternCount ) );
 	
-	var CrashesIds = new HashSet<int>( NoPattern.Select (np => np.Id).ToList() );
-	
-	int NumBuggCrass = Buggs_Crashes.Count (bc => bc.BuggId>0);
+	var NoPatternBatch = new List<Crashes>(NUM_OPS_PER_BATCH);
 	
-	// Unlink from the bugg.	
-	List<Buggs_Crashes> ToRemove = new List<Buggs_Crashes>();
-	HashSet<int> BuggIds = new HashSet<int>();
-	for (int n = 0; n < NoPattern.Count; n ++)
+	DateTime StartDate = Date;
+	DateTime EndDate = Date.Add( Tick );
+	while( EndDate < DateTime.UtcNow )
 	{
-		var Crash = NoPattern[n];
+		var NoPatternCrashes = Crashes
+		.Where (c => c.TimeOfCrash > StartDate && c.TimeOfCrash <= EndDate)
+		.Where( c => c.Pattern == null )
+		.ToList();
 	
-		var BC = Buggs_Crashes.Where (bc => bc.CrashId==Crash.Id).FirstOrDefault();
-		if(BC!=null)
-		{
-			ToRemove.Add(BC);
-			BuggIds.Add(BC.BuggId);
-		}
+		WriteLine( string.Format( "NoPatternCrashes: {0} {1} -> {2}", NoPatternCrashes.Count, StartDate.ToString("yyyy-MM-dd"), EndDate.ToString("yyyy-MM-dd") ) );
+		StartDate = EndDate;
+	 	EndDate = EndDate.Add( Tick );
 		
-		if (ToRemove.Count == NUM_OPS_PER_BATCH)
-		{
-			DeleteBuggCrash(ToRemove);
-		}
-	}
-	
-	DeleteBuggCrash(ToRemove);
-	
-	List<int> BuggsToRemove = new List<int>();
-	foreach (var BuggId in BuggIds)
-	{
-		var CrashesForBugg = Buggs_Crashes.Where (bc => bc.BuggId==BuggId).Select (bc => bc.CrashId);
-		int CrashCount = CrashesForBugg.Count (cfb => cfb>0);
-		if(CrashCount==0)
+		var CrashesIds = new HashSet<int>( NoPatternCrashes.Select (np => np.Id).ToList() );
+		int NumBuggCrass = Buggs_Crashes.Count (bc => bc.BuggId>0);
+		
+		// Unlink from the bugg.	
+		List<Buggs_Crashes> ToRemove = new List<Buggs_Crashes>();
+		HashSet<int> BuggIds = new HashSet<int>();
+		for (int n = 0; n < NoPatternCrashes.Count; n ++)
 		{
-			// Remove the bugg.
-			BuggsToRemove.Add(BuggId);
+			var Crash = NoPatternCrashes[n];
+		
+			var BC = Buggs_Crashes.Where (bc => bc.CrashId==Crash.Id).FirstOrDefault();
+			if(BC!=null)
+			{
+				ToRemove.Add(BC);
+				BuggIds.Add(BC.BuggId);
+			}
+			
+			if (ToRemove.Count == NUM_OPS_PER_BATCH)
+			{
+				DeleteBuggCrash(ToRemove);
+			}
 		}
 		
-		if (BuggsToRemove.Count > NUM_OPS_PER_BATCH)
+		DeleteBuggCrash(ToRemove);
+		
+		List<int> BuggsToRemove = new List<int>();
+		foreach (var BuggId in BuggIds)
 		{
-			DeleteBuggs( BuggsToRemove );
+			var CrashesForBugg = Buggs_Crashes.Where (bc => bc.BuggId==BuggId).Select (bc => bc.CrashId);
+			int CrashCount = CrashesForBugg.Count (cfb => cfb>0);
+			if(CrashCount==0)
+			{
+				// Remove the bugg.
+				BuggsToRemove.Add(BuggId);
+			}
+			
+			if (BuggsToRemove.Count == NUM_OPS_PER_BATCH)
+			{
+				DeleteBuggs( BuggsToRemove );
+			}
 		}
-	}
-	
-	DeleteBuggs( BuggsToRemove );
-	
-	WriteLine( string.Format( "NoPattern: {0}", NoPattern.Count ) );
-	
-	var NoPatternBatch = new List<Crashes>(NUM_OPS_PER_BATCH);
-	for (int n = 0; n < NoPattern.Count; n ++)
-	{
-		NoPatternBatch.Add(NoPattern[n]);
 		
-		if (NoPatternBatch.Count == NUM_OPS_PER_BATCH)
+		DeleteBuggs( BuggsToRemove );
+	
+		for (int n = 0; n < NoPatternCrashes.Count; n ++)
 		{
-			DeleteNoPatternBatch(NoPatternBatch);
+			NoPatternBatch.Add(NoPatternCrashes[n]);
+			
+			if (NoPatternBatch.Count == NUM_OPS_PER_BATCH)
+			{
+				DeleteNoPatternBatch(NoPatternBatch);
+			}
 		}
 	}
 	
@@ -166,23 +179,37 @@ void DeleteNoPatternAll()
 // Sets patterns for all crashes without branch, will be deleted after running the script again.
 void SetNoPatternForNoBranchAll()
 {
-	var NoBranch = Crashes
+	var NoBranchCount = Crashes
 	.Where( c => c.Branch == null )
 	.Where (c => c.TimeOfCrash> Date)
-	.ToList();
+	.Count();
 	
-	WriteLine( string.Format( "NoBranch: {0}", NoBranch.Count ) );
+	WriteLine( string.Format( "NoBranch: {0}", NoBranchCount ) );
 	
 	var NoBranchBatch = new List<Crashes>(NUM_OPS_PER_BATCH);
-	for (int n = 0; n < NoBranch.Count; n ++)
+	DateTime StartDate = Date;
+	DateTime EndDate = Date.Add( Tick );
+	while( EndDate < DateTime.UtcNow )
 	{
-		NoBranchBatch.Add(NoBranch[n]);		
-		NoBranch[n].Pattern = null;
-		NoBranch[n].Branch = "INVALID";
-		
-		if (n % NUM_OPS_PER_BATCH == 0)
+		var NoBranchCrashes = Crashes
+		.Where (c => c.TimeOfCrash > StartDate && c.TimeOfCrash <= EndDate)
+		.Where( c => c.Branch == null )
+		.ToList();
+	
+		WriteLine( string.Format( "NoBranchCrashes: {0} {1} -> {2}", NoBranchCrashes.Count, StartDate.ToString("yyyy-MM-dd"), EndDate.ToString("yyyy-MM-dd") ) );
+		StartDate = EndDate;
+	 	EndDate = EndDate.Add( Tick );
+	
+		for (int n = 0; n < NoBranchCrashes.Count; n ++)
 		{
-			SumbitChanges(NoBranchBatch);
+			NoBranchBatch.Add(NoBranchCrashes[n]);		
+			NoBranchCrashes[n].Pattern = null;
+			NoBranchCrashes[n].Branch = "INVALID";
+			
+			if (NoBranchBatch.Count == NUM_OPS_PER_BATCH)
+			{
+				SumbitChanges(NoBranchBatch);
+			}
 		}
 	}
 	
@@ -195,6 +222,7 @@ void SumbitChanges( List<Crashes> CrashesBatch )
 	//Crashes.InsertAllOnSubmit( CrashesBatch );
 	Crashes.Context.SubmitChanges();
 	WriteLine( string.Format( "SubmitBatchChanges: {0}", CrashesBatch.Count ) );
+	CrashesBatch.Clear();
 }
 
 class CustomFuncComparer : IEqualityComparer<string>
@@ -238,18 +266,13 @@ void SetNoPatternForInvalidPattern()
 	
 	DateTime StartDate = Date;
 	DateTime EndDate = Date.Add( Tick );
-	while( true )
+	while( EndDate < DateTime.UtcNow )
 	{
 		var AllCrashes = Crashes
 		.Where (c => c.TimeOfCrash > StartDate && c.TimeOfCrash <= EndDate)
 		.Where (c => c.Pattern != null)
 		.ToList();
 		
-		if( AllCrashes.Count == 0 )
-		{
-			break;
-		}
-		
 		WriteLine( string.Format( "AllCrashes: {0} {1} -> {2}", AllCrashes.Count, StartDate.ToString("yyyy-MM-dd"), EndDate.ToString("yyyy-MM-dd") ) );
 		StartDate = EndDate;
 	 	EndDate = EndDate.Add( Tick );
@@ -294,19 +317,19 @@ void SetNoPatternForInvalidPattern()
 		{
 			WriteLine( string.Format( "InvalidCrashes: {0}", InvalidCrashes.Count ) );
 			SumbitChanges(InvalidCrashes);
-			InvalidCrashes.Clear();
 		}
 	}
 }
 
 void Main()
 {	
-	int NumCrashes = Crashes.Count (c => c.Id>0);
-	int MaxCrashId = Crashes.OrderByDescending (c => c.TimeOfCrash).Take(1).FirstOrDefault().Id;
-	int Diff = MaxCrashId-NumCrashes;
-	WriteLine( string.Format( "NumCrashes: {0}", NumCrashes ) );
-	WriteLine( string.Format( "MaxCrashId: {0}", MaxCrashId ) );
-	WriteLine( string.Format( "Diff:       {0}", Diff ) );
+	int NumCrashes = Crashes.Count();
+	int MaxCrashId = Crashes.Max (c => c.Id);
+	int NumBuggs = Buggs.Count();
+	int MaxBuggId = Buggs.Max(c => c.Id);
+	WriteLine( string.Format( "NumCrashes: {0} MaxCrashId: {1} Diff: {2}", NumCrashes, MaxCrashId, MaxCrashId-NumCrashes ) );
+	WriteLine( string.Format( "NumBuggs: {0} MaxBuggId: {1} Diff: {2}", NumBuggs, MaxBuggId, MaxBuggId-NumBuggs ) );
+	
 	SetNoPatternForNoBranchAll();
 	SetNoPatternForInvalidPattern();
 	DeleteNoPatternAll();
