CrashReportWWW - UECORE-101/Crash Reporter: Include Callstack line numbers in CopyToJira

diff --git a/Engine/Source/Programs/CrashReporter/CrashReportWebSite/Controllers/ReportsController.cs b/Engine/Source/Programs/CrashReporter/CrashReportWebSite/Controllers/ReportsController.cs
--- a/Engine/Source/Programs/CrashReporter/CrashReportWebSite/Controllers/ReportsController.cs
+++ b/Engine/Source/Programs/CrashReporter/CrashReportWebSite/Controllers/ReportsController.cs
@@ -73,6 +73,7 @@ namespace Tools.CrashReporter.CrashReportWebSite.Controllers
 						MachineID = Crash.ComputerName,
 						Branch = Crash.Branch,
 						Description = Crash.Description,
+						RawCallStack = Crash.RawCallStack,
 					} )
 					.ToList();
 				int NumCrashes = Crashes.Count;
@@ -167,6 +168,7 @@ namespace Tools.CrashReporter.CrashReportWebSite.Controllers
 									MachineID = Anon.MachineID,
 									Branch = Anon.Branch,
 									Description = Anon.Description,
+									RawCallStack = Anon.RawCallStack,
 								} );
 							}
 
diff --git a/Engine/Source/Programs/CrashReporter/CrashReportWebSite/Models/CrashDataFormater.cs b/Engine/Source/Programs/CrashReporter/CrashReportWebSite/Models/CrashDataFormater.cs
--- a/Engine/Source/Programs/CrashReporter/CrashReportWebSite/Models/CrashDataFormater.cs
+++ b/Engine/Source/Programs/CrashReporter/CrashReportWebSite/Models/CrashDataFormater.cs
@@ -458,5 +458,23 @@ namespace Tools.CrashReporter.CrashReportWebSite.Models
 			}
 			return Result;
 		}
+
+		/// <returns>A list of functions calls.</returns>
+		public List<string> GetFunctionCallsForJira()
+		{
+			var Result = new List<string>();
+			foreach( CallStackEntry CurrentEntry in CallStackEntries )
+			{
+				if( CurrentEntry.FileName != null && CurrentEntry.FileName.Length > 3 )
+				{
+					Result.Add( string.Format( "{0}!{1} [{2}]", CurrentEntry.ModuleName, CurrentEntry.FunctionName, CurrentEntry.FileName ) );
+				}
+				else
+				{
+					Result.Add( string.Format( "{0}!{1}", CurrentEntry.ModuleName, CurrentEntry.FunctionName ) );
+				}
+			}
+			return Result;
+		}
 	}
 }
diff --git a/Engine/Source/Programs/CrashReporter/CrashReportWebSite/Models/CrashReport.cs b/Engine/Source/Programs/CrashReporter/CrashReportWebSite/Models/CrashReport.cs
--- a/Engine/Source/Programs/CrashReporter/CrashReportWebSite/Models/CrashReport.cs
+++ b/Engine/Source/Programs/CrashReporter/CrashReportWebSite/Models/CrashReport.cs
@@ -134,6 +134,7 @@ namespace Tools.CrashReporter.CrashReportWebSite.Models
 		List<string> ToJiraDescriptions = new List<string>();
 
 		List<string> ToJiraFunctionCalls = null;
+		List<string> GenericFunctionCalls = null;
 
 		/// <summary> Helper method, display this Bugg as a human readable string. Debugging purpose. </summary>
 		public override string ToString()
@@ -243,7 +244,7 @@ namespace Tools.CrashReporter.CrashReportWebSite.Models
 			this.LatestOSAffected = LatestOSAffected;			// Latest Environment Affected
 
 			// ToJiraSummary
-			var FunctionCalls = this.GetFunctionCalls();
+			var FunctionCalls = CrashesForBugg.First().GetCallStack().GetFunctionCallsForJira();
 			if( FunctionCalls.Count > 0 )
 			{
 				this.ToJiraSummary = FunctionCalls[0];
@@ -251,7 +252,8 @@ namespace Tools.CrashReporter.CrashReportWebSite.Models
 			}
 			else
 			{
-				this.ToJiraSummary = "No valid callstack found"; 
+				this.ToJiraSummary = "No valid callstack found";
+				this.ToJiraFunctionCalls = new List<string>();
 			}
 
 			// ToJiraVersions
@@ -477,7 +479,7 @@ namespace Tools.CrashReporter.CrashReportWebSite.Models
 		{
 			using( FAutoScopedLogTimer LogTimer = new FAutoScopedLogTimer( this.GetType().ToString() + "(Id=" + this.Id + ")" ) )
 			{
-				if( ToJiraFunctionCalls == null )
+				if( GenericFunctionCalls == null )
 				{
 					CrashRepository CrashRepo = new CrashRepository();
 
@@ -490,14 +492,14 @@ namespace Tools.CrashReporter.CrashReportWebSite.Models
 
 					if( Crash != null )
 					{
-						ToJiraFunctionCalls = Crash.GetCallStack().GetFunctionCalls();
+						GenericFunctionCalls = Crash.GetCallStack().GetFunctionCalls();
 					}
 					else
 					{
-						ToJiraFunctionCalls = new List<string>();
+						GenericFunctionCalls = new List<string>();
 					}
 				}
-				return ToJiraFunctionCalls;
+				return GenericFunctionCalls;
 			}
 		}
 	}
