CrashReporter now identifies itself to analytics as CrashReporter instead of looking like another engine run. #UE-12598
* AppID is CrashReporter.[Release|Dev]
* UserID matches what we use everywhere else.
* After discussing with Analytics, they are not using this data yet, but it was easier to fix up the AppID than remove.

diff --git a/Engine/Source/Programs/CrashReporter/CrashReportClient/Private/CrashReportAnalytics.cpp b/Engine/Source/Programs/CrashReporter/CrashReportClient/Private/CrashReportAnalytics.cpp
--- a/Engine/Source/Programs/CrashReporter/CrashReportClient/Private/CrashReportAnalytics.cpp
+++ b/Engine/Source/Programs/CrashReporter/CrashReportClient/Private/CrashReportAnalytics.cpp
@@ -37,9 +37,6 @@ void FCrashReportAnalytics::Initialize()
 {
 	checkf(!bIsInitialized, TEXT("FCrashReportAnalytics::Initialize called more than once."));
 
-	// Never use analytics when running a commandlet tool
-	const bool bShouldInitAnalytics = !IsRunningCommandlet();
-	if( bShouldInitAnalytics )
 	{
 		// Setup some default engine analytics if there is nothing custom bound
 		FAnalytics::FProviderConfigurationDelegate DefaultEngineAnalyticsConfig;
@@ -59,24 +56,7 @@ void FCrashReportAnalytics::Initialize()
 					FAnalytics::Get().GetBuildType() == FAnalytics::Release) &&
 					!FEngineBuildSettings::IsInternalBuild();	// Internal Epic build
 
-				if( FRocketSupport::IsRocket() )
-				{
-					const TCHAR* DevelopmentAccountAPIKeyET = TEXT( "Rocket.Dev" );
-					const TCHAR* ReleaseAccountAPIKeyET = TEXT( "Rocket.Release" );
-					ConfigMap.Add( TEXT( "APIKeyET" ), bUseReleaseAccount ? ReleaseAccountAPIKeyET : DevelopmentAccountAPIKeyET );
-				}
-				else if( FEngineBuildSettings::IsPerforceBuild() )
-				{
-					const TCHAR* DevelopmentAccountAPIKeyET = TEXT( "Perforce.Dev" );
-					const TCHAR* ReleaseAccountAPIKeyET = TEXT( "Perforce.Release" );
-					ConfigMap.Add( TEXT( "APIKeyET" ), bUseReleaseAccount ? ReleaseAccountAPIKeyET : DevelopmentAccountAPIKeyET );
-				}
-				else
-				{
-					const TCHAR* DevelopmentAccountAPIKeyET = TEXT( "UnrealEngine.Dev" );
-					const TCHAR* ReleaseAccountAPIKeyET = TEXT( "UnrealEngine.Release" );
-					ConfigMap.Add( TEXT( "APIKeyET" ), bUseReleaseAccount ? ReleaseAccountAPIKeyET : DevelopmentAccountAPIKeyET );
-				}
+				ConfigMap.Add( TEXT( "APIKeyET" ), bUseReleaseAccount ? TEXT("CrashReporter.Release") : TEXT("CrashReporter.Dev") );
 			}
 
 			// Check for overrides
@@ -99,9 +79,7 @@ void FCrashReportAnalytics::Initialize()
 			DefaultEngineAnalyticsConfig );
 		if( Analytics.IsValid() )
 		{
-			//const FString MachineId = FPlatformMisc::GetMachineId().ToString( EGuidFormats::Digits );
-			//Analytics->SetUserID( MachineId );
-			Analytics->SetUserID(FPlatformMisc::GetUniqueDeviceId());
+			Analytics->SetUserID(FString::Printf(TEXT("%s|%s|%s"), *FPlatformMisc::GetMachineId().ToString(EGuidFormats::Digits).ToLower(), *FPlatformMisc::GetEpicAccountId(), *FPlatformMisc::GetOperatingSystemId()));
 			Analytics->StartSession();
 		}
 	}
