CrashReportWWW - Clean

diff --git a/Engine/Source/Programs/CrashReporter/CrashReportWebSite/CrashReportWebSite.csproj b/Engine/Source/Programs/CrashReporter/CrashReportWebSite/CrashReportWebSite.csproj
--- a/Engine/Source/Programs/CrashReporter/CrashReportWebSite/CrashReportWebSite.csproj
+++ b/Engine/Source/Programs/CrashReporter/CrashReportWebSite/CrashReportWebSite.csproj
@@ -89,7 +89,6 @@
     <Compile Include="Controllers\HomeController.cs" />
     <Compile Include="Controllers\ReportsController.cs" />
     <Compile Include="Controllers\UsersController.cs" />
-    <Compile Include="DatabaseHelper.cs" />
     <Compile Include="Models\BuggsViewModel.cs" />
     <Compile Include="Models\CachedDataService.cs" />
     <Compile Include="Models\BuggRepository.cs" />
diff --git a/Engine/Source/Programs/CrashReporter/CrashReportWebSite/DatabaseHelper.cs b/Engine/Source/Programs/CrashReporter/CrashReportWebSite/DatabaseHelper.cs
deleted file mode 100644
diff --git a/Engine/Source/Programs/CrashReporter/CrashReportWebSite/Models/CrashRepository.cs b/Engine/Source/Programs/CrashReporter/CrashReportWebSite/Models/CrashRepository.cs
--- a/Engine/Source/Programs/CrashReporter/CrashReportWebSite/Models/CrashRepository.cs
+++ b/Engine/Source/Programs/CrashReporter/CrashReportWebSite/Models/CrashRepository.cs
@@ -13,6 +13,7 @@ using Naspinski.IQueryableSearch;
 
 using Tools.CrashReporter.CrashReportCommon;
 using Tools.CrashReporter.CrashReportWebSite.Controllers;
+using System.Web.Mvc;
 
 namespace Tools.CrashReporter.CrashReportWebSite.Models
 {
@@ -59,6 +60,20 @@ namespace Tools.CrashReporter.CrashReportWebSite.Models
 		{
 			Context.Dispose();
 		}
+
+		/// <summary>
+		/// Retrieves a list of distinct UE4 Branches from the CrashRepository
+		/// </summary>
+		/// <returns></returns>
+		public static List<SelectListItem> GetBranches()
+		{
+			CrashReportDataContext Context = new CrashReportDataContext();
+
+			var lm = Context.Crashes.Where( n => n.Branch.Contains( "UE4" ) ).Select( n => n.Branch ).Distinct().ToList();
+			var list = lm.Select( listitem => new SelectListItem { Selected = false, Text = listitem, Value = listitem } ).ToList();
+			list.Insert( 0, new SelectListItem { Selected = true, Text = "", Value = "" } );
+			return list;
+		}
 		
 		/// <summary>
 		/// Gets the crash from an id.
@@ -117,7 +132,7 @@ namespace Tools.CrashReporter.CrashReportWebSite.Models
 		/// <returns>A container of all known crashes.</returns>
 		public IQueryable<Crash> ListAll()
 		{
-			return Context.Crashes.AsQueryable();
+			return Context.Crashes.AsQueryable().Where( c => c.Branch != null );
 		}
 
 		/// <summary>
@@ -170,7 +185,7 @@ namespace Tools.CrashReporter.CrashReportWebSite.Models
 			}
 		}
 
-        /// <summary>
+		/// <summary>
 		/// Return a view model containing a sorted list of crashes based on the user input.
 		/// </summary>
 		/// <param name="FormData">The user input from the client.</param>
@@ -196,10 +211,10 @@ namespace Tools.CrashReporter.CrashReportWebSite.Models
 				int Skip = ( FormData.Page - 1 ) * FormData.PageSize;
 				int Take = FormData.PageSize;
 
-			    var ResultsAll = ConstructQuery(FormData);
+				var ResultsAll = ConstructQuery(FormData);
 
 				// Filter by data and get as enumerable.
-                Results = FilterByDate(ResultsAll, FormData.DateFrom, FormData.DateTo);
+				Results = FilterByDate(ResultsAll, FormData.DateFrom, FormData.DateTo);
 
 
 
@@ -229,7 +244,8 @@ namespace Tools.CrashReporter.CrashReportWebSite.Models
 							NewResults.Add( Crash );
 						}
 					}
-                    Results = NewResults;					
+
+					Results = NewResults;					
 				}
 
 				// Pass in the results and return them sorted properly
@@ -265,10 +281,10 @@ namespace Tools.CrashReporter.CrashReportWebSite.Models
 					UserGroup = FormData.UserGroup,
 					CrashType = FormData.CrashType,
 					SearchQuery = FormData.SearchQuery,
-                    UsernameQuery = FormData.UsernameQuery,
-                    EpicIdQuery = FormData.EpicIdQuery,
-                    MachineIdQuery = FormData.MachineIdQuery,
-                    JiraQuery = FormData.JiraQuery,
+					UsernameQuery = FormData.UsernameQuery,
+					EpicIdQuery = FormData.EpicIdQuery,
+					MachineIdQuery = FormData.MachineIdQuery,
+					JiraQuery = FormData.JiraQuery,
 					DateFrom = (long)( FormData.DateFrom - CrashesViewModel.Epoch ).TotalMilliseconds,
 					DateTo = (long)( FormData.DateTo - CrashesViewModel.Epoch ).TotalMilliseconds,
 					BranchName = FormData.BranchName,
@@ -339,8 +355,8 @@ namespace Tools.CrashReporter.CrashReportWebSite.Models
 			{
 				IQueryable<Crash> CrashesInTimeFrame = Results
 					.Where( MyCrash => MyCrash.TimeOfCrash >= DateFrom && MyCrash.TimeOfCrash <= DateTo.AddDays( 1 ) );
-                IEnumerable<Crash> CrashesInTimeFrameEnumerable = CrashesInTimeFrame.ToList();
-                return CrashesInTimeFrameEnumerable;
+				IEnumerable<Crash> CrashesInTimeFrameEnumerable = CrashesInTimeFrame.ToList();
+				return CrashesInTimeFrameEnumerable;
 			}
 		}
 
@@ -474,7 +490,7 @@ namespace Tools.CrashReporter.CrashReportWebSite.Models
 			NewCrash.ChangeListVersion = NewCrashInfo.BuiltFromCL.ToString();
 			NewCrash.CommandLine = NewCrashInfo.CommandLine;
 			NewCrash.EngineMode = NewCrashInfo.EngineMode;
-		    NewCrash.ComputerName = NewCrashInfo.MachineGuid;
+			NewCrash.ComputerName = NewCrashInfo.MachineGuid;
 
 			// Valid MachineID and UserName, updated crash from non-UE4 release
 			if(!string.IsNullOrEmpty(NewCrashInfo.UserName))
@@ -588,135 +604,135 @@ namespace Tools.CrashReporter.CrashReportWebSite.Models
 			return NewID;
 		}
 
-        private IQueryable<Crash> ConstructQuery(FormHelper FormData)
-        {
-            var results = ListAll();
-
-            // Grab Results 
-            if (!string.IsNullOrEmpty(FormData.SearchQuery))
-            {
-                if (!string.IsNullOrEmpty(FormData.UsernameQuery))
-                {
-                    //We only use SearchQuery now for CallStack searching - if there's a searchquery value and a Username value, we need to get rid of the 
-                    //Username so that we can create a broader search range
-                    FormData.UsernameQuery = "";
-                }
-                string DecodedQuery = HttpUtility.HtmlDecode(FormData.SearchQuery).ToLower();
-                using (FScopedLogTimer LogTimer2 = new FScopedLogTimer("CrashRepository.GetResults.FindUserFromQuery" + "(Query=" + DecodedQuery + ")"))
-                {
-                    results = results.Where(item => item.RawCallStack.Contains(FormData.SearchQuery));
-                }
-            }
-
-            // Filter by Crash Type
-            if (FormData.CrashType != "All")
-            {
-                switch (FormData.CrashType)
-                {
-                    case "Crashes":
-                        results = results.Where(CrashInstance => CrashInstance.CrashType == 1);
-                        break;
-                    case "Assert":
-                        results = results.Where(CrashInstance => CrashInstance.CrashType == 2);
-                        break;
-                    case "Ensure":
-                        results = results.Where(CrashInstance => CrashInstance.CrashType == 3);
-                        break;
-                    case "CrashesAsserts":
-                        results = results.Where(CrashInstance => CrashInstance.CrashType == 1 || CrashInstance.CrashType == 2);
-                        break;
-                }
-            }
-
-            if (!string.IsNullOrEmpty(FormData.UsernameQuery))
-            {
-                results =
-                    (
-                        from CrashDetail in results
-                        where CrashDetail.UserName.Equals(FormData.UsernameQuery)
-                        select CrashDetail
-                        );
-            }
-
-            // Start Filtering the results
-            if (!string.IsNullOrEmpty(FormData.EpicIdQuery))
-            {
-                var DecodedQuery = HttpUtility.HtmlDecode(FormData.EpicIdQuery).ToLower();
-                results =
-                        (
-                            from CrashDetail in results
-                            where CrashDetail.EpicAccountId.Equals(FormData.EpicIdQuery)
-                            select CrashDetail
-                        );
-            }
-
-            if (!string.IsNullOrEmpty(FormData.MachineIdQuery))
-            {
-                var DecodedQuery = HttpUtility.HtmlDecode(FormData.MachineIdQuery).ToLower();
-                results =
-                        (
-                            from CrashDetail in results
-                            where CrashDetail.ComputerName.Equals(FormData.MachineIdQuery)
-                            select CrashDetail
-                        );
-            }
-
-            if (!string.IsNullOrEmpty(FormData.JiraQuery))
-            {
-                var DecodedQuery = HttpUtility.HtmlDecode(FormData.JiraQuery).ToLower();
-                results =
-                    (
-                        from CrashDetail in results
-                        where CrashDetail.TTPID.Equals(FormData.JiraQuery)
-                        select CrashDetail
-                    );
-            }
-            // Filter by BranchName
-            if (!string.IsNullOrEmpty(FormData.BranchName))
-            {
-                if (FormData.BranchName.StartsWith("-"))
-                {
-                    results =
-                    (
-                        from CrashDetail in results
-                        where !CrashDetail.Branch.Contains(FormData.BranchName.Substring(1))
-                        select CrashDetail
-                    );
-                }
-                else
-                {
-                    results =
-                    (
-                        from CrashDetail in results
-                        where CrashDetail.Branch.Equals(FormData.BranchName)
-                        select CrashDetail
-                    );
-                }
-            }
-
-            // Filter by GameName
-            if (!string.IsNullOrEmpty(FormData.GameName))
-            {
-                if (FormData.GameName.StartsWith("-"))
-                {
-                    results =
-                    (
-                        from CrashDetail in results
-                        where !CrashDetail.GameName.Contains(FormData.GameName.Substring(1))
-                        select CrashDetail
-                    );
-                }
-                else
-                {
-                    results =
-                    (
-                        from CrashDetail in results
-                        where CrashDetail.GameName.Contains(FormData.GameName)
-                        select CrashDetail
-                    );
-                }
-            }
-            return results;
-        }
+		private IQueryable<Crash> ConstructQuery(FormHelper FormData)
+		{
+			var results = ListAll();
+
+			// Grab Results 
+			if (!string.IsNullOrEmpty(FormData.SearchQuery))
+			{
+				if (!string.IsNullOrEmpty(FormData.UsernameQuery))
+				{
+					//We only use SearchQuery now for CallStack searching - if there's a searchquery value and a Username value, we need to get rid of the 
+					//Username so that we can create a broader search range
+					FormData.UsernameQuery = "";
+				}
+				string DecodedQuery = HttpUtility.HtmlDecode(FormData.SearchQuery).ToLower();
+				using (FScopedLogTimer LogTimer2 = new FScopedLogTimer("CrashRepository.GetResults.FindUserFromQuery" + "(Query=" + DecodedQuery + ")"))
+				{
+					results = results.Where(item => item.RawCallStack.Contains(FormData.SearchQuery));
+				}
+			}
+
+			// Filter by Crash Type
+			if (FormData.CrashType != "All")
+			{
+				switch (FormData.CrashType)
+				{
+					case "Crashes":
+						results = results.Where(CrashInstance => CrashInstance.CrashType == 1);
+						break;
+					case "Assert":
+						results = results.Where(CrashInstance => CrashInstance.CrashType == 2);
+						break;
+					case "Ensure":
+						results = results.Where(CrashInstance => CrashInstance.CrashType == 3);
+						break;
+					case "CrashesAsserts":
+						results = results.Where(CrashInstance => CrashInstance.CrashType == 1 || CrashInstance.CrashType == 2);
+						break;
+				}
+			}
+
+			if (!string.IsNullOrEmpty(FormData.UsernameQuery))
+			{
+				results =
+					(
+						from CrashDetail in results
+						where CrashDetail.UserName.Equals(FormData.UsernameQuery)
+						select CrashDetail
+						);
+			}
+
+			// Start Filtering the results
+			if (!string.IsNullOrEmpty(FormData.EpicIdQuery))
+			{
+				var DecodedQuery = HttpUtility.HtmlDecode(FormData.EpicIdQuery).ToLower();
+				results =
+						(
+							from CrashDetail in results
+							where CrashDetail.EpicAccountId.Equals(FormData.EpicIdQuery)
+							select CrashDetail
+						);
+			}
+
+			if (!string.IsNullOrEmpty(FormData.MachineIdQuery))
+			{
+				var DecodedQuery = HttpUtility.HtmlDecode(FormData.MachineIdQuery).ToLower();
+				results =
+						(
+							from CrashDetail in results
+							where CrashDetail.ComputerName.Equals(FormData.MachineIdQuery)
+							select CrashDetail
+						);
+			}
+
+			if (!string.IsNullOrEmpty(FormData.JiraQuery))
+			{
+				var DecodedQuery = HttpUtility.HtmlDecode(FormData.JiraQuery).ToLower();
+				results =
+					(
+						from CrashDetail in results
+						where CrashDetail.TTPID.Equals(FormData.JiraQuery)
+						select CrashDetail
+					);
+			}
+			// Filter by BranchName
+			if (!string.IsNullOrEmpty(FormData.BranchName))
+			{
+				if (FormData.BranchName.StartsWith("-"))
+				{
+					results =
+					(
+						from CrashDetail in results
+						where !CrashDetail.Branch.Contains(FormData.BranchName.Substring(1))
+						select CrashDetail
+					);
+				}
+				else
+				{
+					results =
+					(
+						from CrashDetail in results
+						where CrashDetail.Branch.Equals(FormData.BranchName)
+						select CrashDetail
+					);
+				}
+			}
+
+			// Filter by GameName
+			if (!string.IsNullOrEmpty(FormData.GameName))
+			{
+				if (FormData.GameName.StartsWith("-"))
+				{
+					results =
+					(
+						from CrashDetail in results
+						where !CrashDetail.GameName.Contains(FormData.GameName.Substring(1))
+						select CrashDetail
+					);
+				}
+				else
+				{
+					results =
+					(
+						from CrashDetail in results
+						where CrashDetail.GameName.Contains(FormData.GameName)
+						select CrashDetail
+					);
+				}
+			}
+			return results;
+		}
 	}
 }
diff --git a/Engine/Source/Programs/CrashReporter/CrashReportWebSite/Models/CrashesViewModel.cs b/Engine/Source/Programs/CrashReporter/CrashReportWebSite/Models/CrashesViewModel.cs
--- a/Engine/Source/Programs/CrashReporter/CrashReportWebSite/Models/CrashesViewModel.cs
+++ b/Engine/Source/Programs/CrashReporter/CrashReportWebSite/Models/CrashesViewModel.cs
@@ -38,13 +38,13 @@ namespace Tools.CrashReporter.CrashReportWebSite.Models
 		/// <summary>The query for advanced crash filtering.</summary>
 		public string SearchQuery { get; set; }
 
-        public string UsernameQuery { get; set; }
+		public string UsernameQuery { get; set; }
 
-        public string EpicIdQuery { get; set; }
+		public string EpicIdQuery { get; set; }
 
-        public string MachineIdQuery { get; set; }
+		public string MachineIdQuery { get; set; }
 
-        public string JiraQuery { get; set; }
+		public string JiraQuery { get; set; }
 
 		/// <summary>The date of the earliest crash to display.</summary>
 		public long DateFrom { get; set; }
@@ -64,8 +64,8 @@ namespace Tools.CrashReporter.CrashReportWebSite.Models
 		/// <summary>The user supplied parameters.</summary>
 		public FormCollection FormCollection { get; set; }
 
-        /// <summary>A collection of Branch Names used in the drop down on the main search form</summary>
-        public List<SelectListItem> BranchNames { get; set; }
+		/// <summary>A collection of Branch Names used in the drop down on the main search form</summary>
+		public List<SelectListItem> BranchNames { get; set; }
 
 		/// <summary>The set of statuses a crash could have its status set to.</summary>
 		public IEnumerable<string> SetStatus { get { return new List<string>( new string[] { "Unset", "Reviewed", "New", "Coder", "EngineQA", "GameQA" } ); } }
@@ -76,13 +76,13 @@ namespace Tools.CrashReporter.CrashReportWebSite.Models
 		/// <summary></summary>
 		public string GenerationTime { get; set; }
 
-        public CrashesViewModel()
-        {
-            DateTime FromDate = DateTime.Today.AddDays(-7);
-            DateTime ToDate = DateTime.Today;
-            BranchNames = DatabaseHelper.GetBranches();
-            DateTo = (long)(ToDate - Epoch).TotalMilliseconds;
-            DateFrom = (long)(FromDate - Epoch).TotalMilliseconds;
-        }
+		public CrashesViewModel()
+		{
+			DateTime FromDate = DateTime.Today.AddDays(-7);
+			DateTime ToDate = DateTime.Today;
+			BranchNames = CrashRepository.GetBranches();
+			DateTo = (long)(ToDate - Epoch).TotalMilliseconds;
+			DateFrom = (long)(FromDate - Epoch).TotalMilliseconds;
+		}
 	}
 }
\ No newline at end of file
