CrashReportLinq - Updated RepairDatabase program

diff --git a/Engine/Source/Programs/CrashReporter/CrashReportWebSite/Queries/Fix-RepairDatabase.linq b/Engine/Source/Programs/CrashReporter/CrashReportWebSite/Queries/Fix-RepairDatabase.linq
--- a/Engine/Source/Programs/CrashReporter/CrashReportWebSite/Queries/Fix-RepairDatabase.linq
+++ b/Engine/Source/Programs/CrashReporter/CrashReportWebSite/Queries/Fix-RepairDatabase.linq
@@ -33,13 +33,13 @@ HashSet<int> GetInvalidFuncIds()
 }
 
 static Stopwatch Timer = Stopwatch.StartNew();
-static DateTime Date = new DateTime(2010,05,26);//DateTime.UtcNow.AddMonths(-2000);
+static DateTime Date = DateTime.UtcNow.AddMonths(-1);
 static TimeSpan Tick = TimeSpan.FromDays(3);
 const int NUM_OPS_PER_BATCH = 32;
 
 void WriteLine( string Line )
 {
-	Debug.WriteLine( string.Format( "{0,5}: {1}", (long)Timer.Elapsed.TotalSeconds, Line ) );
+	Debug.WriteLine( string.Format( "{0,9} {1}", Timer.Elapsed.ToString(@"hh\:mm\:ss"), Line ) );
 }
 
 int NoPatternTotal = 0;
