Fix deprecation warning in SCrashReportClient.cpp

diff --git a/Engine/Source/Programs/CrashReporter/CrashReportClient/Private/SCrashReportClient.cpp b/Engine/Source/Programs/CrashReporter/CrashReportClient/Private/SCrashReportClient.cpp
--- a/Engine/Source/Programs/CrashReporter/CrashReportClient/Private/SCrashReportClient.cpp
+++ b/Engine/Source/Programs/CrashReporter/CrashReportClient/Private/SCrashReportClient.cpp
@@ -57,7 +57,7 @@ void SCrashReportClient::Construct(const FArguments& InArgs, TSharedRef<FCrashRe
 				[
 					SNew(STextBlock)
 					.TextStyle(FCrashReportClientStyle::Get(), "Title")
-					.Text(CrashedAppName)
+					.Text(FText::FromString(CrashedAppName))
 				]
 			]
 
@@ -66,8 +66,7 @@ void SCrashReportClient::Construct(const FArguments& InArgs, TSharedRef<FCrashRe
 			.Padding( FMargin( 4, 10 ) )
 			[
 				SNew( SRichTextBlock )
-				.Text( LOCTEXT("CrashDetailed", "We are very sorry that this crash occurred. Our goal is to prevent crashes like this from occurring in the future. Please help us track down and fix this crash by providing detailed information about what you were doing so that we may reproduce the crash and fix it quickly. You can also log a Bug Report with us at <a id=\"browser\" href=\"https://answers.unrealengine.com\" style=\"Hyperlink\">AnswerHub</> and work directly with support staff to report this issue.")
-				)
+				.Text( LOCTEXT("CrashDetailed", "We are very sorry that this crash occurred. Our goal is to prevent crashes like this from occurring in the future. Please help us track down and fix this crash by providing detailed information about what you were doing so that we may reproduce the crash and fix it quickly. You can also log a Bug Report with us at <a id=\"browser\" href=\"https://answers.unrealengine.com\" style=\"Hyperlink\">AnswerHub</> and work directly with support staff to report this issue."))
 				.AutoWrapText(true)
 				.DecoratorStyleSet( &FCoreStyle::Get() )
 				+ SRichTextBlock::HyperlinkDecorator( TEXT("browser"), FSlateHyperlinkRun::FOnClick::CreateStatic( &OnBrowserLinkClicked, AsShared() ) )
