CrashReportWWW - Added CrashesInTimeFrame for the crash reports

diff --git a/Engine/Source/Programs/CrashReporter/CrashReportWebSite/Models/CrashReport.cs b/Engine/Source/Programs/CrashReporter/CrashReportWebSite/Models/CrashReport.cs
--- a/Engine/Source/Programs/CrashReporter/CrashReportWebSite/Models/CrashReport.cs
+++ b/Engine/Source/Programs/CrashReporter/CrashReportWebSite/Models/CrashReport.cs
@@ -154,6 +154,8 @@ namespace Tools.CrashReporter.CrashReportWebSite.Models
 			this.AffectedMajorVersions = new SortedSet<string>(); // 4.4, 4.5 and so
 			this.BranchesFoundIn = new SortedSet<string>();
 			this.AffectedPlatforms = new SortedSet<string>();
+			this.CrashesInTimeFrameAll = CrashesForBugg.Count;
+			this.CrashesInTimeFrameGroup = CrashesForBugg.Count;
 			var HashSetDecsriptions = new HashSet<string>();
 
 			HashSet<string> MachineIds = new HashSet<string>();
