CrashReportWWW - Now ignore all calls from FDebug class

diff --git a/Engine/Source/Programs/CrashReporter/CrashReportWebSite/Models/CrashDataFormater.cs b/Engine/Source/Programs/CrashReporter/CrashReportWebSite/Models/CrashDataFormater.cs
--- a/Engine/Source/Programs/CrashReporter/CrashReportWebSite/Models/CrashDataFormater.cs
+++ b/Engine/Source/Programs/CrashReporter/CrashReportWebSite/Models/CrashDataFormater.cs
@@ -388,10 +388,8 @@ namespace Tools.CrashReporter.CrashReportWebSite.Models
 				var FuncsToRemove = new HashSet<string>( new string[] 
 				{ 
 					"RaiseException", 
-					"FDebug::EnsureFailed", 		
+					"FDebug::", 		
 					"Error::Serialize",
-					"FDebug::AssertFailed",
-					"FDebug::EnsureNotFalseFormatted",
 					"FOutputDevice::Logf",
 					"FMsg::Logf",
 					"ReportCrash",
