CrashReportWWW - UECORE-138 / Crash Reporter search defaults to an end date of yesterday in the US

diff --git a/Engine/Source/Programs/CrashReporter/CrashReportWebSite/Models/CrashesViewModel.cs b/Engine/Source/Programs/CrashReporter/CrashReportWebSite/Models/CrashesViewModel.cs
--- a/Engine/Source/Programs/CrashReporter/CrashReportWebSite/Models/CrashesViewModel.cs
+++ b/Engine/Source/Programs/CrashReporter/CrashReportWebSite/Models/CrashesViewModel.cs
@@ -15,7 +15,7 @@ namespace Tools.CrashReporter.CrashReportWebSite.Models
 	public class CrashesViewModel
 	{
 		/// <summary>Dates are in milliseconds relative to this base (for ease of Javascript usage)</summary>
-		public static readonly DateTime Epoch = new DateTime(1970, 1, 1);
+		public static readonly DateTime Epoch = new DateTime( 1970, 1, 1 );
 
 		/// <summary>A container for all the crashes that pass the current filters.</summary>
 		public IEnumerable<Crash> Results { get; set; }
@@ -46,9 +46,9 @@ namespace Tools.CrashReporter.CrashReportWebSite.Models
 
 		public string JiraQuery { get; set; }
 
-        public string MessageQuery { get; set; }
+		public string MessageQuery { get; set; }
 
-        public string DescriptionQuery { get; set; }
+		public string DescriptionQuery { get; set; }
 
 		/// <summary>The date of the earliest crash to display.</summary>
 		public long DateFrom { get; set; }
@@ -77,16 +77,19 @@ namespace Tools.CrashReporter.CrashReportWebSite.Models
 		/// <summary> The real user name, displayed only for searches. </summary>
 		public string RealUserName { get; set; }
 
-		/// <summary></summary>
+		/// <summary>Time spent in generating this site, formatted as a string.</summary>
 		public string GenerationTime { get; set; }
 
+		/// <summary>
+		/// Default constructor, used by HomeController
+		/// </summary>
 		public CrashesViewModel()
 		{
-			DateTime FromDate = DateTime.Today.AddDays(-7);
-			DateTime ToDate = DateTime.Today;
+			DateTime FromDate = DateTime.Today.AddDays( -7 ).ToUniversalTime();
+			DateTime ToDate = DateTime.Today.ToUniversalTime();
 			BranchNames = CrashRepository.GetBranches();
-			DateTo = (long)(ToDate - Epoch).TotalMilliseconds;
-			DateFrom = (long)(FromDate - Epoch).TotalMilliseconds;
+			DateFrom = (long)( FromDate - Epoch ).TotalMilliseconds;
+			DateTo = (long)( ToDate - Epoch ).TotalMilliseconds;		
 		}
 	}
 }
\ No newline at end of file
