CrashReportWWW - Only list branches that start with "UE4"

diff --git a/Engine/Source/Programs/CrashReporter/CrashReportWebSite/Models/CrashRepository.cs b/Engine/Source/Programs/CrashReporter/CrashReportWebSite/Models/CrashRepository.cs
--- a/Engine/Source/Programs/CrashReporter/CrashReportWebSite/Models/CrashRepository.cs
+++ b/Engine/Source/Programs/CrashReporter/CrashReportWebSite/Models/CrashRepository.cs
@@ -69,10 +69,10 @@ namespace Tools.CrashReporter.CrashReportWebSite.Models
 		{
 			CrashReportDataContext Context = new CrashReportDataContext();
 
-			var lm = Context.Crashes.Where( n => n.Branch.Contains( "UE4" ) ).Select( n => n.Branch ).Distinct().ToList();
-			var list = lm.Select( listitem => new SelectListItem { Selected = false, Text = listitem, Value = listitem } ).ToList();
-			list.Insert( 0, new SelectListItem { Selected = true, Text = "", Value = "" } );
-			return list;
+			var BranchList = Context.Crashes.Where( n => n.Branch.StartsWith( "UE4" ) ).Select( n => n.Branch ).Distinct().ToList();
+			var SelectListItens = BranchList.Select( listitem => new SelectListItem { Selected = false, Text = listitem, Value = listitem } ).ToList();
+			SelectListItens.Insert( 0, new SelectListItem { Selected = true, Text = "", Value = "" } );
+			return SelectListItens;
 		}
 		
 		/// <summary>
