CrashReportWWW - UECORE-137 / Crash reports page showing lifetime number of occurrences instead of date range number.

diff --git a/Engine/Source/Programs/CrashReporter/CrashReportWebSite/Controllers/ReportsController.cs b/Engine/Source/Programs/CrashReporter/CrashReportWebSite/Controllers/ReportsController.cs
--- a/Engine/Source/Programs/CrashReporter/CrashReportWebSite/Controllers/ReportsController.cs
+++ b/Engine/Source/Programs/CrashReporter/CrashReportWebSite/Controllers/ReportsController.cs
@@ -51,10 +51,8 @@ namespace Tools.CrashReporter.CrashReportWebSite.Controllers
 				int AnonymousID = AnonumousIDs.First();
 				HashSet<string> UserNamesForUserGroup = FRepository.Get( BuggsRepo ).GetUserNamesFromGroupName( AnonumousGroup );
 
-				//FormData.DateFrom = DateTime.Now.AddDays( -1 );
-
 				var Crashes = CrashRepo
-					.FilterByDate( CrashRepo.ListAll(), FormData.DateFrom, FormData.DateTo.AddDays( 1 ) )
+					.FilterByDate( CrashRepo.ListAll(), FormData.DateFrom, FormData.DateTo )
 					// Only crashes and asserts
 					.Where( Crash => Crash.CrashType == 1 || Crash.CrashType == 2 )
 					// Only anonymous user
@@ -291,7 +289,7 @@ namespace Tools.CrashReporter.CrashReportWebSite.Controllers
 					}
 				}
 
-				Buggs = Buggs.OrderByDescending( b => b.NumberOfCrashes ).ToList();
+				Buggs = Buggs.OrderByDescending( b => b.CrashesInTimeFrameGroup ).ToList();
 
 				return new ReportsViewModel
 				{
diff --git a/Engine/Source/Programs/CrashReporter/CrashReportWebSite/Views/Reports/ViewReports.ascx b/Engine/Source/Programs/CrashReporter/CrashReportWebSite/Views/Reports/ViewReports.ascx
--- a/Engine/Source/Programs/CrashReporter/CrashReportWebSite/Views/Reports/ViewReports.ascx
+++ b/Engine/Source/Programs/CrashReporter/CrashReportWebSite/Views/Reports/ViewReports.ascx
@@ -133,7 +133,7 @@ table, table tr, table th, table td
 		</td>
 		<td>
 			<%--NewBugg.NumberOfCrashes = Top.Value;					// # Occurrences--%>
-			<%=Bugg.NumberOfCrashes%>
+			<%=Bugg.CrashesInTimeFrameGroup%>
 		</td>
 		<td>
 			<%--NewBugg.BuildVersion = NewBugg.AffectedVersions.Last();	// Latest Version Affected--%>
