CrashReportWWW - Show only latest 64 crashes in the bugg

diff --git a/Engine/Source/Programs/CrashReporter/CrashReportWebSite/Views/Buggs/Show.aspx b/Engine/Source/Programs/CrashReporter/CrashReportWebSite/Views/Buggs/Show.aspx
--- a/Engine/Source/Programs/CrashReporter/CrashReportWebSite/Views/Buggs/Show.aspx
+++ b/Engine/Source/Programs/CrashReporter/CrashReportWebSite/Views/Buggs/Show.aspx
@@ -250,7 +250,12 @@
 					int IterationCount = 0;
 					foreach( Crash CrashInstance in Model.Crashes )
 					{
-						IterationCount++;%>
+						IterationCount++;
+						if (IterationCount > 64)
+						{
+							break;
+						}
+						%>
 			
 						<tr class='CrashRow'>
 							<td class="Id"><%=Html.ActionLink( CrashInstance.Id.ToString(), "Show", new { controller = "crashes", id = CrashInstance.Id }, null )%></td>
