CrashReportWWW - Added a few comments

diff --git a/Engine/Source/Programs/CrashReporter/CrashReportWebSite/Models/BuggViewModel.cs b/Engine/Source/Programs/CrashReporter/CrashReportWebSite/Models/BuggViewModel.cs
--- a/Engine/Source/Programs/CrashReporter/CrashReportWebSite/Models/BuggViewModel.cs
+++ b/Engine/Source/Programs/CrashReporter/CrashReportWebSite/Models/BuggViewModel.cs
@@ -22,7 +22,7 @@ namespace Tools.CrashReporter.CrashReportWebSite.Models
 		/// <summary></summary>
 		public string SourceContext { get; set; }
 
-		/// <summary></summary>
+		/// <summary>Time spent in generating this site, formatted as a string.</summary>
 		public string GenerationTime { get; set; }
 	}
 }
\ No newline at end of file
diff --git a/Engine/Source/Programs/CrashReporter/CrashReportWebSite/Models/BuggsViewModel.cs b/Engine/Source/Programs/CrashReporter/CrashReportWebSite/Models/BuggsViewModel.cs
--- a/Engine/Source/Programs/CrashReporter/CrashReportWebSite/Models/BuggsViewModel.cs
+++ b/Engine/Source/Programs/CrashReporter/CrashReportWebSite/Models/BuggsViewModel.cs
@@ -90,7 +90,7 @@ namespace Tools.CrashReporter.CrashReportWebSite.Models
 			}
 		}
 
-		/// <summary></summary>
+		/// <summary>Time spent in generating this site, formatted as a string.</summary>
 		public string GenerationTime { get; set; }
 	}
 }
\ No newline at end of file
diff --git a/Engine/Source/Programs/CrashReporter/CrashReportWebSite/Models/CrashViewModel.cs b/Engine/Source/Programs/CrashReporter/CrashReportWebSite/Models/CrashViewModel.cs
--- a/Engine/Source/Programs/CrashReporter/CrashReportWebSite/Models/CrashViewModel.cs
+++ b/Engine/Source/Programs/CrashReporter/CrashReportWebSite/Models/CrashViewModel.cs
@@ -12,7 +12,7 @@ namespace Tools.CrashReporter.CrashReportWebSite.Models
 		/// <summary>The callstack associated with the crash.</summary>
 		public CallStackContainer CallStack { get; set; }
 
-		/// <summary></summary>
+		/// <summary>Time spent in generating this site, formatted as a string.</summary>
 		public string GenerationTime { get; set; }
 	}
 }
\ No newline at end of file
diff --git a/Engine/Source/Programs/CrashReporter/CrashReportWebSite/Models/DashboardViewModel.cs b/Engine/Source/Programs/CrashReporter/CrashReportWebSite/Models/DashboardViewModel.cs
--- a/Engine/Source/Programs/CrashReporter/CrashReportWebSite/Models/DashboardViewModel.cs
+++ b/Engine/Source/Programs/CrashReporter/CrashReportWebSite/Models/DashboardViewModel.cs
@@ -14,7 +14,7 @@ namespace Tools.CrashReporter.CrashReportWebSite.Models
 		/// <summary>An encoded table of buggs by day for the display plugin to use.</summary>
 		public string BuggsByDay { get; set; }
 
-		/// <summary></summary>
+		/// <summary>Time spent in generating this site, formatted as a string.</summary>
 		public string GenerationTime { get; set; }
 	}
 }
\ No newline at end of file
diff --git a/Engine/Source/Programs/CrashReporter/CrashReportWebSite/Models/ReportsViewModel.cs b/Engine/Source/Programs/CrashReporter/CrashReportWebSite/Models/ReportsViewModel.cs
--- a/Engine/Source/Programs/CrashReporter/CrashReportWebSite/Models/ReportsViewModel.cs
+++ b/Engine/Source/Programs/CrashReporter/CrashReportWebSite/Models/ReportsViewModel.cs
@@ -33,7 +33,7 @@ namespace Tools.CrashReporter.CrashReportWebSite.Models
 		/// <summary></summary>
 		public int TotalAffectedUsers { get; set; }
 
-		/// <summary></summary>
+		/// <summary>Time spent in generating this site, formatted as a string.</summary>
 		public string GenerationTime { get; set; }
 
 		/// <summary>User input from the client.</summary>
