CrashReportWWW - Minor optimization, GetBranches will return results from the crashes from the last 3 months

diff --git a/Engine/Source/Programs/CrashReporter/CrashReportWebSite/Models/CrashRepository.cs b/Engine/Source/Programs/CrashReporter/CrashReportWebSite/Models/CrashRepository.cs
--- a/Engine/Source/Programs/CrashReporter/CrashReportWebSite/Models/CrashRepository.cs
+++ b/Engine/Source/Programs/CrashReporter/CrashReportWebSite/Models/CrashRepository.cs
@@ -64,13 +64,19 @@ namespace Tools.CrashReporter.CrashReportWebSite.Models
 		/// <summary>
 		/// Retrieves a list of distinct UE4 Branches from the CrashRepository
 		/// </summary>
-		/// <returns></returns>
 		public static List<SelectListItem> GetBranches()
 		{
 			CrashReportDataContext Context = new CrashReportDataContext();
 
-			var BranchList = Context.Crashes.Where( n => n.Branch.StartsWith( "UE4" ) ).Select( n => n.Branch ).Distinct().ToList();
-			var SelectListItens = BranchList.Select( listitem => new SelectListItem { Selected = false, Text = listitem, Value = listitem } ).ToList();
+			var BranchList = Context.Crashes
+				.Where( n => n.TimeOfCrash > DateTime.Now.AddMonths( -3 ) )
+				.Where( n => n.Branch.StartsWith( "UE4" ) )
+				.Select( n => n.Branch )
+				.Distinct()
+				.ToList();
+			var SelectListItens = BranchList
+				.Select( listitem => new SelectListItem { Selected = false, Text = listitem, Value = listitem } )
+				.ToList();
 			SelectListItens.Insert( 0, new SelectListItem { Selected = true, Text = "", Value = "" } );
 			return SelectListItens;
 		}
@@ -214,9 +220,9 @@ namespace Tools.CrashReporter.CrashReportWebSite.Models
 				var ResultsAll = ConstructQuery(FormData);
 
 				// Filter by data and get as enumerable.
-			    Results = FilterByDate(ResultsAll, FormData.DateFrom, FormData.DateTo);
+				Results = FilterByDate(ResultsAll, FormData.DateFrom, FormData.DateTo);
 
-			    // Get UserGroup ResultCounts
+				// Get UserGroup ResultCounts
 				Dictionary<string, int> GroupCounts = GetCountsByGroupFromCrashes( Results );
 
 				// Filter by user group if present
@@ -282,8 +288,8 @@ namespace Tools.CrashReporter.CrashReportWebSite.Models
 					UsernameQuery = FormData.UsernameQuery,
 					EpicIdQuery = FormData.EpicIdQuery,
 					MachineIdQuery = FormData.MachineIdQuery,
-                    DescriptionQuery = FormData.DescriptionQuery,
-                    MessageQuery = FormData.MessageQuery,
+					DescriptionQuery = FormData.DescriptionQuery,
+					MessageQuery = FormData.MessageQuery,
 					JiraQuery = FormData.JiraQuery,
 					DateFrom = (long)( FormData.DateFrom - CrashesViewModel.Epoch ).TotalMilliseconds,
 					DateTo = (long)( FormData.DateTo - CrashesViewModel.Epoch ).TotalMilliseconds,
@@ -711,12 +717,12 @@ namespace Tools.CrashReporter.CrashReportWebSite.Models
 			}
 			else
 			{
-			    results = 
-                    (
-                    from CrashDetail in results
-                    where !CrashDetail.Branch.Equals("UE4-UT-Releases") && !CrashDetail.Branch.Equals("UE4-UT")
-			        select CrashDetail
-                    );
+				results = 
+					(
+					from CrashDetail in results
+					where !CrashDetail.Branch.Equals("UE4-UT-Releases") && !CrashDetail.Branch.Equals("UE4-UT")
+					select CrashDetail
+					);
 			}
 
 			// Filter by GameName
@@ -742,35 +748,35 @@ namespace Tools.CrashReporter.CrashReportWebSite.Models
 				}
 			}
 
-		    if (!string.IsNullOrEmpty(FormData.MessageQuery))
-		    {
-		        results =
-		            (
-		                from CrashDetail in results
-		                where SqlMethods.Like(CrashDetail.Summary, "%" + FormData.MessageQuery + "%")
-		                select CrashDetail
-		                );
-		    }
-
-		    if (!string.IsNullOrEmpty(FormData.DescriptionQuery))
-		    {
-		        results = 
-                    (
-                        from CrashDetail in results
-                        where SqlMethods.Like(CrashDetail.Description, "%" + FormData.DescriptionQuery + "%")
-		                select CrashDetail
-                    );
-		    }
-
-		    if (!string.IsNullOrEmpty(FormData.SearchQuery))
-		    {
-		        results = 
-                    (
-                        from CrashDetail in results
-                        where SqlMethods.Like(CrashDetail.RawCallStack, "%" + FormData.SearchQuery + "%")
-		                select CrashDetail
-                    );
-		    }
+			if (!string.IsNullOrEmpty(FormData.MessageQuery))
+			{
+				results =
+					(
+						from CrashDetail in results
+						where SqlMethods.Like(CrashDetail.Summary, "%" + FormData.MessageQuery + "%")
+						select CrashDetail
+						);
+			}
+
+			if (!string.IsNullOrEmpty(FormData.DescriptionQuery))
+			{
+				results = 
+					(
+						from CrashDetail in results
+						where SqlMethods.Like(CrashDetail.Description, "%" + FormData.DescriptionQuery + "%")
+						select CrashDetail
+					);
+			}
+
+			if (!string.IsNullOrEmpty(FormData.SearchQuery))
+			{
+				results = 
+					(
+						from CrashDetail in results
+						where SqlMethods.Like(CrashDetail.RawCallStack, "%" + FormData.SearchQuery + "%")
+						select CrashDetail
+					);
+			}
 
 			return results;
 		}
