CrashReportWWW - UECORE-110/Crash Reporter: Incomplete CrashGroups search results

diff --git a/Engine/Source/Programs/CrashReporter/CrashReportWebSite/Models/BuggRepository.cs b/Engine/Source/Programs/CrashReporter/CrashReportWebSite/Models/BuggRepository.cs
--- a/Engine/Source/Programs/CrashReporter/CrashReportWebSite/Models/BuggRepository.cs
+++ b/Engine/Source/Programs/CrashReporter/CrashReportWebSite/Models/BuggRepository.cs
@@ -363,8 +363,7 @@ namespace Tools.CrashReporter.CrashReportWebSite.Models
 					// Search for all function ids. OR operation, not very efficient, but for searching for one function should be ok.
 					foreach (int Id in AllFuncionCallIds)
 					{
-						string Pattern = "+" + Id + "+";
-						var BuggsForFunc = Results.Where( X => X.Pattern.Contains( Pattern ) ).ToList();
+						var BuggsForFunc = Results.Where( X => X.Pattern.Contains( Id + "+" ) || X.Pattern.Contains( "+" + Id ) ).ToList();
 						
 						foreach(Bugg Bugg in BuggsForFunc)
 						{
@@ -561,7 +560,16 @@ namespace Tools.CrashReporter.CrashReportWebSite.Models
 		{
 			using( FAutoScopedLogTimer LogTimer = new FAutoScopedLogTimer( this.GetType().ToString() + " SQL" ) )
 			{
-				IQueryable<Bugg> BuggsInTimeFrame = Results.Where( Bugg => Bugg.TimeOfLastCrash >= DateFrom && Bugg.TimeOfLastCrash <= DateTo.AddDays( 1 ) );
+				var DateTo1 = DateTo.AddDays( 1 );
+
+				IQueryable<Bugg> BuggsInTimeFrame = Results
+					.Where( Bugg =>
+					( Bugg.TimeOfFirstCrash <= DateFrom && Bugg.TimeOfLastCrash >= DateTo1 ) ||
+					( Bugg.TimeOfFirstCrash >= DateFrom && Bugg.TimeOfLastCrash <= DateTo1 ) ||
+					( Bugg.TimeOfFirstCrash >= DateFrom && Bugg.TimeOfFirstCrash <= DateTo1 && Bugg.TimeOfLastCrash >= DateTo1 ) ||
+					( Bugg.TimeOfFirstCrash <= DateFrom && Bugg.TimeOfLastCrash <= DateTo1 && Bugg.TimeOfLastCrash >= DateFrom )
+					);
+
 				IEnumerable<Bugg> BuggsInTimeFrameEnumerable = BuggsInTimeFrame.ToList();
 				return BuggsInTimeFrameEnumerable;
 			}
