CrashReportWWW - UECORE-109/Crash Reporter : Reports page shows 0 affected users for crashes in UE4-EpicGamesLauncher

diff --git a/Engine/Source/Programs/CrashReporter/CrashReportWebSite/Models/CrashReport.cs b/Engine/Source/Programs/CrashReporter/CrashReportWebSite/Models/CrashReport.cs
--- a/Engine/Source/Programs/CrashReporter/CrashReportWebSite/Models/CrashReport.cs
+++ b/Engine/Source/Programs/CrashReporter/CrashReportWebSite/Models/CrashReport.cs
@@ -175,9 +175,7 @@ namespace Tools.CrashReporter.CrashReportWebSite.Models
 					}
 				}
 
-				// Ignore bad build versions.
-				// @TODO yrx 2015-02-17 What about launcher?
-				if( Crash.BuildVersion.StartsWith( "4." ) )
+				// @TODO Ignore bad build versions.
 				{
 					if( !string.IsNullOrEmpty( Crash.BuildVersion ) )
 					{
