CrashReportWWW - Updated buggs search panel, added options to filter crashes by the build version

diff --git a/Engine/Source/Programs/CrashReporter/CrashReportWebSite/Views/Buggs/Index.aspx b/Engine/Source/Programs/CrashReporter/CrashReportWebSite/Views/Buggs/Index.aspx
--- a/Engine/Source/Programs/CrashReporter/CrashReportWebSite/Views/Buggs/Index.aspx
+++ b/Engine/Source/Programs/CrashReporter/CrashReportWebSite/Views/Buggs/Index.aspx
@@ -11,89 +11,135 @@
 	Crash Report: Buggs
 </asp:Content>
 
-<asp:Content ID="ScriptContent"  ContentPlaceHolderID="ScriptContent" runat="server" >
-<script type="text/javascript">
-	$(function ()
-	{
-		$("#dateFromVisible")
-			.datepicker({ maxDate: '+0D' })
-			.datepicker('setDate', new Date(parseInt($('#dateFrom').val())));
-
-		$("#dateToVisible")
-			.datepicker({ maxDate: '+0D' })
-			.datepicker('setDate', new Date(parseInt($('#dateTo').val())));
-
-	});
-
-	$.datepicker.setDefaults({
-		onSelect: function ()
+<asp:Content ID="ScriptContent" ContentPlaceHolderID="ScriptContent" runat="server">
+	<script src="../../Scripts/searchCollapse.js"></script>
+	<script type="text/javascript">
+		$(function ()
 		{
-			$("#dateFrom").val($("#dateFromVisible").datepicker('getDate').getTime());
-			$("#dateTo").val($("#dateToVisible").datepicker('getDate').getTime());
-		}
-	});
-
-	$(document).ready(function ()
-	{
-		//Shift Check box
-		$(":checkbox").shiftcheckbox();
-
-		//Zebrastripes
-		$("#CrashesTable tr:nth-child(even)").css("background-color", "#C3CAD0");
-		$("#CrashesTable tr:nth-child(odd)").css("background-color", "#eeeeee");
-
-	});
-</script>
-</asp:Content>
-
-<asp:Content ID="AboveMainContent" ContentPlaceHolderID="AboveMainContent" runat="server" >
-	<div style="clear:both;"><small style="color: lightgray;">Generated in <%=Model.GenerationTime%> second(s)</small><br /></div>
-
-<div id='SearchForm' style="clear:both;">
-<%using( Html.BeginForm( "", "Buggs", FormMethod.Get, new { id = "FilterBuggsForm" } ) )
-{ %>
-	<%=Html.HiddenFor( u => u.UserGroup )%>
-	<%=Html.Hidden( "SortTerm", Model.SortTerm )%>
-	<%=Html.Hidden( "SortOrder", Model.SortOrder )%>
-
-	<span style="float:left;">
-			<input type="radio" name="CrashType" class="CrashType" value="CrashesAsserts" <%=( Model.CrashType == "CrashesAsserts" ) ? "checked='checked'" : "" %> /> <span title='All Crashes Except Ensures'>Crashes+Asserts</span>
-			<input type="radio" name="CrashType" class="CrashType" value="Ensure" <%=( Model.CrashType == "Ensure" ) ? "checked='checked'" : "" %>/> <span title='Only Ensures'>Ensures</span>
-			<input type="radio" name="CrashType" class="CrashType" value="Assert" <%=( Model.CrashType == "Assert" ) ? "checked='checked'" : "" %>/> <span title='Only Asserts'>Asserts</span>
-			<input type="radio" name="CrashType" class="CrashType" value="Crashes" <%=( Model.CrashType == "Crashes" ) ? "checked='checked'" : "" %>/> <span title='Crashes Except Ensures and Asserts'>Crashes</span>
-			<input type="radio" name="CrashType" class="CrashType" value="All" <%=( Model.CrashType == "All" ) ? "checked='checked'" : "" %>/> <span title='All Crashes'>All</span>
-	</span>
-
-	<div>
-		<%=Html.TextBox( "SearchQuery", Model.SearchQuery )%><input type="submit" value="Search" class='SearchButton' />
-	</div>
+			$("#dateFromVisible")
+				.datepicker({ maxDate: '+0D' })
+				.datepicker('setDate', new Date(parseInt($('#dateFrom').val())));
 
-	<script>$.datepicker.setDefaults($.datepicker.regional['']);</script>
+			$("#dateToVisible")
+				.datepicker({ maxDate: '+0D' })
+				.datepicker('setDate', new Date(parseInt($('#dateTo').val())));
 
-	<span style="margin-left: 10px; font-weight: bold;">Filter by Date </span>
+		});
 
-	<span>From: 
-	<input id="dateFromVisible" type="text" class="date" autocomplete="OFF" /></span>
-	<input id="dateFrom" name="dateFrom" type="hidden" value="<%=Model.DateFrom %>" autocomplete="OFF" />
+		$.datepicker.setDefaults({
+			onSelect: function ()
+			{
+				$("#dateFrom").val($("#dateFromVisible").datepicker('getDate').getTime());
+				$("#dateTo").val($("#dateToVisible").datepicker('getDate').getTime());
+			}
+		});
 
-	<span>To: 
-	<input id="dateToVisible" type="text" class="date" autocomplete="OFF" /></span>
-	<input id="dateTo" name="dateTo" type="hidden" value="<%=Model.DateTo %>" autocomplete="OFF" />
+		$(document).ready(function ()
+		{
+			//Shift Check box
+			$(":checkbox").shiftcheckbox();
 
+			//Zebrastripes
+			$("#CrashesTable tr:nth-child(even)").css("background-color", "#C3CAD0");
+			$("#CrashesTable tr:nth-child(odd)").css("background-color", "#eeeeee");
 
-	<span style="margin-left: 10px; font-weight:bold;">Filter Build Version:</span>
-	<span><%=Html.DropDownListFor( m=>m.VersionName, Model.VersionNames )%></span>
+		});
+	</script>
+</asp:Content>
 
-<%} %>
-</div>
+<asp:Content ID="Content"  ContentPlaceHolderID="AboveMainContent" runat="server" >
+	<div style="clear: both;">
+		<small style="color: lightgray;">Generated in <%=Model.GenerationTime%> second(s)</small>
+		<br />
+
+		<div id='SearchForm' style="clear: both">
+			<div id="TitleBar" class="TitleBar">
+				SEARCH OPTIONS
+			</div>
+			<% using (Html.BeginForm( "", "Buggs", FormMethod.Get, new { id = "FilterBuggsForm" } ))
+	  { %>
+			<%=Html.HiddenFor( u => u.UserGroup )%>
+			<%=Html.Hidden( "SortTerm", Model.SortTerm )%>
+			<%=Html.Hidden( "SortOrder", Model.SortOrder )%>
+
+			<div id="Container">
+				<table>
+					<tr>
+						<td rowspan="3">
+							<div id="CrashTypeList">
+								<p class="SearchTextTitle" title="Press to hide">Show buggs of type</p>
+								<ul>
+									<li>
+										<input type="radio" name="CrashType" value="CrashesAsserts" <%=( Model.CrashType == "CrashesAsserts" ) ? "checked='checked'" : "" %> />
+										<span title='All Crashes Except Ensures'>Crashes+Asserts</span></li>
+									<li>
+										<input type="radio" name="CrashType" value="Ensure" <%=( Model.CrashType == "Ensure" ) ? "checked='checked'" : "" %> />
+										<span title='Only Ensures'>Ensures</span></li>
+									<li>
+										<input type="radio" name="CrashType" value="Assert" <%=( Model.CrashType == "Assert" ) ? "checked='checked'" : "" %> />
+										<span title='Only Asserts'>Asserts</span></li>
+									<li>
+										<input type="radio" name="CrashType" value="Crashes" <%=( Model.CrashType == "Crashes" ) ? "checked='checked'" : "" %> />
+										<span title='Crashes Except Ensures and Asserts'>Crashes</span></li>
+									<li>
+										<input type="radio" name="CrashType" value="All" <%=( Model.CrashType == "All" ) ? "checked='checked'" : "" %> />
+										<span title='All Crashes'>All</span></li>
+								</ul>
+							</div>
+						</td>
+						
+						<td>
+							<p class="SearchTextTitle">Call Stack</p>
+						</td>
+						<td>
+							<input name="SearchQuery" type="text" value="<%=Model.SearchQuery%>" title="" />
+						</td>
+
+						<td>
+							<p class="SearchTextTitle">Filter by Date</p>
+						</td>
+						<td>
+							<script>$.datepicker.setDefaults($.datepicker.regional['']);</script>
+							<span class="SearchTextTitle">From:</span>
+							<input id="dateFromVisible" type="text" class="date" autocomplete="OFF" style="width:80px" />
+							<input id="dateFrom" name="dateFrom" type="hidden" value="<%=Model.DateFrom %>" autocomplete="OFF" />
+							<span class="SearchTextTitle">To:</span>
+							<input id="dateToVisible" type="text" class="date" autocomplete="OFF" style="width:80px" />
+							<input id="dateTo" name="dateTo" type="hidden" value="<%=Model.DateTo %>"autocomplete="OFF"  />
+						</td>
+					</tr>
+
+
+					<tr>
+						
+						<td>
+							<p class="SearchTextTitle">Filter by Version</p>
+						</td>
+						<td>
+							<%=Html.DropDownListFor( m=>m.VersionName, Model.VersionNames )%>
+						</td>
+					</tr>
+					<tr>
+						<td colspan="7" >
+							<input type="submit" value="Search" class='SearchButton' />
+						</td>
+					</tr>
+				</table>
+			</div>
+			<% } %>
+		</div>
+	</div>
 </asp:Content>
 
 <asp:Content ID="MainContent" ContentPlaceHolderID="MainContent" runat="server">
 	<div id='CrashesTableContainer'>
 		<div id='UserGroupBar'>
-			<%foreach( var GroupCount in Model.GroupCounts )
-			{%>
-				<span class=<%if( Model.UserGroup == GroupCount.Key ){ %> "UserGroupTabSelected" <%} else {%> "UserGroupTab"<%} %> id="<%=GroupCount.Key%>Tab">
+			<%foreach (var GroupCount in Model.GroupCounts)
+	 {%>
+				<span class=<%if (Model.UserGroup == GroupCount.Key)
+				  { %> "UserGroupTabSelected" <%}
+				  else
+				  {%> "UserGroupTab"<%} %> id="<%=GroupCount.Key%>Tab">
 					<%=Url.UserGroupLink( GroupCount.Key, Model )%> 
 				<span class="UserGroupResults">
 					(<%=GroupCount.Value%>)
@@ -137,7 +183,7 @@
 							, 
 							new { style = "color:black; text-decoration:none;" } )%></span>
 				</div>
-				<% Html.RenderPartial("/Views/Buggs/ViewBuggs.ascx"); %>
+				<% Html.RenderPartial( "/Views/Buggs/ViewBuggs.ascx" ); %>
 			</form>
 		</div>
 	</div>
diff --git a/Engine/Source/Programs/CrashReporter/CrashReportWebSite/Views/Crashes/Index.aspx b/Engine/Source/Programs/CrashReporter/CrashReportWebSite/Views/Crashes/Index.aspx
--- a/Engine/Source/Programs/CrashReporter/CrashReportWebSite/Views/Crashes/Index.aspx
+++ b/Engine/Source/Programs/CrashReporter/CrashReportWebSite/Views/Crashes/Index.aspx
@@ -31,7 +31,6 @@ Crash Reports
 			{
 				$("#dateFrom").val($("#dateFromVisible").datepicker('getDate').getTime());
 				$("#dateTo").val($("#dateToVisible").datepicker('getDate').getTime());
-				//$("#FilterCrashesForm").submit();
 			}
 		});
 
