CrashReportWWW - Optimized GetVersions, GetBranches

diff --git a/Engine/Source/Programs/CrashReporter/CrashReportWebSite/Models/CrashRepository.cs b/Engine/Source/Programs/CrashReporter/CrashReportWebSite/Models/CrashRepository.cs
--- a/Engine/Source/Programs/CrashReporter/CrashReportWebSite/Models/CrashRepository.cs
+++ b/Engine/Source/Programs/CrashReporter/CrashReportWebSite/Models/CrashRepository.cs
@@ -61,24 +61,83 @@ namespace Tools.CrashReporter.CrashReportWebSite.Models
 			Context.Dispose();
 		}
 
+		private static DateTime LastBranchesDate = DateTime.UtcNow.AddDays(-1);
+		private static List<SelectListItem> BranchesAsSelectList = null;
 		/// <summary>
 		/// Retrieves a list of distinct UE4 Branches from the CrashRepository
 		/// </summary>
 		public static List<SelectListItem> GetBranches()
 		{
-			CrashReportDataContext Context = new CrashReportDataContext();
-
-			var BranchList = Context.Crashes
-				.Where( n => n.TimeOfCrash > DateTime.Now.AddMonths( -3 ) )
-				.Where( n => n.Branch.StartsWith( "UE4" ) )
-				.Select( n => n.Branch )
-				.Distinct()
-				.ToList();
-			var SelectListItens = BranchList
-				.Select( listitem => new SelectListItem { Selected = false, Text = listitem, Value = listitem } )
-				.ToList();
-			SelectListItens.Insert( 0, new SelectListItem { Selected = true, Text = "", Value = "" } );
-			return SelectListItens;
+			using (FAutoScopedLogTimer LogTimer = new FAutoScopedLogTimer( "CrashRepository.GetBranches" ))
+			{
+				CrashReportDataContext Context = new CrashReportDataContext();
+
+				DateTime Now = DateTime.UtcNow;
+
+				if (Now - LastBranchesDate > TimeSpan.FromHours( 1 ))
+				{
+					var BranchList = Context.Crashes
+					.Where( n => n.TimeOfCrash > DateTime.Now.AddMonths( -3 ) )
+					.Where( n => n.Branch.StartsWith( "UE4" ) )
+					.Select( n => n.Branch )
+					.Distinct()
+					.ToList();
+					BranchesAsSelectList = BranchList
+						.Select( listitem => new SelectListItem { Selected = false, Text = listitem, Value = listitem } )
+						.ToList();
+					BranchesAsSelectList.Insert( 0, new SelectListItem { Selected = true, Text = "", Value = "" } );
+
+					LastBranchesDate = Now;
+				}
+
+				return BranchesAsSelectList;
+			}
+		}
+
+
+		private static DateTime LastVersionDate = DateTime.UtcNow.AddDays( -1 );
+		private static List<SelectListItem> VersionsAsSelectList = null;
+		/// <summary>
+		/// Retrieves a list of distinct UE4 Versions from the CrashRepository
+		/// </summary>
+		public static List<SelectListItem> GetVersions()
+		{
+			using (FAutoScopedLogTimer LogTimer = new FAutoScopedLogTimer( "CrashRepository.GetVersions" ))
+			{
+				CrashReportDataContext Context = new CrashReportDataContext();
+				DateTime Now = DateTime.UtcNow;
+
+				if (Now - LastVersionDate > TimeSpan.FromHours( 1 ))
+				{
+					var BuildVersions = Context.Crashes
+					.Where( c => c.TimeOfCrash > DateTime.Now.AddMonths( -3 ) )
+					.Select( c => c.BuildVersion )
+					.Distinct()
+					.ToList();
+
+					HashSet<string> DistinctBuildVersions = new HashSet<string>();
+
+					foreach (var BuildVersion in BuildVersions)
+					{
+						var BVParts = BuildVersion.Split( new char[] { '.' }, StringSplitOptions.RemoveEmptyEntries );
+						if (BVParts.Length > 2 && BVParts[0] != "0")
+						{
+							string CleanBV = string.Format( "{0}.{1}.{2}", BVParts[0], BVParts[1], BVParts[2] );
+							DistinctBuildVersions.Add( CleanBV );
+						}
+					}
+
+
+					VersionsAsSelectList = DistinctBuildVersions
+						.Select( listitem => new SelectListItem { Selected = false, Text = listitem, Value = listitem } )
+						.ToList();
+					VersionsAsSelectList.Insert( 0, new SelectListItem { Selected = true, Text = "", Value = "" } );
+
+					LastVersionDate = Now;
+				}
+
+				return VersionsAsSelectList;
+			}
 		}
 		
 		/// <summary>
@@ -217,7 +276,7 @@ namespace Tools.CrashReporter.CrashReportWebSite.Models
 				int Skip = ( FormData.Page - 1 ) * FormData.PageSize;
 				int Take = FormData.PageSize;
 
-				var ResultsAll = ConstructQuery(FormData);
+				var ResultsAll = ConstructQueryForFiltering(FormData);
 
 				// Filter by data and get as enumerable.
 				Results = FilterByDate(ResultsAll, FormData.DateFrom, FormData.DateTo);
@@ -286,14 +345,13 @@ namespace Tools.CrashReporter.CrashReportWebSite.Models
 					CrashType = FormData.CrashType,
 					SearchQuery = FormData.SearchQuery,
 					UsernameQuery = FormData.UsernameQuery,
-					EpicIdQuery = FormData.EpicIdQuery,
-					MachineIdQuery = FormData.MachineIdQuery,
-					DescriptionQuery = FormData.DescriptionQuery,
+					EpicIdOrMachineQuery = FormData.EpicIdOrMachineQuery,
 					MessageQuery = FormData.MessageQuery,
 					JiraQuery = FormData.JiraQuery,
 					DateFrom = (long)( FormData.DateFrom - CrashesViewModel.Epoch ).TotalMilliseconds,
 					DateTo = (long)( FormData.DateTo - CrashesViewModel.Epoch ).TotalMilliseconds,
 					BranchName = FormData.BranchName,
+					VersionName = FormData.VersionName,
 					GameName = FormData.GameName,
 					GroupCounts = GroupCounts,
 					RealUserName = UniqueUser != null ? UniqueUser.ToString() : null,
@@ -610,24 +668,27 @@ namespace Tools.CrashReporter.CrashReportWebSite.Models
 			return NewID;
 		}
 
-		private IQueryable<Crash> ConstructQuery(FormHelper FormData)
+		/// <summary>Constructs query for filtering.</summary>
+		private IQueryable<Crash> ConstructQueryForFiltering( FormHelper FormData )
 		{
-			var results = ListAll();
+			var Results = ListAll();
 
 			// Grab Results 
-			if (!string.IsNullOrEmpty(FormData.SearchQuery))
+			if (!string.IsNullOrEmpty( FormData.SearchQuery ))
 			{
-				if (!string.IsNullOrEmpty(FormData.UsernameQuery))
+				if (!string.IsNullOrEmpty( FormData.UsernameQuery ))
 				{
-					//We only use SearchQuery now for CallStack searching - if there's a searchquery value and a Username value, we need to get rid of the 
+					//We only use SearchQuery now for CallStack searching - if there's a SearchQuery value and a Username value, we need to get rid of the 
 					//Username so that we can create a broader search range
 					FormData.UsernameQuery = "";
 				}
-				string DecodedQuery = HttpUtility.HtmlDecode(FormData.SearchQuery).ToLower();
-				using (FScopedLogTimer LogTimer2 = new FScopedLogTimer("CrashRepository.GetResults.FindUserFromQuery" + "(Query=" + DecodedQuery + ")"))
-				{
-					results = results.Where(item => item.RawCallStack.Contains(FormData.SearchQuery));
-				}
+
+				Results =
+					(
+						from CrashDetail in Results
+						where SqlMethods.Like( CrashDetail.RawCallStack, "%" + FormData.SearchQuery + "%" )
+						select CrashDetail
+					);
 			}
 
 			// Filter by Crash Type
@@ -636,149 +697,119 @@ namespace Tools.CrashReporter.CrashReportWebSite.Models
 				switch (FormData.CrashType)
 				{
 					case "Crashes":
-						results = results.Where(CrashInstance => CrashInstance.CrashType == 1);
+						Results = Results.Where( CrashInstance => CrashInstance.CrashType == 1 );
 						break;
 					case "Assert":
-						results = results.Where(CrashInstance => CrashInstance.CrashType == 2);
+						Results = Results.Where( CrashInstance => CrashInstance.CrashType == 2 );
 						break;
 					case "Ensure":
-						results = results.Where(CrashInstance => CrashInstance.CrashType == 3);
+						Results = Results.Where( CrashInstance => CrashInstance.CrashType == 3 );
 						break;
 					case "CrashesAsserts":
-						results = results.Where(CrashInstance => CrashInstance.CrashType == 1 || CrashInstance.CrashType == 2);
+						Results = Results.Where( CrashInstance => CrashInstance.CrashType == 1 || CrashInstance.CrashType == 2 );
 						break;
 				}
 			}
 
-			if (!string.IsNullOrEmpty(FormData.UsernameQuery))
+			// JRX Restore EpicID/UserName searching
+			if (!string.IsNullOrEmpty( FormData.UsernameQuery ))
 			{
-				results =
+				Results =
 					(
-						from CrashDetail in results
-						where CrashDetail.UserName.Equals(FormData.UsernameQuery)
+						from CrashDetail in Results
+						where CrashDetail.UserName.Equals( FormData.UsernameQuery )
 						select CrashDetail
 						);
 			}
 
 			// Start Filtering the results
-			if (!string.IsNullOrEmpty(FormData.EpicIdQuery))
+			if (!string.IsNullOrEmpty( FormData.EpicIdOrMachineQuery ))
 			{
-				var DecodedQuery = HttpUtility.HtmlDecode(FormData.EpicIdQuery).ToLower();
-				results =
+				var DecodedQuery = HttpUtility.HtmlDecode( FormData.EpicIdOrMachineQuery ).ToLower();
+				Results =
 						(
-							from CrashDetail in results
-							where CrashDetail.EpicAccountId.Equals(FormData.EpicIdQuery)
+							from CrashDetail in Results
+							where CrashDetail.EpicAccountId.Equals( FormData.EpicIdOrMachineQuery ) || CrashDetail.MachineID.Equals( FormData.EpicIdOrMachineQuery )
 							select CrashDetail
 						);
 			}
 
-			if (!string.IsNullOrEmpty(FormData.MachineIdQuery))
+			if (!string.IsNullOrEmpty( FormData.JiraQuery ))
 			{
-				var DecodedQuery = HttpUtility.HtmlDecode(FormData.MachineIdQuery).ToLower();
-				results =
-						(
-							from CrashDetail in results
-							where CrashDetail.ComputerName.Equals(FormData.MachineIdQuery)
-							select CrashDetail
-						);
-			}
-
-			if (!string.IsNullOrEmpty(FormData.JiraQuery))
-			{
-				var DecodedQuery = HttpUtility.HtmlDecode(FormData.JiraQuery).ToLower();
-				results =
+				var DecodedQuery = HttpUtility.HtmlDecode( FormData.JiraQuery ).ToLower();
+				Results =
 					(
-						from CrashDetail in results
-						where CrashDetail.TTPID.Equals(FormData.JiraQuery)
+						from CrashDetail in Results
+						where CrashDetail.TTPID.Equals( FormData.JiraQuery )
 						select CrashDetail
 					);
 			}
+
 			// Filter by BranchName
-			if (!string.IsNullOrEmpty(FormData.BranchName))
+			if (!string.IsNullOrEmpty( FormData.BranchName ))
 			{
-				if (FormData.BranchName.StartsWith("-"))
-				{
-					results =
+				Results =
 					(
-						from CrashDetail in results
-						where !CrashDetail.Branch.Contains(FormData.BranchName.Substring(1))
+						from CrashDetail in Results
+						where CrashDetail.Branch.Equals( FormData.BranchName )
 						select CrashDetail
 					);
-				}
-				else
-				{
-					results =
-					(
-						from CrashDetail in results
-						where CrashDetail.Branch.Equals(FormData.BranchName)
-						select CrashDetail
-					);
-				}
 			}
 			else
 			{
-				results = 
-					(
-					from CrashDetail in results
-					where !CrashDetail.Branch.Equals("UE4-UT-Releases") && !CrashDetail.Branch.Equals("UE4-UT")
+				Results =
+				(
+					from CrashDetail in Results
+					where !CrashDetail.Branch.Contains( "UE4-UT" )
 					select CrashDetail
+				);
+			}
+
+			// Filter by VersionName
+			if (!string.IsNullOrEmpty( FormData.VersionName ))
+			{
+				Results =
+					(
+						from CrashDetail in Results
+						where CrashDetail.BuildVersion.Contains( FormData.VersionName )
+						select CrashDetail
 					);
 			}
 
 			// Filter by GameName
-			if (!string.IsNullOrEmpty(FormData.GameName))
+			if (!string.IsNullOrEmpty( FormData.GameName ))
 			{
-				if (FormData.GameName.StartsWith("-"))
+				if (FormData.GameName.StartsWith( "-" ))
 				{
-					results =
+					Results =
 					(
-						from CrashDetail in results
-						where !CrashDetail.GameName.Contains(FormData.GameName.Substring(1))
+						from CrashDetail in Results
+						where !CrashDetail.GameName.Contains( FormData.GameName.Substring( 1 ) )
 						select CrashDetail
 					);
 				}
 				else
 				{
-					results =
+					Results =
 					(
-						from CrashDetail in results
-						where CrashDetail.GameName.Contains(FormData.GameName)
+						from CrashDetail in Results
+						where CrashDetail.GameName.Contains( FormData.GameName )
 						select CrashDetail
 					);
 				}
 			}
 
-			if (!string.IsNullOrEmpty(FormData.MessageQuery))
+			if (!string.IsNullOrEmpty( FormData.MessageQuery ))
 			{
-				results =
+				Results =
 					(
-						from CrashDetail in results
-						where SqlMethods.Like(CrashDetail.Summary, "%" + FormData.MessageQuery + "%")
+						from CrashDetail in Results
+						where SqlMethods.Like( CrashDetail.Summary, "%" + FormData.MessageQuery + "%" ) || SqlMethods.Like( CrashDetail.Description, "%" + FormData.MessageQuery + "%" )
 						select CrashDetail
 						);
 			}
 
-			if (!string.IsNullOrEmpty(FormData.DescriptionQuery))
-			{
-				results = 
-					(
-						from CrashDetail in results
-						where SqlMethods.Like(CrashDetail.Description, "%" + FormData.DescriptionQuery + "%")
-						select CrashDetail
-					);
-			}
-
-			if (!string.IsNullOrEmpty(FormData.SearchQuery))
-			{
-				results = 
-					(
-						from CrashDetail in results
-						where SqlMethods.Like(CrashDetail.RawCallStack, "%" + FormData.SearchQuery + "%")
-						select CrashDetail
-					);
-			}
-
-			return results;
+			return Results;
 		}
 	}
 }
