CrashReportWWW - Fixed filtering by MachineID

diff --git a/Engine/Source/Programs/CrashReporter/CrashReportWebSite/Models/CrashRepository.cs b/Engine/Source/Programs/CrashReporter/CrashReportWebSite/Models/CrashRepository.cs
--- a/Engine/Source/Programs/CrashReporter/CrashReportWebSite/Models/CrashRepository.cs
+++ b/Engine/Source/Programs/CrashReporter/CrashReportWebSite/Models/CrashRepository.cs
@@ -714,10 +714,11 @@ namespace Tools.CrashReporter.CrashReportWebSite.Models
 			// JRX Restore EpicID/UserName searching
 			if (!string.IsNullOrEmpty( FormData.UsernameQuery ))
 			{
+				var DecodedUsername = HttpUtility.HtmlDecode( FormData.UsernameQuery ).ToLower();
 				Results =
 					(
 						from CrashDetail in Results
-						where CrashDetail.UserName.Equals( FormData.UsernameQuery )
+						where CrashDetail.UserName.Equals( DecodedUsername )
 						select CrashDetail
 						);
 			}
@@ -725,22 +726,22 @@ namespace Tools.CrashReporter.CrashReportWebSite.Models
 			// Start Filtering the results
 			if (!string.IsNullOrEmpty( FormData.EpicIdOrMachineQuery ))
 			{
-				var DecodedQuery = HttpUtility.HtmlDecode( FormData.EpicIdOrMachineQuery ).ToLower();
+				var DecodedEpicOrMachineID = HttpUtility.HtmlDecode( FormData.EpicIdOrMachineQuery ).ToLower();
 				Results =
 						(
 							from CrashDetail in Results
-							where CrashDetail.EpicAccountId.Equals( FormData.EpicIdOrMachineQuery ) || CrashDetail.MachineID.Equals( FormData.EpicIdOrMachineQuery )
+							where CrashDetail.EpicAccountId.Equals( DecodedEpicOrMachineID ) || CrashDetail.ComputerName.Equals( DecodedEpicOrMachineID )
 							select CrashDetail
 						);
 			}
 
 			if (!string.IsNullOrEmpty( FormData.JiraQuery ))
 			{
-				var DecodedQuery = HttpUtility.HtmlDecode( FormData.JiraQuery ).ToLower();
+				var DecodedJIRA = HttpUtility.HtmlDecode( FormData.JiraQuery ).ToLower();
 				Results =
 					(
 						from CrashDetail in Results
-						where CrashDetail.TTPID.Equals( FormData.JiraQuery )
+						where CrashDetail.TTPID.Contains( DecodedJIRA )
 						select CrashDetail
 					);
 			}
@@ -771,7 +772,7 @@ namespace Tools.CrashReporter.CrashReportWebSite.Models
 				Results =
 					(
 						from CrashDetail in Results
-						where CrashDetail.BuildVersion.Contains( FormData.VersionName )
+						where CrashDetail.BuildVersion.Equals( FormData.VersionName )
 						select CrashDetail
 					);
 			}
@@ -779,12 +780,14 @@ namespace Tools.CrashReporter.CrashReportWebSite.Models
 			// Filter by GameName
 			if (!string.IsNullOrEmpty( FormData.GameName ))
 			{
-				if (FormData.GameName.StartsWith( "-" ))
+				var DecodedGameName = HttpUtility.HtmlDecode( FormData.GameName ).ToLower();
+
+				if (DecodedGameName.StartsWith( "-" ))
 				{
 					Results =
 					(
 						from CrashDetail in Results
-						where !CrashDetail.GameName.Contains( FormData.GameName.Substring( 1 ) )
+						where !CrashDetail.GameName.Contains( DecodedGameName.Substring( 1 ) )
 						select CrashDetail
 					);
 				}
@@ -793,7 +796,7 @@ namespace Tools.CrashReporter.CrashReportWebSite.Models
 					Results =
 					(
 						from CrashDetail in Results
-						where CrashDetail.GameName.Contains( FormData.GameName )
+						where CrashDetail.GameName.Contains( DecodedGameName )
 						select CrashDetail
 					);
 				}
