CrashReportWWW - Removed unused database properties from the layout, checked against the file system for existence

diff --git a/Engine/Source/Programs/CrashReporter/CrashReportWebSite/Controllers/BuggsController.cs b/Engine/Source/Programs/CrashReporter/CrashReportWebSite/Controllers/BuggsController.cs
--- a/Engine/Source/Programs/CrashReporter/CrashReportWebSite/Controllers/BuggsController.cs
+++ b/Engine/Source/Programs/CrashReporter/CrashReportWebSite/Controllers/BuggsController.cs
@@ -22,6 +22,21 @@ namespace Tools.CrashReporter.CrashReportWebSite.Controllers
 		{
 		}
 
+		/*
+		/// <summary>
+		/// The empty view of the buggs page.
+		/// </summary>
+		/// 
+		public ActionResult Index()
+		{
+			using (FAutoScopedLogTimer LogTimer = new FAutoScopedLogTimer( this.GetType().ToString(), bCreateNewLog: true ))
+			{
+				BuggsViewModel Results = new BuggsViewModel();
+				Results.GenerationTime = LogTimer.GetElapsedSeconds().ToString( "F2" );
+				return View( "Index", Results );
+			}
+		}*/
+
 		/// <summary>
 		/// The Index action.
 		/// </summary>
@@ -29,13 +44,13 @@ namespace Tools.CrashReporter.CrashReportWebSite.Controllers
 		/// <returns>The view to display a list of Buggs on the client.</returns>
 		public ActionResult Index( FormCollection BuggsForm )
 		{
-			using( FAutoScopedLogTimer LogTimer = new FAutoScopedLogTimer( this.GetType().ToString(), bCreateNewLog: true ) )
+			using (FAutoScopedLogTimer LogTimer = new FAutoScopedLogTimer( this.GetType().ToString(), bCreateNewLog: true ))
 			{
 				BuggRepository Buggs = new BuggRepository();
 
 				FormHelper FormData = new FormHelper( Request, BuggsForm, "CrashesInTimeFrameGroup" );
 				BuggsViewModel Results = Buggs.GetResults( FormData );
-				foreach( var Bugg in Results.Results )
+				foreach (var Bugg in Results.Results)
 				{
 					// Populate function calls.
 					Bugg.GetFunctionCalls();
@@ -128,10 +143,10 @@ namespace Tools.CrashReporter.CrashReportWebSite.Controllers
 					bool bValidJira = false;
 
 					// Verify valid JiraID, this may be still a TTP 
-					if( !string.IsNullOrEmpty( NewBugg.TTPID ) )
+					if( !string.IsNullOrEmpty( NewBugg.Jira ) )
 					{
 						int TTPID = 0;
-						int.TryParse( NewBugg.TTPID, out TTPID );
+						int.TryParse( NewBugg.Jira, out TTPID );
 
 						if( TTPID == 0 )
 						{
@@ -143,7 +158,7 @@ namespace Tools.CrashReporter.CrashReportWebSite.Controllers
 					if( JC.CanBeUsed() && bValidJira )
 					{
 						// Grab the data form JIRA.
-						string JiraSearchQuery = "key = " + NewBugg.TTPID;
+						string JiraSearchQuery = "key = " + NewBugg.Jira;
 
 						var JiraResults = JC.SearchJiraTickets(
 							JiraSearchQuery,
@@ -217,8 +232,8 @@ namespace Tools.CrashReporter.CrashReportWebSite.Controllers
 
 					if( !string.IsNullOrEmpty( BuggsForm["SetTTP"] ) )
 					{
-						NewBugg.TTPID = BuggsForm["SetTTP"];
-						Buggs.SetJIRAForBuggAndCrashes( NewBugg.TTPID, id );
+						NewBugg.Jira = BuggsForm["SetTTP"];
+						Buggs.SetJIRAForBuggAndCrashes( NewBugg.Jira, id );
 					}
 
 					// <STATUS>
diff --git a/Engine/Source/Programs/CrashReporter/CrashReportWebSite/Controllers/CrashesController.cs b/Engine/Source/Programs/CrashReporter/CrashReportWebSite/Controllers/CrashesController.cs
--- a/Engine/Source/Programs/CrashReporter/CrashReportWebSite/Controllers/CrashesController.cs
+++ b/Engine/Source/Programs/CrashReporter/CrashReportWebSite/Controllers/CrashesController.cs
@@ -56,7 +56,7 @@ namespace Tools.CrashReporter.CrashReportWebSite.Controllers
 
 							if( !string.IsNullOrEmpty( CrashesForm["SetTTP"] ) )
 							{
-								CurrentCrash.TTPID = CrashesForm["SetTTP"];
+								CurrentCrash.Jira = CrashesForm["SetTTP"];
 							}
 						}
 					}
@@ -116,7 +116,7 @@ namespace Tools.CrashReporter.CrashReportWebSite.Controllers
 				FormValue = CrashesForm["SetTTP"];
 				if( !string.IsNullOrEmpty( FormValue ) )
 				{
-					CurrentCrash.TTPID = FormValue;
+					CurrentCrash.Jira = FormValue;
 				}
 
 				// Valid to set description to an empty string
diff --git a/Engine/Source/Programs/CrashReporter/CrashReportWebSite/Controllers/ReportsController.cs b/Engine/Source/Programs/CrashReporter/CrashReportWebSite/Controllers/ReportsController.cs
--- a/Engine/Source/Programs/CrashReporter/CrashReportWebSite/Controllers/ReportsController.cs
+++ b/Engine/Source/Programs/CrashReporter/CrashReportWebSite/Controllers/ReportsController.cs
@@ -16,6 +16,15 @@ using Tools.DotNETCommon;
 namespace Tools.CrashReporter.CrashReportWebSite.Controllers
 {
 
+	//How hard would it be for you to create a CSV of:
+	// Epic ID, 
+	// buggs ID, 
+	// engine version, 
+	// 
+	// # of crashes with that buggs ID for 
+	//	this engine version and 
+	//	user
+
 	/// <summary>
 	/// The controller to handle graphing of crashes per user group over time.
 	/// </summary>
@@ -36,6 +45,8 @@ namespace Tools.CrashReporter.CrashReportWebSite.Controllers
 			CrashRepository CrashRepo = new CrashRepository();
 
 			// @TODO yrx 2015-02-17 BuggIDToBeAddedToJira replace with List<int> based on check box and Submit?
+			// It would be great to have a CSV export of this as well with buggs ID being the key I can then use to join them :)
+			// 
 			// Enumerate JIRA projects if needed.
 			// https://jira.ol.epicgames.net//rest/api/2/project
 			var JC = JiraConnection.Get();
@@ -63,12 +74,12 @@ namespace Tools.CrashReporter.CrashReportWebSite.Controllers
 						TimeOfCrash = Crash.TimeOfCrash.Value,
 						//UserID = Crash.UserNameId.Value, 
 						BuildVersion = Crash.BuildVersion,
-						JIRA = Crash.TTPID,
+						JIRA = Crash.Jira,
 						Platform = Crash.PlatformName,
 						FixCL = Crash.FixedChangeList,
-						BuiltFromCL = Crash.ChangeListVersion,
+						BuiltFromCL = Crash.BuiltFromCL,
 						Pattern = Crash.Pattern,
-						MachineID = Crash.ComputerName,
+						MachineID = Crash.MachineId,
 						Branch = Crash.Branch,
 						Description = Crash.Description,
 						RawCallStack = Crash.RawCallStack,
@@ -158,12 +169,12 @@ namespace Tools.CrashReporter.CrashReportWebSite.Controllers
 									ID = Anon.ID,
 									TimeOfCrash = Anon.TimeOfCrash,
 									BuildVersion = Anon.BuildVersion,
-									JIRA = Anon.JIRA,
+									Jira = Anon.JIRA,
 									Platform = Anon.Platform,
 									FixCL = Anon.FixCL,
 									BuiltFromCL = Anon.BuiltFromCL,
 									Pattern = Anon.Pattern,
-									MachineID = Anon.MachineID,
+									MachineId = Anon.MachineID,
 									Branch = Anon.Branch,
 									Description = Anon.Description,
 									RawCallStack = Anon.RawCallStack,
@@ -173,10 +184,10 @@ namespace Tools.CrashReporter.CrashReportWebSite.Controllers
 							NewBugg.PrepareBuggForJira( FullCrashesForBugg );
 
 							// Verify valid JiraID, this may be still a TTP 
-							if( !string.IsNullOrEmpty( NewBugg.TTPID ) )
+							if( !string.IsNullOrEmpty( NewBugg.Jira ) )
 							{
 								int TTPID = 0;
-								int.TryParse( NewBugg.TTPID, out TTPID );
+								int.TryParse( NewBugg.Jira, out TTPID );
 
 								if( TTPID == 0 )
 								{
@@ -279,7 +290,7 @@ namespace Tools.CrashReporter.CrashReportWebSite.Controllers
 
 					// If there are buggs, we need to update the summary to indicate an error.
 					// Usually caused when bugg's project has changed.
-					foreach( var Bugg in BuggsCopy.Where( b => !string.IsNullOrEmpty( b.TTPID ) ) )
+					foreach( var Bugg in BuggsCopy.Where( b => !string.IsNullOrEmpty( b.Jira ) ) )
 					{
 						Bugg.JiraSummary = "JIRA MISMATCH";
 						Bugg.JiraComponentsText = "JIRA MISMATCH";
@@ -306,7 +317,7 @@ namespace Tools.CrashReporter.CrashReportWebSite.Controllers
 
 		private void AddBuggJiraMapping( Bugg NewBugg, ref HashSet<string> FoundJiras, ref Dictionary<string, List<Bugg>> JiraIDtoBugg )
 		{
-			string JiraID = NewBugg.TTPID;
+			string JiraID = NewBugg.Jira;
 			FoundJiras.Add( "key = " + JiraID );
 
 			bool bAdd = !JiraIDtoBugg.ContainsKey( JiraID );
diff --git a/Engine/Source/Programs/CrashReporter/CrashReportWebSite/Models/BuggRepository.cs b/Engine/Source/Programs/CrashReporter/CrashReportWebSite/Models/BuggRepository.cs
--- a/Engine/Source/Programs/CrashReporter/CrashReportWebSite/Models/BuggRepository.cs
+++ b/Engine/Source/Programs/CrashReporter/CrashReportWebSite/Models/BuggRepository.cs
@@ -226,7 +226,7 @@ namespace Tools.CrashReporter.CrashReportWebSite.Models
 						}
 
 						CurrentCrash.FixedChangeList = Bugg.FixedChangeList;
-						CurrentCrash.TTPID = Bugg.TTPID;
+						CurrentCrash.Jira = Bugg.Jira;
 						CurrentCrash.Status = Bugg.Status;
 
 						if( Bugg.Id != 0 )
@@ -692,7 +692,7 @@ namespace Tools.CrashReporter.CrashReportWebSite.Models
 						(
 							from Crash in Context.Crashes
 							where Crash.TimeOfCrash >= DateFrom && Crash.TimeOfCrash <= DateTo.AddDays( 1 )
-							select new { Id = Crash.Id, UserName = Crash.UserName, MachineId = Crash.ComputerName }
+							select new { Id = Crash.Id, UserName = Crash.UserName, MachineId = Crash.MachineId }
 						);
 
 						CrashToUser = CrashesWithIdUserMachine.ToDictionary( x => x.Id, y => y.UserName );
@@ -834,7 +834,7 @@ namespace Tools.CrashReporter.CrashReportWebSite.Models
 							break;
 
 						case "TTPID":
-							Results = EnumerableOrderBy( Results, BuggCrashInstance => BuggCrashInstance.TTPID, bSortDescending );
+							Results = EnumerableOrderBy( Results, BuggCrashInstance => BuggCrashInstance.Jira, bSortDescending );
 							break;
 
 					}
diff --git a/Engine/Source/Programs/CrashReporter/CrashReportWebSite/Models/CrashReport.cs b/Engine/Source/Programs/CrashReporter/CrashReportWebSite/Models/CrashReport.cs
--- a/Engine/Source/Programs/CrashReporter/CrashReportWebSite/Models/CrashReport.cs
+++ b/Engine/Source/Programs/CrashReporter/CrashReportWebSite/Models/CrashReport.cs
@@ -164,9 +164,9 @@ namespace Tools.CrashReporter.CrashReportWebSite.Models
 			foreach( var Crash in CrashesForBugg )
 			{
 				// Only add machine if the number has 32 characters
-				if( Crash.MachineID != null && Crash.MachineID.Length == 32 )
+				if (Crash.MachineId != null && Crash.MachineId.Length == 32)
 				{
-					MachineIds.Add( Crash.MachineID );
+					MachineIds.Add( Crash.MachineId );
 
 					// Sent in the unattended mode 29 char
 					if( Crash.Description.Length > 32 )
@@ -298,7 +298,7 @@ namespace Tools.CrashReporter.CrashReportWebSite.Models
 		public void CopyToJira()
 		{
 			var JC = JiraConnection.Get();
-			if( JC.CanBeUsed() && string.IsNullOrEmpty( this.TTPID ) )
+			if( JC.CanBeUsed() && string.IsNullOrEmpty( this.Jira ) )
 			{
 				Dictionary<string, object> Fields = new Dictionary<string, object>();
 				Fields.Add( "project", new Dictionary<string, object> { { "id", 11205 } } );	// UE
@@ -332,7 +332,7 @@ namespace Tools.CrashReporter.CrashReportWebSite.Models
 				string Key = JC.AddJiraTicket( Fields );
 				if( !string.IsNullOrEmpty( Key ) )
 				{
-					TTPID = Key;
+					Jira = Key;
 					BuggRepository Buggs = new BuggRepository();
 					Buggs.SetJIRAForBuggAndCrashes( Key, Id );
 				}
@@ -533,7 +533,6 @@ namespace Tools.CrashReporter.CrashReportWebSite.Models
 		/// <summary>
 		/// Return the Url of the log.
 		/// </summary>
-		/// <returns>The Url of the log file.</returns>
 		public string GetLogUrl()
 		{
 			return Properties.Settings.Default.CrashReporterFiles + Id + "_Launch.log";
@@ -542,8 +541,7 @@ namespace Tools.CrashReporter.CrashReportWebSite.Models
 		/// <summary>
 		/// Return the Url of the minidump.
 		/// </summary>
-		/// <returns>The Url of the minidump file.</returns>
-		public string GetDumpUrl()
+		public string GetMiniDumpUrl()
 		{
 			return Properties.Settings.Default.CrashReporterFiles + Id + "_MiniDump.dmp";
 		}
@@ -551,7 +549,6 @@ namespace Tools.CrashReporter.CrashReportWebSite.Models
 		/// <summary>
 		/// Return the Url of the diagnostics file.
 		/// </summary>
-		/// <returns>The Url of the diagnostics file.</returns>
 		public string GetDiagnosticsUrl()
 		{
 			return Properties.Settings.Default.CrashReporterFiles + Id + "_Diagnostics.txt";
@@ -560,7 +557,6 @@ namespace Tools.CrashReporter.CrashReportWebSite.Models
 		/// <summary>
 		/// Return the Url of the Windows Error Report meta data file.
 		/// </summary>
-		/// <returns>The Url of the meta data file.</returns>
 		public string GetMetaDataUrl()
 		{
 			return Properties.Settings.Default.CrashReporterFiles + Id + "_WERMeta.xml";
@@ -762,78 +758,71 @@ namespace Tools.CrashReporter.CrashReportWebSite.Models
 			}
 		}
 
-		/// <summary>
-		/// 
-		/// </summary>
-		public string JIRA
+		/// <summary>Hard coded site path.</summary>
+		const string SitePath = @"\\devweb-02\Sites";
+
+		/// <summary>Return true, if there is a diagnostics file associated with the crash</summary>
+		public bool HasDiagnosticsFile()
 		{
-			get
-			{
-				return this.TTPID;
-			}
-			set
-			{
-				this.TTPID = value;
-			}
+			var Path = SitePath + GetDiagnosticsUrl();
+			return System.IO.File.Exists( Path );
 		}
 
-		/// <summary>
-		/// 
-		/// </summary>
-		public string Platform
+		/// <summary>Return true, if there is a minidump file associated with the crash</summary>
+		public bool HasMiniDumpFile()
 		{
-			get
-			{
-				return this.PlatformName;
-			}
-			set
-			{
-				this.PlatformName = value;
-			}
+			var Path = SitePath + GetMiniDumpUrl();
+			return System.IO.File.Exists( Path );
 		}
 
-		/// <summary>
-		/// 
-		/// </summary>
-		public string FixCL
+		/// <summary>Return true, if there is a log file associated with the crash</summary>
+		public bool HasLogFile()
 		{
-			get
-			{
-				return this.FixedChangeList;
-			}
-			set
-			{
-				this.FixedChangeList = value;
-			}
+			var Path = SitePath + GetLogUrl();
+			return System.IO.File.Exists( Path );
+		}
+
+		/// <summary>Return true, if there is a video file associated with the crash</summary>
+		public bool HasVideoFile()
+		{
+			var Path = SitePath + GetVideoUrl();
+			return System.IO.File.Exists( Path );
+		}
+
+		/// <summary>Return true, if there is a metadata file associated with the crash</summary>
+		public bool HasMetaDataFile()
+		{
+			var Path = SitePath + GetMetaDataUrl();
+			return System.IO.File.Exists( Path );
 		}
 
 		/// <summary>
 		/// 
 		/// </summary>
-		public string MachineID
+		public string Platform
 		{
 			get
 			{
-				return this.ComputerName;
+				return this.PlatformName;
 			}
 			set
 			{
-				this.ComputerName = value;
+				this.PlatformName = value;
 			}
 		}
 
 		/// <summary>
 		/// 
 		/// </summary>
-		public string BuiltFromCL
+		public string FixCL
 		{
 			get
 			{
-				return this.ChangeListVersion;
+				return this.FixedChangeList;
 			}
 			set
 			{
-				this.ChangeListVersion = value;
+				this.FixedChangeList = value;
 			}
 		}
 	}
diff --git a/Engine/Source/Programs/CrashReporter/CrashReportWebSite/Models/CrashReport.dbml b/Engine/Source/Programs/CrashReporter/CrashReportWebSite/Models/CrashReport.dbml
--- a/Engine/Source/Programs/CrashReporter/CrashReportWebSite/Models/CrashReport.dbml
+++ b/Engine/Source/Programs/CrashReporter/CrashReportWebSite/Models/CrashReport.dbml
@@ -3,7 +3,7 @@
   <Table Name="dbo.Buggs" Member="Buggs">
     <Type Name="Bugg">
       <Column Name="Id" Type="System.Int32" DbType="Int NOT NULL IDENTITY" IsPrimaryKey="true" IsDbGenerated="true" CanBeNull="false" />
-      <Column Name="TTPID" Type="System.String" DbType="VarChar(50)" CanBeNull="true" />
+      <Column Name="TTPID" Member="Jira" Storage="_TTPID" Type="System.String" DbType="VarChar(50)" CanBeNull="true" />
       <Column Name="Pattern" Type="System.String" DbType="VarChar(800) NOT NULL" CanBeNull="false" />
       <Column Name="NumberOfCrashes" Type="System.Int32" DbType="Int" CanBeNull="true" />
       <Column Name="NumberOfUsers" Type="System.Int32" DbType="Int" CanBeNull="true" />
@@ -56,30 +56,25 @@
       <Column Name="GameName" Type="System.String" DbType="VarChar(64)" CanBeNull="true" />
       <Column Name="Status" Type="System.String" DbType="VarChar(64)" CanBeNull="true" />
       <Column Name="TimeOfCrash" Type="System.DateTime" DbType="DateTime" CanBeNull="true" />
-      <Column Name="ChangeListVersion" Type="System.String" DbType="VarChar(64)" CanBeNull="true" />
+      <Column Name="ChangeListVersion" Member="BuiltFromCL" Storage="_ChangeListVersion" Type="System.String" DbType="VarChar(64)" CanBeNull="true" />
       <Column Name="PlatformName" Type="System.String" DbType="VarChar(64)" CanBeNull="true" />
       <Column Name="EngineMode" Type="System.String" DbType="VarChar(64)" CanBeNull="true" />
       <Column Name="Description" Type="System.String" DbType="VarChar(512)" CanBeNull="true" />
       <Column Name="RawCallStack" Type="System.String" DbType="VarChar(MAX)" CanBeNull="true" />
       <Column Name="Pattern" Type="System.String" DbType="VarChar(800)" CanBeNull="true" />
       <Column Name="CommandLine" Type="System.String" DbType="VarChar(512)" CanBeNull="true" />
-      <Column Name="ComputerName" Type="System.String" DbType="VarChar(64)" CanBeNull="true" />
+      <Column Name="ComputerName" Member="MachineId" Storage="_ComputerName" Type="System.String" DbType="VarChar(64)" CanBeNull="true" />
       <Column Name="FixedChangeList" Type="System.String" DbType="VarChar(256)" CanBeNull="true" />
       <Column Name="LanguageExt" Type="System.String" DbType="VarChar(64)" CanBeNull="true" />
       <Column Name="Module" Type="System.String" DbType="VarChar(128)" CanBeNull="true" />
       <Column Name="BuildVersion" Type="System.String" DbType="VarChar(128)" CanBeNull="true" />
       <Column Name="BaseDir" Type="System.String" DbType="VarChar(512)" CanBeNull="true" />
       <Column Name="UserName" Type="System.String" DbType="VarChar(64)" CanBeNull="true" />
-      <Column Name="TTPID" Type="System.String" DbType="VarChar(64)" CanBeNull="true" />
+      <Column Name="TTPID" Member="Jira" Storage="_TTPID" Type="System.String" DbType="VarChar(64)" CanBeNull="true" />
       <Column Name="Processed" Member="bAllowToBeContacted" Storage="_Processed" Type="System.Boolean" DbType="Bit" CanBeNull="true" />
-      <Column Name="HasLogFile" Type="System.Boolean" DbType="Bit" CanBeNull="true" />
-      <Column Name="HasMiniDumpFile" Type="System.Boolean" DbType="Bit" CanBeNull="true" />
-      <Column Name="HasVideoFile" Type="System.Boolean" DbType="Bit" CanBeNull="true" />
-      <Column Name="HasDiagnosticsFile" Type="System.Boolean" DbType="Bit" CanBeNull="true" />
       <Column Name="Branch" Type="System.String" DbType="VarChar(32)" CanBeNull="true" />
       <Column Name="UserNameId" Type="System.Int32" DbType="Int" CanBeNull="true" />
       <Column Name="CrashType" Type="System.Int16" DbType="SmallInt" CanBeNull="true" />
-      <Column Name="HasMetaData" Type="System.Boolean" DbType="Bit" CanBeNull="true" />
       <Column Name="SourceContext" Type="System.String" DbType="VarChar(MAX)" CanBeNull="true" />
       <Column Name="EpicAccountId" Type="System.String" DbType="varchar(64)" CanBeNull="true" />
       <Column Name="EngineVersion" Type="System.String" DbType="varchar(64)" CanBeNull="true" />
diff --git a/Engine/Source/Programs/CrashReporter/CrashReportWebSite/Models/CrashReport.dbml.layout b/Engine/Source/Programs/CrashReporter/CrashReportWebSite/Models/CrashReport.dbml.layout
--- a/Engine/Source/Programs/CrashReporter/CrashReportWebSite/Models/CrashReport.dbml.layout
+++ b/Engine/Source/Programs/CrashReporter/CrashReportWebSite/Models/CrashReport.dbml.layout
@@ -32,10 +32,10 @@
         <elementListCompartment Id="77c436bb-7563-449c-a08a-8c417bcde3fa" absoluteBounds="3.3900000000000006, 6.3350000000000009, 1.9700000000000002, 0.8262939453125" name="DataPropertiesCompartment" titleTextColor="Black" itemTextColor="Black" />
       </nestedChildShapes>
     </classShape>
-    <classShape Id="f1204ef8-6660-48ce-8c1d-131c49e2dd25" absoluteBounds="6.125, 1.125, 2, 6.9630354817708326">
+    <classShape Id="f1204ef8-6660-48ce-8c1d-131c49e2dd25" absoluteBounds="6.125, 1.125, 2, 6.0015283203125">
       <DataClassMoniker Name="/CrashReportDataContext/Crash" />
       <nestedChildShapes>
-        <elementListCompartment Id="53c11757-f0be-47da-bcae-7d20a571002e" absoluteBounds="6.1400000000000006, 1.5850000000000009, 1.9700000000000002, 6.403035481770833" name="DataPropertiesCompartment" titleTextColor="Black" itemTextColor="Black" />
+        <elementListCompartment Id="53c11757-f0be-47da-bcae-7d20a571002e" absoluteBounds="6.1400000000000006, 1.5850000000000009, 1.9700000000000002, 5.4415283203125" name="DataPropertiesCompartment" titleTextColor="Black" itemTextColor="Black" />
       </nestedChildShapes>
     </classShape>
     <classShape Id="01d6823d-5d7e-435f-99a2-7d6fac1e698d" absoluteBounds="8.75, 4, 2, 1.193992513020838">
@@ -57,7 +57,7 @@
         <classShapeMoniker Id="b4a7205e-4223-43f1-837f-ac4a1390bf89" />
       </nodes>
     </associationConnector>
-    <associationConnector edgePoints="[(5.375 : 4.84084598496556); (6.125 : 4.84084598496556)]" fixedFrom="NotFixed" fixedTo="NotFixed">
+    <associationConnector edgePoints="[(5.375 : 4.56814697265625); (6.125 : 4.56814697265625)]" fixedFrom="NotFixed" fixedTo="NotFixed">
       <AssociationMoniker Name="/CrashReportDataContext/User/User_Crash" />
       <nodes>
         <classShapeMoniker Id="01dddf4f-058f-43b4-b6c0-b90e7acfe0b9" />
@@ -71,7 +71,7 @@
         <classShapeMoniker Id="8a60fc4a-a797-417f-87ed-25b22bd5e420" />
       </nodes>
     </associationConnector>
-    <associationConnector edgePoints="[(11.5 : 2.47199625651042); (10.875 : 2.47199625651042)]" fixedFrom="Algorithm" fixedTo="Algorithm">
+    <associationConnector edgePoints="[(11.5 : 2.47199625651042); (10.875 : 2.47199625651042)]" fixedFrom="NotFixed" fixedTo="NotFixed">
       <AssociationMoniker Name="/CrashReportDataContext/Bugg/Bugg_Buggs_Crash" />
       <nodes>
         <classShapeMoniker Id="a37f1cdb-93d7-48ef-8da7-2b7ca09282b0" />
diff --git a/Engine/Source/Programs/CrashReporter/CrashReportWebSite/Models/CrashRepository.cs b/Engine/Source/Programs/CrashReporter/CrashReportWebSite/Models/CrashRepository.cs
--- a/Engine/Source/Programs/CrashReporter/CrashReportWebSite/Models/CrashRepository.cs
+++ b/Engine/Source/Programs/CrashReporter/CrashReportWebSite/Models/CrashRepository.cs
@@ -222,18 +222,18 @@ namespace Tools.CrashReporter.CrashReportWebSite.Models
 					{
 						Results = Results.Where( X =>
 
-							( !string.IsNullOrEmpty( X.ChangeListVersion ) ? X.ChangeListVersion.ToLower().Contains( Term ) : false ) ||
+							( !string.IsNullOrEmpty( X.BuiltFromCL ) ? X.BuiltFromCL.ToLower().Contains( Term ) : false ) ||
 							( !string.IsNullOrEmpty( X.PlatformName ) ? X.PlatformName.ToLower().Contains( Term ) : false ) ||
 							( !string.IsNullOrEmpty( X.Summary ) ? X.Summary.ToLower().Contains( Term ) : false ) ||
 							( !string.IsNullOrEmpty( X.Description ) ? X.Description.ToLower().Contains( Term ) : false ) ||
 							( !string.IsNullOrEmpty( X.RawCallStack ) ? X.RawCallStack.ToLower().Contains( Term ) : false ) ||
 							( !string.IsNullOrEmpty( X.CommandLine ) ? X.CommandLine.ToLower().Contains( Term ) : false ) ||
-							( !string.IsNullOrEmpty( X.ComputerName ) ? X.ComputerName.ToLower().Contains( Term ) : false ) ||
+							( !string.IsNullOrEmpty( X.MachineId ) ? X.MachineId.ToLower().Contains( Term ) : false ) ||
 
 							( !string.IsNullOrEmpty( X.Module ) ? X.Module.ToLower().Contains( Term ) : false ) ||
 
 							( !string.IsNullOrEmpty( X.BaseDir ) ? X.BaseDir.ToLower().Contains( Term ) : false ) ||
-							( !string.IsNullOrEmpty( X.TTPID ) ? X.TTPID.ToLower().Contains( Term ) : false ) 
+							( !string.IsNullOrEmpty( X.Jira ) ? X.Jira.ToLower().Contains( Term ) : false ) 
 						);
 					}
 
@@ -486,7 +486,7 @@ namespace Tools.CrashReporter.CrashReportWebSite.Models
 						break;
 
 					case "TTPID":
-						Results = EnumerableOrderBy( Results, CrashInstance => CrashInstance.TTPID, bSortByDescending );
+						Results = EnumerableOrderBy( Results, CrashInstance => CrashInstance.Jira, bSortByDescending );
 						break;
 
 					case "Branch":
@@ -494,11 +494,11 @@ namespace Tools.CrashReporter.CrashReportWebSite.Models
 						break;
 
 					case "ChangeListVersion":
-						Results = EnumerableOrderBy( Results, CrashInstance => CrashInstance.ChangeListVersion, bSortByDescending );
+						Results = EnumerableOrderBy( Results, CrashInstance => CrashInstance.BuiltFromCL, bSortByDescending );
 						break;
 
 					case "ComputerName":
-						Results = EnumerableOrderBy( Results, CrashInstance => CrashInstance.ComputerName, bSortByDescending );
+						Results = EnumerableOrderBy( Results, CrashInstance => CrashInstance.MachineId, bSortByDescending );
 						break;
 
 					case "PlatformName":
@@ -551,10 +551,10 @@ namespace Tools.CrashReporter.CrashReportWebSite.Models
 			NewCrash.Branch = NewCrashInfo.BranchName;
 			NewCrash.BaseDir = NewCrashInfo.BaseDir;
 			NewCrash.BuildVersion = NewCrashInfo.BuildVersion;
-			NewCrash.ChangeListVersion = NewCrashInfo.BuiltFromCL.ToString();
+			NewCrash.BuiltFromCL = NewCrashInfo.BuiltFromCL.ToString();
 			NewCrash.CommandLine = NewCrashInfo.CommandLine;
 			NewCrash.EngineMode = NewCrashInfo.EngineMode;
-			NewCrash.ComputerName = NewCrashInfo.MachineGuid;
+			NewCrash.MachineId = NewCrashInfo.MachineGuid;
 
 			// Valid MachineID and UserName, updated crash from non-UE4 release
 			if(!string.IsNullOrEmpty(NewCrashInfo.UserName))
@@ -605,16 +605,9 @@ namespace Tools.CrashReporter.CrashReportWebSite.Models
 			}
 
 			NewCrash.TimeOfCrash = NewCrashInfo.TimeofCrash;
-					
-			NewCrash.HasLogFile = NewCrashInfo.bHasLog;
-			NewCrash.HasMiniDumpFile = NewCrashInfo.bHasMiniDump;
-			NewCrash.HasDiagnosticsFile = NewCrashInfo.bHasDiags;
-			NewCrash.HasVideoFile = NewCrashInfo.bHasVideo;
-			NewCrash.HasMetaData = NewCrashInfo.bHasWERData;
-
 			NewCrash.bAllowToBeContacted = NewCrashInfo.bAllowToBeContacted;
 
-			NewCrash.TTPID = "";
+			NewCrash.Jira = "";
 			NewCrash.FixedChangeList = "";
 
 			// Set the crash type
@@ -730,7 +723,7 @@ namespace Tools.CrashReporter.CrashReportWebSite.Models
 				Results =
 						(
 							from CrashDetail in Results
-							where CrashDetail.EpicAccountId.Equals( DecodedEpicOrMachineID ) || CrashDetail.ComputerName.Equals( DecodedEpicOrMachineID )
+							where CrashDetail.EpicAccountId.Equals( DecodedEpicOrMachineID ) || CrashDetail.MachineId.Equals( DecodedEpicOrMachineID )
 							select CrashDetail
 						);
 			}
@@ -741,7 +734,7 @@ namespace Tools.CrashReporter.CrashReportWebSite.Models
 				Results =
 					(
 						from CrashDetail in Results
-						where CrashDetail.TTPID.Contains( DecodedJIRA )
+						where CrashDetail.Jira.Contains( DecodedJIRA )
 						select CrashDetail
 					);
 			}
diff --git a/Engine/Source/Programs/CrashReporter/CrashReportWebSite/Models/FormHelper.cs b/Engine/Source/Programs/CrashReporter/CrashReportWebSite/Models/FormHelper.cs
--- a/Engine/Source/Programs/CrashReporter/CrashReportWebSite/Models/FormHelper.cs
+++ b/Engine/Source/Programs/CrashReporter/CrashReportWebSite/Models/FormHelper.cs
@@ -34,7 +34,7 @@ namespace Tools.CrashReporter.CrashReportWebSite.Models
 		/// <summary>The number of items to display per page.</summary>
 		public int PageSize = 100;
 
-		/// <summary>The term to sort by. e.g. TTPID.</summary>
+		/// <summary>The term to sort by. e.g. Jira.</summary>
 		public string SortTerm = "";
 
 		/// <summary>The types of crashes we wish to see.</summary>
diff --git a/Engine/Source/Programs/CrashReporter/CrashReportWebSite/Views/Buggs/Show.aspx b/Engine/Source/Programs/CrashReporter/CrashReportWebSite/Views/Buggs/Show.aspx
--- a/Engine/Source/Programs/CrashReporter/CrashReportWebSite/Views/Buggs/Show.aspx
+++ b/Engine/Source/Programs/CrashReporter/CrashReportWebSite/Views/Buggs/Show.aspx
@@ -118,14 +118,14 @@
 		<%using( Html.BeginForm( "Show", "Buggs" ) )
 		{	%>
 
-			<%if( string.IsNullOrEmpty( Model.Bugg.TTPID ) )
+			<%if (string.IsNullOrEmpty( Model.Bugg.Jira ))
 			{ %>
 				<dd class='even'>
 					<input type="submit" value="CopyToJira" title="<%=Model.Bugg.ToTooltip()%>" id="id <%=Model.Bugg.Id%>" name="CopyToJira-<%=Model.Bugg.Id%>" class="SearchButton CopyToJiraStyle" />
 				</dd>
 			<% } else { %>
 				<dd class='even'>
-					<a href="https://jira.ol.epicgames.net/browse/<%=Model.Bugg.TTPID%>" target="_blank"><%=Model.Bugg.TTPID%></a>
+					<a href="https://jira.ol.epicgames.net/browse/<%=Model.Bugg.Jira%>" target="_blank"><%=Model.Bugg.Jira%></a>
 				</dd>
 			<% } %>
 
@@ -263,7 +263,7 @@
 							<td class="TimeOfCrash">
 								<%=CrashInstance.GetTimeOfCrash()[0]%> &nbsp;
 								<%=CrashInstance.GetTimeOfCrash()[1]%><br />
-								Change: <%= CrashInstance.ChangeListVersion %><br />
+								Change: <%= CrashInstance.BuiltFromCL %><br />
 							</td>
 							<td class="Summary">
 								<%=CrashInstance.UserName%><br />
diff --git a/Engine/Source/Programs/CrashReporter/CrashReportWebSite/Views/Buggs/ViewBuggs.ascx b/Engine/Source/Programs/CrashReporter/CrashReportWebSite/Views/Buggs/ViewBuggs.ascx
--- a/Engine/Source/Programs/CrashReporter/CrashReportWebSite/Views/Buggs/ViewBuggs.ascx
+++ b/Engine/Source/Programs/CrashReporter/CrashReportWebSite/Views/Buggs/ViewBuggs.ascx
@@ -65,16 +65,16 @@
 						string BuggRowColor = "grey";
 						string BuggColorDescription = "Incoming CrashGroup";
 
-						if( string.IsNullOrWhiteSpace( CurrentBugg.FixedChangeList ) && string.IsNullOrWhiteSpace( CurrentBugg.TTPID ) )
+						if( string.IsNullOrWhiteSpace( CurrentBugg.FixedChangeList ) && string.IsNullOrWhiteSpace( CurrentBugg.Jira ) )
 						{
 							BuggRowColor = "#FFFF88"; // yellow
 							BuggColorDescription = "This CrashGroup has not been fixed or assigned a JIRA";
 						}
 
-						if( !string.IsNullOrWhiteSpace( CurrentBugg.TTPID ) && string.IsNullOrWhiteSpace( CurrentBugg.FixedChangeList ) )
+						if (!string.IsNullOrWhiteSpace( CurrentBugg.Jira ) && string.IsNullOrWhiteSpace( CurrentBugg.FixedChangeList ))
 						{
 							BuggRowColor = "#D01F3C"; // red
-							BuggColorDescription = "This CrashGroup has  been assigned a JIRA: " + CurrentBugg.TTPID + " but has not been fixed.";
+							BuggColorDescription = "This CrashGroup has  been assigned a JIRA: " + CurrentBugg.Jira + " but has not been fixed.";
 						}
 
 						if( CurrentBugg.Status == "Coder" )
@@ -149,7 +149,7 @@
 			<td><%=CurrentBugg.GetCrashTypeAsString()%></td>
 			<td><%=CurrentBugg.Status%></td>
 			<td><%=CurrentBugg.FixedChangeList%></td>
-			<td> <span><a href="https://jira.ol.epicgames.net/browse/<%=CurrentBugg.TTPID%>" target="_blank"><%=CurrentBugg.TTPID%></a></span>  </td>
+			<td> <span><a href="https://jira.ol.epicgames.net/browse/<%=CurrentBugg.Jira%>" target="_blank"><%=CurrentBugg.Jira%></a></span>  </td>
 		</tr>
 	<%				}
 				}
diff --git a/Engine/Source/Programs/CrashReporter/CrashReportWebSite/Views/Crashes/Show.aspx b/Engine/Source/Programs/CrashReporter/CrashReportWebSite/Views/Crashes/Show.aspx
--- a/Engine/Source/Programs/CrashReporter/CrashReportWebSite/Views/Crashes/Show.aspx
+++ b/Engine/Source/Programs/CrashReporter/CrashReportWebSite/Views/Crashes/Show.aspx
@@ -116,7 +116,7 @@
 				
 			<dt>Saved Files</dt>
 				<dd class='even'>
-					<%if( !Model.Crash.HasLogFile.HasValue || Model.Crash.HasLogFile.Equals( true ) ) 
+					<%if( Model.Crash.HasLogFile() ) 
 					{ %>
 						<a style='text-decoration:none; vertical-align:middle;' href='<%=Model.Crash.GetLogUrl() %>'><img src="../../Content/Images/Icons/log.png" style="height:20px;border:none;" />&nbsp;Log</a>
 						<%
@@ -125,15 +125,15 @@
 				</dd>
 			<dt>&nbsp;</dt>
 				<dd>	
-					<%if( !Model.Crash.HasMiniDumpFile.HasValue || Model.Crash.HasMiniDumpFile.Equals( true ) ) 
+					<%if( Model.Crash.HasMiniDumpFile() ) 
 					{ %>
-						<a style='text-decoration:none;' href='<%=Model.Crash.GetDumpUrl() %>'><img src="../../Content/Images/Icons/miniDump.png" style="height:20px;border:none;" />&nbsp;MiniDump</a>
+						<a style='text-decoration:none;' href='<%=Model.Crash.GetMiniDumpUrl() %>'><img src="../../Content/Images/Icons/miniDump.png" style="height:20px;border:none;" />&nbsp;MiniDump</a>
 						<%
 					} %>
 				</dd>
 			<dt>&nbsp;</dt>
 				<dd>
-					<%if( Model.Crash.HasVideoFile.HasValue && Model.Crash.HasVideoFile.Equals( true ) ) 
+					<%if( Model.Crash.HasVideoFile() ) 
 					{ %>
 						<a id="VideoLink" style="text-decoration:none;" href='<%=Model.Crash.GetVideoUrl() %>'><img src="../../Content/Images/Icons/video.png" style="height:20px;border:none;" />&nbsp;Video</a>
 						<%
@@ -141,7 +141,7 @@
 				</dd>
 			<dt>&nbsp;</dt>
 				<dd>
-					<%if( Model.Crash.HasDiagnosticsFile.HasValue && Model.Crash.HasDiagnosticsFile.Equals( true ) ) 
+					<%if (Model.Crash.HasDiagnosticsFile() ) 
 					{ %>
 						<a id="A1" style="text-decoration:none;" href='<%=Model.Crash.GetDiagnosticsUrl() %>'><img src="../../Content/Images/Icons/log.png" style="height:20px;border:none;" />&nbsp;Diags</a>
 						<%
@@ -149,7 +149,7 @@
 				</dd>
 			<dt>&nbsp;</dt>
 				<dd>
-					<%if( Model.Crash.HasMetaData.HasValue && Model.Crash.HasMetaData.Equals( true ) ) 
+					<%if( Model.Crash.HasMetaDataFile() ) 
 					{ %>
 						<a id="A2" style="text-decoration:none;" href='<%=Model.Crash.GetMetaDataUrl() %>'><img src="../../Content/Images/Icons/log.png" style="height:20px;border:none;" />&nbsp;WERInfo</a>
 						<%
@@ -189,7 +189,7 @@
 				<dd ><%=Html.DisplayFor(m => Model.Crash.PlatformName) %></dd>
 
 			<dt>Machine Id</dt> 
-				<dd class='even'><small><%=Html.DisplayFor(m => Model.Crash.ComputerName) %></small></dd>
+				<dd class='even'><small><%=Html.DisplayFor(m => Model.Crash.MachineId) %></small></dd>
 
 			<dt>Epic Account Id</dt> 
 				<dd class='even'><small><%=Html.DisplayFor(m => Model.Crash.EpicAccountId) %></small></dd>
@@ -204,11 +204,11 @@
 				<dd class='even'><%=Html.DisplayFor(m => Model.Crash.FixedChangeList) %></dd>
 
 			<dt>Changelist #</dt>
-				<dd class='even'><%=Html.DisplayFor(m => Model.Crash.ChangeListVersion) %></dd>
+				<dd class='even'><%=Html.DisplayFor(m => Model.Crash.BuiltFromCL) %></dd>
 
 			<dt>JIRA</dt>
-				<%--<dd ><%=Html.DisplayFor(m => Model.Crash.TTPID) %></dd>--%>
-				<dd><a href="https://jira.ol.epicgames.net/browse/<%=Model.Crash.TTPID%>" target="_blank"><%=Model.Crash.TTPID%></a></dd>
+				<%--<dd ><%=Html.DisplayFor(m => Model.Crash.Jira) %></dd>--%>
+				<dd><a href="https://jira.ol.epicgames.net/browse/<%=Model.Crash.Jira%>" target="_blank"><%=Model.Crash.Jira%></a></dd>
 
 			<dt>Status</dt>
 				<dd ><%=Html.DisplayFor(m => Model.Crash.Status) %></dd>
diff --git a/Engine/Source/Programs/CrashReporter/CrashReportWebSite/Views/Crashes/ViewCrash.ascx b/Engine/Source/Programs/CrashReporter/CrashReportWebSite/Views/Crashes/ViewCrash.ascx
--- a/Engine/Source/Programs/CrashReporter/CrashReportWebSite/Views/Crashes/ViewCrash.ascx
+++ b/Engine/Source/Programs/CrashReporter/CrashReportWebSite/Views/Crashes/ViewCrash.ascx
@@ -113,7 +113,7 @@
 							string CrashRowColor = "grey";
 							string CrashColorDescription = "Incoming Crash";
 
-							if( string.IsNullOrWhiteSpace( CurrentCrash.FixedChangeList ) && string.IsNullOrWhiteSpace( CurrentCrash.TTPID ) )
+							if( string.IsNullOrWhiteSpace( CurrentCrash.FixedChangeList ) && string.IsNullOrWhiteSpace( CurrentCrash.Jira ) )
 							{
 								CrashRowColor = "#FFFF88"; // yellow
 								CrashColorDescription = "This crash has not been fixed or assigned a JIRA";
@@ -126,10 +126,10 @@
 								CrashColorDescription = "This crash has been fixed in CL# " + CurrentCrash.FixedChangeList;
 							}
 
-							if( ( BuggCrash != null ) && !string.IsNullOrWhiteSpace( CurrentCrash.TTPID ) && string.IsNullOrWhiteSpace( CurrentCrash.FixedChangeList ) )
+							if (( BuggCrash != null ) && !string.IsNullOrWhiteSpace( CurrentCrash.Jira ) && string.IsNullOrWhiteSpace( CurrentCrash.FixedChangeList ))
 							{
 								CrashRowColor = "#D01F3C"; // red
-								CrashColorDescription = "This crash has occurred more than once and been assigned a JIRA: " + CurrentCrash.TTPID + " but has not been fixed.";
+								CrashColorDescription = "This crash has occurred more than once and been assigned a JIRA: " + CurrentCrash.Jira + " but has not been fixed.";
 							}
 
 							if( CurrentCrash.Status == "Tester" )
@@ -182,12 +182,12 @@
 									<td class="Game"><%=CurrentCrash.GameName%></td>
 									<td class="Mode"><%=CurrentCrash.EngineMode%></td>
 									<td class="FixedChangeList"><%=CurrentCrash.FixedChangeList%></td>
-									<td class="Jira"> <span><a href="https://jira.ol.epicgames.net/browse/<%=CurrentCrash.TTPID%>" target="_blank"><%=CurrentCrash.TTPID%></a></span>  </td>
+									<td class="Jira"> <span><a href="https://jira.ol.epicgames.net/browse/<%=CurrentCrash.Jira%>" target="_blank"><%=CurrentCrash.Jira%></a></span>  </td>
 									<td class="Branch"><%=CurrentCrash.Branch%>&nbsp;</td>
 									<td class="Description"><span class="TableData"><%=CurrentCrash.Description%>&nbsp;</span></td>
 									<td class="Summary"><%=Html.Encode(CurrentCrash.Summary)%></td>
-									<td class="ChangeListVersion"><%=CurrentCrash.ChangeListVersion%></td>
-									<td class="Computer"><%=CurrentCrash.ComputerName%></td>
+									<td class="BuiltFromCL"><%=CurrentCrash.BuiltFromCL%></td>
+									<td class="Computer"><%=CurrentCrash.MachineId%></td>
 									<td class="Platform"><%=CurrentCrash.PlatformName%></td>
 									<td class="Status"><%=CurrentCrash.Status%></td>
 									<td class="Module"><%=CurrentCrash.Module%></td>
diff --git a/Engine/Source/Programs/CrashReporter/CrashReportWebSite/Views/Reports/ViewReports.ascx b/Engine/Source/Programs/CrashReporter/CrashReportWebSite/Views/Reports/ViewReports.ascx
--- a/Engine/Source/Programs/CrashReporter/CrashReportWebSite/Views/Reports/ViewReports.ascx
+++ b/Engine/Source/Programs/CrashReporter/CrashReportWebSite/Views/Reports/ViewReports.ascx
@@ -102,7 +102,7 @@ table, table tr, table th, table td
 		</td>
 		<td>
 			<%--NewBugg.Id = RealBugg.Id;								// CrashGroup URL (a link to the Bugg)--%>
-			<%--<a href="https://jira.ol.epicgames.net/browse/<%=Model.Bugg.TTPID%>" target="_blank"><%=Model.Bugg.TTPID%></a>--%>
+			<%--<a href="https://jira.ol.epicgames.net/browse/<%=Model.Bugg.Jira%>" target="_blank"><%=Model.Bugg.Jira%></a>--%>
 			<%=Html.ActionLink( Bugg.Id.ToString(), "Show", new { controller = "Buggs", id = Bugg.Id }, null )%>
 		</td>
 		<td>
@@ -125,20 +125,20 @@ table, table tr, table th, table td
 			<%--NewBugg.LatestOSAffected = LatestOSAffected;			// Latest Environment Affected--%>
 			<%=Bugg.LatestOSAffected%>
 		</td>
-		<%if( string.IsNullOrEmpty( Bugg.TTPID ) && !string.IsNullOrEmpty(Bugg.Pattern ) )
+		<%if (string.IsNullOrEmpty( Bugg.Jira ) && !string.IsNullOrEmpty( Bugg.Pattern ))
 		{ %>
 			<td>
 				<input type="submit" value="CopyToJira" title="<%=Bugg.ToTooltip()%>" id="id <%=Bugg.Id%>" name="CopyToJira-<%=Bugg.Id%>" class="SearchButton CopyToJiraStyle CopyToJiraStyle" />
 			</td>
-		<% } else if( string.IsNullOrEmpty( Bugg.TTPID ) && string.IsNullOrEmpty(Bugg.Pattern ) )
+		<% } else if( string.IsNullOrEmpty( Bugg.Jira ) && string.IsNullOrEmpty(Bugg.Pattern ) )
 		{ %>
 			<td>
 				No pattern!
 			</td>
 		<% } else { %>
 			<td>
-				<%--NewBugg.TTPID = RealBugg.TTPID;							// JIRA--%>
-				<a href="https://jira.ol.epicgames.net/browse/<%=Bugg.TTPID%>" target="_blank"><%=Bugg.TTPID%></a>
+				<%--NewBugg.Jira = RealBugg.Jira;							// JIRA--%>
+				<a href="https://jira.ol.epicgames.net/browse/<%=Bugg.Jira%>" target="_blank"><%=Bugg.Jira%></a>
 			</td>
 		<% } %>
 
