CrashReportWWW - Added CSV controller

diff --git a/Engine/Source/Programs/CrashReporter/CrashReportWebSite/Controllers/CSVController.cs b/Engine/Source/Programs/CrashReporter/CrashReportWebSite/Controllers/CSVController.cs
new file mode 100644
--- /dev/null
+++ b/Engine/Source/Programs/CrashReporter/CrashReportWebSite/Controllers/CSVController.cs
@@ -0,0 +1,292 @@
+﻿// Copyright 1998-2015 Epic Games, Inc. All Rights Reserved.
+
+using System;
+using System.Collections.Generic;
+using System.Diagnostics;
+using System.IO;
+using System.Linq;
+using System.Linq.Expressions;
+using System.Net;
+using System.Text;
+using System.Web.Mvc;
+using System.Xml;
+using Tools.CrashReporter.CrashReportWebSite.Models;
+using Tools.CrashReporter.CrashReportWebSite.Properties;
+using Tools.DotNETCommon;
+
+
+namespace Tools.CrashReporter.CrashReportWebSite.Controllers
+{
+	// 
+	// How hard would it be for you to create a CSV of:
+	// Epic / Machine ID, 
+	// Buggs ID, 
+	// EngineVersion, 
+	// 
+	// # of crashes with that buggs ID for 
+	//	this engine version and	user
+
+	//E-32chars
+	//M-32chars
+	// 
+	//Engine version: 4.8.1-cl+branch
+	// 
+	// Each time the controller is run, create a new CSV file in the specified folder and add the link and the top of the site,
+	// so it could be easily downloaded
+	// Removed files older that 7 days
+
+	/// <summary>CSV data row.</summary>
+	public class FCSVRow
+	{
+		/// <summary>Epic Id.</summary>
+		public string EpicId = "";
+		/// <summary>Bugg Id.</summary>
+		public int BuggId = 0;
+		/// <summary>Engine version.</summary>
+		public string AffectedVersionsString = "";
+		/// <summary># of crashes with that buggs ID for this engine version and user.</summary>
+		public int NumberOfCrashes = 0;
+
+		/// <summary>User's email.</summary>
+		public string UserEmail = "";
+
+		//// <summary>Reference to a bugg.</summary>
+		//public Bugg BuggRef = null;
+	}
+
+	/// <summary>
+	/// The controller to handle graphing of crashes per user group over time.
+	/// </summary>
+	public class CSVController : Controller
+	{
+		/// <summary>Fake id for all user groups</summary>
+		public static readonly int AllUserGroupId = -1;
+
+		static readonly int MinNumberOfCrashes = 2;
+
+		/// <summary>
+		/// Retrieve all Buggs matching the search criteria.
+		/// </summary>
+		/// <param name="FormData">The incoming form of search criteria from the client.</param>
+		/// <returns>A view to display the filtered Buggs.</returns>
+		public CSV_ViewModel GetResults( FormHelper FormData )
+		{
+			BuggRepository BuggsRepo = new BuggRepository();
+			CrashRepository CrashRepo = new CrashRepository();
+
+			// No Jira support at this moment
+			// Enumerate JIRA projects if needed.
+			// https://jira.ol.epicgames.net//rest/api/2/project
+			var JC = JiraConnection.Get();
+			var JiraComponents = JC.GetNameToComponents();
+			var JiraVersions = JC.GetNameToVersions();
+
+			using( FAutoScopedLogTimer LogTimer = new FAutoScopedLogTimer( this.GetType().ToString() ) )
+			{
+				string AnonymousGroup = "Anonymous";
+				//List<String> Users = FRepository.Get().GetUserNamesFromGroupName( AnonumousGroup );
+				int AnonymousGroupID = FRepository.Get( BuggsRepo ).FindOrAddGroup( AnonymousGroup );
+				HashSet<int> AnonumousIDs = FRepository.Get( BuggsRepo ).GetUserIdsFromUserGroup( AnonymousGroup );
+				int AnonymousID = AnonumousIDs.First();
+				HashSet<string> UserNamesForUserGroup = FRepository.Get( BuggsRepo ).GetUserNamesFromGroupName( AnonymousGroup );
+
+				FormData.DateFrom = FormData.DateTo.AddDays( -1 );
+
+				var Crashes = CrashRepo
+					.FilterByDate( CrashRepo.ListAll(), FormData.DateFrom, FormData.DateTo )
+					// Only crashes and asserts
+					.Where( Crash => Crash.CrashType == 1 || Crash.CrashType == 2 )
+					// Only anonymous user
+					.Where( Crash => Crash.UserNameId.Value == AnonymousID )
+					// We need Epic ID
+					// or Machine ID to identify the owner of a crash
+					//.Where( Crash => !string.IsNullOrEmpty( Crash.EpicAccountId ) || !string.IsNullOrEmpty( Crash.MachineId ) )
+					.Where( Crash => !string.IsNullOrEmpty( Crash.EpicAccountId ) )
+					.Where( Crash => !string.IsNullOrEmpty( Crash.Pattern ) )
+					.Select( Crash => new
+					{
+						ID = Crash.Id,
+						TimeOfCrash = Crash.TimeOfCrash.Value,
+						//UserID = Crash.UserNameId.Value, 
+						BuildVersion = Crash.BuildVersion,
+						JIRA = Crash.Jira,
+						Platform = Crash.PlatformName,
+						FixCL = Crash.FixedChangeList,
+						BuiltFromCL = Crash.BuiltFromCL,
+						Pattern = Crash.Pattern,
+						MachineID = Crash.MachineId,
+						Branch = Crash.Branch,
+						Description = Crash.Description,
+						RawCallStack = Crash.RawCallStack,
+						EpicAccountId = Crash.EpicAccountId,
+					} )
+					.ToList();
+				int NumCrashes = Crashes.Count;
+
+				// Get all users
+				// SLOW
+				//var Users = FRepository.Get( BuggsRepo ).GetAnalyticsUsers().ToDictionary( X => X.EpicAccountId, Y => Y );
+
+				// Total # of UNIQUE (Anonymous) crashes in timeframe
+				Dictionary<string, int> PatternToCount = new Dictionary<string, int>();
+				HashSet<string> UniqueEpicIds = new HashSet<string>();
+				Dictionary<string, FCSVRow> EpicIdBuggToCSVRow = new Dictionary<string, FCSVRow>();
+
+				// Iterate all crashes and generate all patterns, and EpicIds
+				foreach( var Crash in Crashes )
+				{
+					if( string.IsNullOrEmpty( Crash.Pattern ) )
+					{
+						continue;
+					}
+
+					bool bAdd = !PatternToCount.ContainsKey( Crash.Pattern );
+					if (bAdd)
+					{
+						PatternToCount.Add( Crash.Pattern, 1 );
+					}
+					else
+					{
+						PatternToCount[Crash.Pattern]++;
+					}
+
+					UniqueEpicIds.Add( Crash.EpicAccountId );
+				}
+
+				// Select only patterns with at least 2 crashes.
+				var PatternToCountOrdered = PatternToCount.OrderByDescending( X => X.Value ).ToList();
+				var PatternAndCount = PatternToCountOrdered.Where( X => X.Value > MinNumberOfCrashes ).ToDictionary( X => X.Key, Y => Y.Value );
+
+				// Get all buggs associated with crashes.
+				var RealBuggs = new Dictionary<string, Bugg>();
+				foreach (var PatternIt in PatternAndCount)
+				{
+					var RealBugg = BuggsRepo.Context.Buggs.Where( Bugg => Bugg.Pattern == PatternIt.Key ).FirstOrDefault();
+					if (RealBugg != null)
+					{
+						var CrashesForBugg = Crashes.Where( Crash => Crash.Pattern == PatternIt.Key ).ToList();
+
+						// Convert anonymous to full type;
+						var FullCrashesForBugg = new List<Crash>( CrashesForBugg.Count );
+						foreach (var Anon in CrashesForBugg)
+						{
+							FullCrashesForBugg.Add( new Crash()
+							{
+								ID = Anon.ID,
+								TimeOfCrash = Anon.TimeOfCrash,
+								BuildVersion = Anon.BuildVersion,
+								Jira = Anon.JIRA,
+								Platform = Anon.Platform,
+								FixCL = Anon.FixCL,
+								BuiltFromCL = Anon.BuiltFromCL,
+								Pattern = Anon.Pattern,
+								MachineId = Anon.MachineID,
+								Branch = Anon.Branch,
+								Description = Anon.Description,
+								RawCallStack = Anon.RawCallStack,
+							} );
+						}
+
+						RealBugg.PrepareBuggForJira( FullCrashesForBugg );
+						RealBuggs.Add( PatternIt.Key, RealBugg );
+					}
+					else
+					{
+						//FLogger.Global.WriteEvent( "Bugg for pattern " + PatternIt + " not found" );
+					}
+				}
+
+
+				// Iterate all crashes and link EpicId with Bugg, and calculate #
+				foreach( var Crash in Crashes )
+				{
+					Bugg Bugg = null;
+					RealBuggs.TryGetValue( Crash.Pattern, out Bugg );
+
+					if( Bugg != null )
+					{
+						string EpicIdBugg = Crash.EpicAccountId + "%" + Bugg.Id;
+						bool bAdd = !EpicIdBuggToCSVRow.ContainsKey( EpicIdBugg );
+						if (bAdd)
+						{
+							var Row = new FCSVRow 
+							{ 
+								EpicId = Crash.EpicAccountId,
+								BuggId = Bugg.Id,
+								AffectedVersionsString = Bugg.GetAffectedVersions(),
+								NumberOfCrashes = 1,
+								//UserEmail = Users[Crash.EpicAccountId].UserEmail
+							};
+							EpicIdBuggToCSVRow.Add( EpicIdBugg, Row );
+						}
+						else
+						{
+							EpicIdBuggToCSVRow[EpicIdBugg].NumberOfCrashes++;
+						}
+					}
+				}
+
+				// Export data to the file.
+				string CSVPathname = Path.Combine( Settings.Default.CrashReporterCSV, DateTime.UtcNow.ToString( "yyyy-MM-dd.HH-mm-ss" ) ) + ".csv";
+				string ServerPath = Server.MapPath( CSVPathname );
+				var CSVFile = new StreamWriter( ServerPath, true, Encoding.UTF8 );
+
+				CSVFile.WriteLine( "EpicId; BuggId; AffectedVersion; # of crashes;" );
+
+				foreach (var Row in EpicIdBuggToCSVRow)
+				{
+					if (Row.Value.NumberOfCrashes > MinNumberOfCrashes)
+					{
+						CSVFile.WriteLine( "{0};{1};{2};{3};", Row.Value.EpicId, Row.Value.BuggId, Row.Value.AffectedVersionsString, Row.Value.NumberOfCrashes );
+					}
+				}
+
+				CSVFile.Flush();
+				CSVFile.Close();
+				CSVFile = null;
+
+				List<FCSVRow> CSVRows = EpicIdBuggToCSVRow
+					.OrderByDescending( X => X.Value.NumberOfCrashes )
+					.Select( X => X.Value )
+					.Where( X => X.NumberOfCrashes > MinNumberOfCrashes )
+					.ToList();
+				
+				// Total # of unique anonymous crashes.
+				int TotalUniqueAnonymousCrashes = PatternToCount.Count;
+
+				// Total # of AFFECTED USERS (Anonymous) in timeframe
+				int TotalAffectedUsers = UniqueEpicIds.Count;
+
+				// Total # of ALL (Anonymous) crashes in timeframe
+				int TotalAnonymousCrashes = NumCrashes;
+
+				return new CSV_ViewModel
+				{
+					CSVRows = CSVRows,
+					CSVPathname = CSVPathname,
+					DateFrom = (long)( FormData.DateFrom - CrashesViewModel.Epoch ).TotalMilliseconds,
+					DateTo = (long)( FormData.DateTo - CrashesViewModel.Epoch ).TotalMilliseconds,
+					TotalAffectedUsers = TotalAffectedUsers,
+					TotalAnonymousCrashes = TotalAnonymousCrashes,
+					TotalUniqueAnonymousCrashes = TotalUniqueAnonymousCrashes
+				};
+			}
+		}
+
+
+		/// <summary>
+		/// 
+		/// </summary>
+		/// <returns></returns>
+		public ActionResult Index( FormCollection Form )
+		{
+			using( FAutoScopedLogTimer LogTimer = new FAutoScopedLogTimer( this.GetType().ToString(), bCreateNewLog: true ) )
+			{
+				FormHelper FormData = new FormHelper( Request, Form, "JustReport" );
+				CSV_ViewModel Results = GetResults( FormData );
+				Results.GenerationTime = LogTimer.GetElapsedSeconds().ToString( "F2" );
+				return View( "Index", Results );
+			}
+		}
+	}
+}
\ No newline at end of file
diff --git a/Engine/Source/Programs/CrashReporter/CrashReportWebSite/CrashReportWebSite.csproj b/Engine/Source/Programs/CrashReporter/CrashReportWebSite/CrashReportWebSite.csproj
--- a/Engine/Source/Programs/CrashReporter/CrashReportWebSite/CrashReportWebSite.csproj
+++ b/Engine/Source/Programs/CrashReporter/CrashReportWebSite/CrashReportWebSite.csproj
@@ -84,6 +84,7 @@
     <Compile Include="..\..\DotNETCommon\MetaData.cs">
       <Link>Properties\MetaData.cs</Link>
     </Compile>
+    <Compile Include="Controllers\CSVController.cs" />
     <Compile Include="Controllers\DashboardController.cs" />
     <Compile Include="Controllers\ErrorController.cs" />
     <Compile Include="Controllers\HomeController.cs" />
@@ -93,6 +94,7 @@
     <Compile Include="Models\CachedDataService.cs" />
     <Compile Include="Models\BuggRepository.cs" />
     <Compile Include="Models\BuggViewModel.cs" />
+    <Compile Include="Models\CSVViewModel.cs" />
     <Compile Include="Models\DashboardViewModel.cs" />
     <Compile Include="Models\CrashReport.cs">
       <DependentUpon>CrashReport.dbml</DependentUpon>
@@ -201,6 +203,8 @@
     <Content Include="Views\Buggs\Index.aspx" />
     <Content Include="Views\Buggs\Show.aspx" />
     <Content Include="Views\Buggs\ViewBuggs.ascx" />
+    <Content Include="Views\CSV\Index.aspx" />
+    <Content Include="Views\CSV\ViewCSV.ascx" />
     <Content Include="Views\Dashboard\Index.aspx" />
     <Content Include="Views\Crashes\Index.aspx" />
     <Content Include="Views\Crashes\Show.aspx" />
diff --git a/Engine/Source/Programs/CrashReporter/CrashReportWebSite/Models/CSVViewModel.cs b/Engine/Source/Programs/CrashReporter/CrashReportWebSite/Models/CSVViewModel.cs
new file mode 100644
--- /dev/null
+++ b/Engine/Source/Programs/CrashReporter/CrashReportWebSite/Models/CSVViewModel.cs
@@ -0,0 +1,56 @@
+﻿// Copyright 1998-2015 Epic Games, Inc. All Rights Reserved.
+
+using System;
+using System.Collections.Generic;
+using System.Linq;
+using System.Web.Mvc;
+using Tools.CrashReporter.CrashReportWebSite.Controllers;
+using System.IO;
+
+namespace Tools.CrashReporter.CrashReportWebSite.Models
+{
+	/// <summary>
+	/// The view model for the crash summary page.
+	/// </summary>
+	public class CSV_ViewModel
+	{
+		/// <summary>A container of sorted CSV rows.</summary>
+		public List<FCSVRow> CSVRows { get; set; }
+
+		/// <summary>The date of the earliest crash in a Bugg.</summary>
+		public long DateFrom { get; set; }
+
+		/// <summary>The date of the most recent crash in a Bugg.</summary>
+		public long DateTo { get; set; }
+
+		/// <summary></summary>
+		public int TotalAnonymousCrashes { get; set; }
+
+		/// <summary></summary>
+		public int TotalUniqueAnonymousCrashes { get; set; }
+
+		/// <summary></summary>
+		public int TotalAffectedUsers { get; set; }
+
+		/// <summary>Pathname where the CSV file is saved.</summary>
+		public string CSVPathname { get; set; }
+
+		/// <summary>Filename of the CSV file</summary>
+		public string GetCSVFilename()
+		{
+			return Path.GetFileName( CSVPathname );
+		}
+
+		/// <summary>Directory where the CSV files are saved.</summary>
+		public string GetCSVDirectory()
+		{
+			return Path.GetDirectoryName( CSVPathname );
+		}
+
+		/// <summary>Time spent in generating this site, formatted as a string.</summary>
+		public string GenerationTime { get; set; }
+
+		/// <summary>User input from the client.</summary>
+		public FormCollection FormCollection { get; set; }
+	}
+}
\ No newline at end of file
diff --git a/Engine/Source/Programs/CrashReporter/CrashReportWebSite/Models/PerformanceTimers.cs b/Engine/Source/Programs/CrashReporter/CrashReportWebSite/Models/PerformanceTimers.cs
--- a/Engine/Source/Programs/CrashReporter/CrashReportWebSite/Models/PerformanceTimers.cs
+++ b/Engine/Source/Programs/CrashReporter/CrashReportWebSite/Models/PerformanceTimers.cs
@@ -6,7 +6,9 @@ using Tools.CrashReporter.CrashReportCommon;
 using System.IO;
 using System.Runtime.CompilerServices;
 using System.Diagnostics;
-using System.Threading; 
+using System.Threading;
+
+using Tools.CrashReporter.CrashReportWebSite.Properties;
 
 /// <summary> Helper class used to log various things into the log file. </summary>
 public class FLogger
@@ -64,6 +66,17 @@ public class FLogger
 		return FullName;
 	}
 
+	/// <summary>
+	/// Writes a custom message to the log file.
+	/// </summary>
+	public void WriteCustom( string Message )
+	{
+		if (Message != null && Message.Length > 2)
+		{
+			Log.Print( Message );
+		}
+	}
+
 	/// <summary>
 	/// Writes a status update to the log file.
 	/// </summary>
diff --git a/Engine/Source/Programs/CrashReporter/CrashReportWebSite/Models/Repository.cs b/Engine/Source/Programs/CrashReporter/CrashReportWebSite/Models/Repository.cs
--- a/Engine/Source/Programs/CrashReporter/CrashReportWebSite/Models/Repository.cs
+++ b/Engine/Source/Programs/CrashReporter/CrashReportWebSite/Models/Repository.cs
@@ -309,5 +309,15 @@ namespace Tools.CrashReporter.CrashReportWebSite.Models
 				return Results;
 			}
 		}
+
+		/// <summary>Returns a list of all users, from the AnalyticsDB.</summary>
+		public List<UsersMapping> GetAnalyticsUsers()
+		{
+			using( FAutoScopedLogTimer LogTimer = new FAutoScopedLogTimer( this.GetType().ToString()  ) )
+			{
+				var AllUsers = Context.ExecuteQuery<UsersMapping>( @"SELECT * FROM [analyticsdb-01.dmz.epicgames.net].[CrashReport].[dbo].[UsersMapping]" );
+				return AllUsers.ToList();
+			}
+		}
 	}
 }
diff --git a/Engine/Source/Programs/CrashReporter/CrashReportWebSite/Properties/Settings.settings b/Engine/Source/Programs/CrashReporter/CrashReportWebSite/Properties/Settings.settings
--- a/Engine/Source/Programs/CrashReporter/CrashReportWebSite/Properties/Settings.settings
+++ b/Engine/Source/Programs/CrashReporter/CrashReportWebSite/Properties/Settings.settings
@@ -8,5 +8,8 @@
     <Setting Name="CrashReporterVideos" Type="System.String" Scope="Application">
       <Value Profile="(Default)">/CrashReporterVideos/</Value>
     </Setting>
+    <Setting Name="CrashReporterCSV" Type="System.String" Scope="Application">
+      <Value Profile="(Default)">/CrashReporterCSV/</Value>
+    </Setting>
   </Settings>
 </SettingsFile>
\ No newline at end of file
diff --git a/Engine/Source/Programs/CrashReporter/CrashReportWebSite/Views/Buggs/Index.aspx b/Engine/Source/Programs/CrashReporter/CrashReportWebSite/Views/Buggs/Index.aspx
--- a/Engine/Source/Programs/CrashReporter/CrashReportWebSite/Views/Buggs/Index.aspx
+++ b/Engine/Source/Programs/CrashReporter/CrashReportWebSite/Views/Buggs/Index.aspx
@@ -8,7 +8,7 @@
 </asp:Content>
 
 <asp:Content ID="TitleContent" ContentPlaceHolderID="TitleContent" runat="server">
-	Crash Report: Buggs
+[CR] Buggs
 </asp:Content>
 
 <asp:Content ID="ScriptContent" ContentPlaceHolderID="ScriptContent" runat="server">
diff --git a/Engine/Source/Programs/CrashReporter/CrashReportWebSite/Views/CSV/Index.aspx b/Engine/Source/Programs/CrashReporter/CrashReportWebSite/Views/CSV/Index.aspx
new file mode 100644
--- /dev/null
+++ b/Engine/Source/Programs/CrashReporter/CrashReportWebSite/Views/CSV/Index.aspx
@@ -0,0 +1,75 @@
+﻿<%-- // Copyright 1998-2015 Epic Games, Inc. All Rights Reserved. --%>
+
+<%@ Page Title="" Language="C#" MasterPageFile="~/Views/Shared/Site.Master" Inherits="System.Web.Mvc.ViewPage<CSV_ViewModel>" %>
+<%@ Import Namespace="Tools.CrashReporter.CrashReportWebSite.Models" %>
+
+<asp:Content ID="StyleSheet" ContentPlaceHolderID="CssContent" runat="server">
+	<link href="../../Content/Site.css" rel="stylesheet" type="text/css" />
+</asp:Content>
+
+<asp:Content ID="TitleContent" ContentPlaceHolderID="TitleContent" runat="server">
+[CR] CSV
+</asp:Content>
+
+<asp:Content ID="ScriptContent" ContentPlaceHolderID="ScriptContent" runat="server" >
+	<script type="text/javascript">
+
+		$(function ()
+		{
+			$("#dateFromVisible")
+				.datepicker({ maxDate: '+0D' })
+				.datepicker('setDate', new Date(parseInt($('#dateFrom').val())));
+
+			$("#dateToVisible")
+				.datepicker({ maxDate: '+0D' })
+				.datepicker('setDate', new Date(parseInt($('#dateTo').val())));
+
+		});
+
+		$.datepicker.setDefaults({
+			onSelect: function ()
+			{
+				$("#dateFrom").val($("#dateFromVisible").datepicker('getDate').getTime());
+				$("#dateTo").val($("#dateToVisible").datepicker('getDate').getTime());
+			}
+		});
+
+	</script>
+</asp:Content>
+
+<asp:Content ID="AboveMainContent" ContentPlaceHolderID="AboveMainContent" runat="server" >
+<div style="clear:both;"><small style="color: lightgray;">Generated in <%=Model.GenerationTime%> second(s)</small><br /></div>
+
+<div id='SearchForm' style="clear:both;">
+<%using( Html.BeginForm( "", "CSV", FormMethod.Get, new { id = "CSVForm" } ) )
+{ %>
+	<script>$.datepicker.setDefaults($.datepicker.regional['']);</script>
+
+	<p class="SearchTextTitle">Filter by Date</p>
+
+	<p>
+		<span class="SearchTextTitle">From:</span>
+		<input id="dateFromVisible" type="text" class="date" autocomplete="OFF" style="width:80px" />
+		<input id="dateFrom" name="dateFrom" type="hidden" value="<%=Model.DateFrom %>" autocomplete="OFF" />
+		<span class="SearchTextTitle">To:</span>
+		<input id="dateToVisible" type="text" class="date" autocomplete="OFF" style="width:80px" />
+		<input id="dateTo" name="dateTo" type="hidden" value="<%=Model.DateTo %>"autocomplete="OFF"  />
+	</p>
+
+	 <p>
+		 <input type="submit" value="Search" class='SearchButton' />
+	 </p>
+	
+	<%} %>
+
+</div>
+</asp:Content>
+
+<asp:Content ID="CSVContent" ContentPlaceHolderID="CSVContent" runat="server">
+	<div>
+		<div>
+			<% Html.RenderPartial("/Views/CSV/ViewCSV.ascx"); %>
+		</div>
+	</div>
+
+</asp:Content>
diff --git a/Engine/Source/Programs/CrashReporter/CrashReportWebSite/Views/CSV/ViewCSV.ascx b/Engine/Source/Programs/CrashReporter/CrashReportWebSite/Views/CSV/ViewCSV.ascx
new file mode 100644
--- /dev/null
+++ b/Engine/Source/Programs/CrashReporter/CrashReportWebSite/Views/CSV/ViewCSV.ascx
@@ -0,0 +1,115 @@
+<%-- // Copyright 1998-2015 Epic Games, Inc. All Rights Reserved. --%>
+
+<%@ Control Language="C#" Inherits="System.Web.Mvc.ViewUserControl<CSV_ViewModel>" %>
+<%@ Import Namespace="Tools.CrashReporter.CrashReportWebSite.Models" %>
+<%@ Import Namespace="Tools.CrashReporter.CrashReportWebSite.Controllers" %>
+
+<style>
+
+table, table tr, table th, table td
+{
+	border: solid;
+	border-width: 1px;
+
+	color: white;
+	background-color: black;
+
+	font-family: 'Segoe UI', Tahoma, Geneva, Verdana, sans-serif;
+	font-size: 12px;
+
+	padding: 2px;
+	margin: 2px;
+}
+
+</style>
+	
+<body>
+<form runat="server">
+	
+<table >
+	<tr>
+		<td>
+			TotalAnonymousCrashes:&nbsp;<%=Model.TotalAnonymousCrashes %>
+		</td>
+	</tr>
+
+	<tr>
+		<td>
+			TotalUniqueAnonymousCrashes:&nbsp;<%=Model.TotalUniqueAnonymousCrashes %>
+		</td>
+	</tr>
+
+	<tr>
+		<td>
+			TotalAffectedUsers:&nbsp;<%=Model.TotalAffectedUsers %>
+		</td>
+	</tr>
+
+	<tr>
+		<td>
+			File:&nbsp;<a href="<%=Model.CSVPathname%>" target="_blank"><%=Model.GetCSVFilename()%></a>
+		</td>
+	</tr>
+
+	<tr>
+		<td>
+			Directory:&nbsp;<a href="<%=Model.GetCSVDirectory()%>" target="_blank"><%=Model.GetCSVDirectory()%></a>
+		</td>
+	</tr>
+</table>
+<br />
+
+<p>Displaying top 100 records, ordered by # of crashes</p>
+<table style="width: 100%">
+	<tr>
+		<th>
+			#
+		</th>
+		<th>
+			EpicId
+		</th>
+		<%--<th>
+			User's email
+		</th>--%>
+		<th>
+			BuggId
+		</th>
+		<th>
+			Versions Affected
+		</th>
+	</tr>
+
+	<%
+		for (int Index = 0; Index < Model.CSVRows.Count; Index++)
+		{
+			FCSVRow CSVRow = Model.CSVRows[Index];
+			if( Index > 100 )
+			{
+				break;
+			}
+	%>
+	<tr>
+		<td>
+			<%=CSVRow.NumberOfCrashes%>
+		</td>
+		<td>
+			<%=CSVRow.EpicId%>
+		</td>
+		<%--<td>
+			<%=CSVRow.UserEmail%>
+		</td>--%>
+		<td>
+			<%=CSVRow.BuggId%>
+		</td>
+		<td>
+			<%=CSVRow.AffectedVersionsString%>
+		</td>		
+	</tr>
+	<%
+		}	
+	%>
+
+</table>
+
+</form>
+</body>
\ No newline at end of file
diff --git a/Engine/Source/Programs/CrashReporter/CrashReportWebSite/Views/Crashes/Index.aspx b/Engine/Source/Programs/CrashReporter/CrashReportWebSite/Views/Crashes/Index.aspx
--- a/Engine/Source/Programs/CrashReporter/CrashReportWebSite/Views/Crashes/Index.aspx
+++ b/Engine/Source/Programs/CrashReporter/CrashReportWebSite/Views/Crashes/Index.aspx
@@ -9,7 +9,7 @@
 </asp:Content>
 
 <asp:Content ID="TitleContent" ContentPlaceHolderID="TitleContent" runat="server">
-Crash Reports
+[CR] Crashes
 </asp:Content>
 
 <asp:Content ID="ScriptContent"  ContentPlaceHolderID="ScriptContent" runat="server" >
diff --git a/Engine/Source/Programs/CrashReporter/CrashReportWebSite/Views/Dashboard/Index.aspx b/Engine/Source/Programs/CrashReporter/CrashReportWebSite/Views/Dashboard/Index.aspx
--- a/Engine/Source/Programs/CrashReporter/CrashReportWebSite/Views/Dashboard/Index.aspx
+++ b/Engine/Source/Programs/CrashReporter/CrashReportWebSite/Views/Dashboard/Index.aspx
@@ -19,7 +19,7 @@ table
 </asp:Content>
 
 <asp:Content ID="TitleContent" ContentPlaceHolderID="TitleContent" runat="server">
-	Crash Reports
+[CR] Dashboard
 </asp:Content>
 
 <asp:Content ID="ScriptContent" ContentPlaceHolderID="ScriptContent" runat="server" >
diff --git a/Engine/Source/Programs/CrashReporter/CrashReportWebSite/Views/Home/Index.aspx b/Engine/Source/Programs/CrashReporter/CrashReportWebSite/Views/Home/Index.aspx
--- a/Engine/Source/Programs/CrashReporter/CrashReportWebSite/Views/Home/Index.aspx
+++ b/Engine/Source/Programs/CrashReporter/CrashReportWebSite/Views/Home/Index.aspx
@@ -9,7 +9,7 @@
 </asp:Content>
 
 <asp:Content ID="TitleContent" ContentPlaceHolderID="TitleContent" runat="server">
-Crash Reporter - Home
+[CR] Home
 </asp:Content>
 
 <asp:Content ID="ScriptContent"  ContentPlaceHolderID="ScriptContent" runat="server" >
diff --git a/Engine/Source/Programs/CrashReporter/CrashReportWebSite/Views/Reports/Index.aspx b/Engine/Source/Programs/CrashReporter/CrashReportWebSite/Views/Reports/Index.aspx
--- a/Engine/Source/Programs/CrashReporter/CrashReportWebSite/Views/Reports/Index.aspx
+++ b/Engine/Source/Programs/CrashReporter/CrashReportWebSite/Views/Reports/Index.aspx
@@ -8,7 +8,7 @@
 </asp:Content>
 
 <asp:Content ID="TitleContent" ContentPlaceHolderID="TitleContent" runat="server">
-	Crash Reports
+[CR] Reports
 </asp:Content>
 
 <asp:Content ID="ScriptContent" ContentPlaceHolderID="ScriptContent" runat="server" >
diff --git a/Engine/Source/Programs/CrashReporter/CrashReportWebSite/Views/Shared/Site.Master b/Engine/Source/Programs/CrashReporter/CrashReportWebSite/Views/Shared/Site.Master
--- a/Engine/Source/Programs/CrashReporter/CrashReportWebSite/Views/Shared/Site.Master
+++ b/Engine/Source/Programs/CrashReporter/CrashReportWebSite/Views/Shared/Site.Master
@@ -38,6 +38,8 @@
 				&nbsp;|&nbsp;
 				<span><a href="/Reports/Index/">Reports</a></span>
 				&nbsp;|&nbsp;
+				<span><a href="/CSV/Index/">CSV</a></span>
+				&nbsp;|&nbsp;
 				<span><a href="/Dashboard/Index/">Dashboard</a></span>
 				&nbsp;|&nbsp;
 				<span><a href="[REDACTED_USER_HOME]/">Edit Users</a></span>
@@ -59,6 +61,10 @@
 			<asp:ContentPlaceHolder ID="ReportsContent" runat="server" />
 		</div>
 
+		<div style="clear:both;">
+			<asp:ContentPlaceHolder ID="CSVContent" runat="server" />
+		</div>
+
 		<div id="footer" style="clear:both;">
 		</div>
 
diff --git a/Engine/Source/Programs/CrashReporter/CrashReportWebSite/Views/Users/Index.aspx b/Engine/Source/Programs/CrashReporter/CrashReportWebSite/Views/Users/Index.aspx
--- a/Engine/Source/Programs/CrashReporter/CrashReportWebSite/Views/Users/Index.aspx
+++ b/Engine/Source/Programs/CrashReporter/CrashReportWebSite/Views/Users/Index.aspx
@@ -8,7 +8,7 @@
 </asp:Content>
 
 <asp:Content ID="TitleContent" ContentPlaceHolderID="TitleContent" runat="server">
-    Edit the groups the users belong to
+[CR] Edit the groups the users belong to
 </asp:Content>
 
 <asp:Content ID="ScriptContent" ContentPlaceHolderID="ScriptContent" runat="server">
