CrashReportWWW - CSV controller, revision 2

diff --git a/Engine/Source/Programs/CrashReporter/CrashReportWebSite/Controllers/CSVController.cs b/Engine/Source/Programs/CrashReporter/CrashReportWebSite/Controllers/CSVController.cs
--- a/Engine/Source/Programs/CrashReporter/CrashReportWebSite/Controllers/CSVController.cs
+++ b/Engine/Source/Programs/CrashReporter/CrashReportWebSite/Controllers/CSVController.cs
@@ -14,23 +14,8 @@ using Tools.CrashReporter.CrashReportWebSite.Models;
 using Tools.CrashReporter.CrashReportWebSite.Properties;
 using Tools.DotNETCommon;
 
-
 namespace Tools.CrashReporter.CrashReportWebSite.Controllers
 {
-	// 
-	// How hard would it be for you to create a CSV of:
-	// Epic / Machine ID, 
-	// Buggs ID, 
-	// EngineVersion, 
-	// 
-	// # of crashes with that buggs ID for 
-	//	this engine version and	user
-
-	//E-32chars
-	//M-32chars
-	// 
-	//Engine version: 4.8.1-cl+branch
-	// 
 	// Each time the controller is run, create a new CSV file in the specified folder and add the link and the top of the site,
 	// so it could be easily downloaded
 	// Removed files older that 7 days
@@ -43,15 +28,9 @@ namespace Tools.CrashReporter.CrashReportWebSite.Controllers
 		/// <summary>Bugg Id.</summary>
 		public int BuggId = 0;
 		/// <summary>Engine version.</summary>
-		public string AffectedVersionsString = "";
-		/// <summary># of crashes with that buggs ID for this engine version and user.</summary>
-		public int NumberOfCrashes = 0;
-
-		/// <summary>User's email.</summary>
-		public string UserEmail = "";
-
-		//// <summary>Reference to a bugg.</summary>
-		//public Bugg BuggRef = null;
+		public string EngineVersion = "";
+		/// <summary>Time of a crash.</summary>
+		public DateTime TimeOfCrash;
 	}
 
 	/// <summary>
@@ -62,7 +41,7 @@ namespace Tools.CrashReporter.CrashReportWebSite.Controllers
 		/// <summary>Fake id for all user groups</summary>
 		public static readonly int AllUserGroupId = -1;
 
-		static readonly int MinNumberOfCrashes = 2;
+		//static readonly int MinNumberOfCrashes = 2;
 
 		/// <summary>
 		/// Retrieve all Buggs matching the search criteria.
@@ -74,13 +53,6 @@ namespace Tools.CrashReporter.CrashReportWebSite.Controllers
 			BuggRepository BuggsRepo = new BuggRepository();
 			CrashRepository CrashRepo = new CrashRepository();
 
-			// No Jira support at this moment
-			// Enumerate JIRA projects if needed.
-			// https://jira.ol.epicgames.net//rest/api/2/project
-			var JC = JiraConnection.Get();
-			var JiraComponents = JC.GetNameToComponents();
-			var JiraVersions = JC.GetNameToVersions();
-
 			using( FAutoScopedLogTimer LogTimer = new FAutoScopedLogTimer( this.GetType().ToString() ) )
 			{
 				string AnonymousGroup = "Anonymous";
@@ -90,185 +62,127 @@ namespace Tools.CrashReporter.CrashReportWebSite.Controllers
 				int AnonymousID = AnonumousIDs.First();
 				HashSet<string> UserNamesForUserGroup = FRepository.Get( BuggsRepo ).GetUserNamesFromGroupName( AnonymousGroup );
 
-				FormData.DateFrom = FormData.DateTo.AddDays( -1 );
+				// Enable to narrow results and improve debugging performance.
+				//FormData.DateFrom = FormData.DateTo.AddDays( -1 );
+				FormData.DateTo = FormData.DateTo.AddDays( 1 );
 
-				var Crashes = CrashRepo
-					.FilterByDate( CrashRepo.ListAll(), FormData.DateFrom, FormData.DateTo )
+				var FilteringQueryJoin = CrashRepo
+					.ListAll()
+					.Where( c => c.EpicAccountId != "" )
 					// Only crashes and asserts
-					.Where( Crash => Crash.CrashType == 1 || Crash.CrashType == 2 )
+					.Where( c => c.CrashType == 1 || c.CrashType == 2 )
 					// Only anonymous user
-					.Where( Crash => Crash.UserNameId.Value == AnonymousID )
-					// We need Epic ID
-					// or Machine ID to identify the owner of a crash
-					//.Where( Crash => !string.IsNullOrEmpty( Crash.EpicAccountId ) || !string.IsNullOrEmpty( Crash.MachineId ) )
-					.Where( Crash => !string.IsNullOrEmpty( Crash.EpicAccountId ) )
-					.Where( Crash => !string.IsNullOrEmpty( Crash.Pattern ) )
-					.Select( Crash => new
-					{
-						ID = Crash.Id,
-						TimeOfCrash = Crash.TimeOfCrash.Value,
-						//UserID = Crash.UserNameId.Value, 
-						BuildVersion = Crash.BuildVersion,
-						JIRA = Crash.Jira,
-						Platform = Crash.PlatformName,
-						FixCL = Crash.FixedChangeList,
-						BuiltFromCL = Crash.BuiltFromCL,
-						Pattern = Crash.Pattern,
-						MachineID = Crash.MachineId,
-						Branch = Crash.Branch,
-						Description = Crash.Description,
-						RawCallStack = Crash.RawCallStack,
-						EpicAccountId = Crash.EpicAccountId,
-					} )
-					.ToList();
-				int NumCrashes = Crashes.Count;
+					.Where( c => c.UserNameId == AnonymousID )
+					// Filter be date
+					.Where( c => c.TimeOfCrash > FormData.DateFrom && c.TimeOfCrash < FormData.DateTo )
+					.Join
+					(
+						CrashRepo.Context.Buggs_Crashes,
+						c => c,
+						bc => bc.Crash,
+						( c, bc ) => new
+						{
+							EpicId = c.EpicAccountId,
+							BuggId = bc.BuggId,
+							TimeOfCrash = c.TimeOfCrash.Value,
+							EngineVersion = c.BuildVersion,
+						}
+					);
 
-				// Get all users
-				// SLOW
-				//var Users = FRepository.Get( BuggsRepo ).GetAnalyticsUsers().ToDictionary( X => X.EpicAccountId, Y => Y );
+				var FilteringQueryCrashes = CrashRepo
+					.ListAll()
+					.Where( c => c.EpicAccountId != "" )
+					// Only crashes and asserts
+					.Where( c => c.CrashType == 1 || c.CrashType == 2 )
+					// Only anonymous user
+					.Where( c => c.UserNameId == AnonymousID );
 
-				// Total # of UNIQUE (Anonymous) crashes in timeframe
-				Dictionary<string, int> PatternToCount = new Dictionary<string, int>();
-				HashSet<string> UniqueEpicIds = new HashSet<string>();
-				Dictionary<string, FCSVRow> EpicIdBuggToCSVRow = new Dictionary<string, FCSVRow>();
+				int TotalCrashes = CrashRepo
+					.ListAll()
+					.Count();
 
-				// Iterate all crashes and generate all patterns, and EpicIds
-				foreach( var Crash in Crashes )
-				{
-					if( string.IsNullOrEmpty( Crash.Pattern ) )
-					{
-						continue;
-					}
+				int TotalCrashesYearToDate = CrashRepo
+					.ListAll()
+					// Year to date
+					.Where( c => c.TimeOfCrash > new DateTime( DateTime.UtcNow.Year, 1, 1 ) )
+					.Count();
 
-					bool bAdd = !PatternToCount.ContainsKey( Crash.Pattern );
-					if (bAdd)
-					{
-						PatternToCount.Add( Crash.Pattern, 1 );
-					}
-					else
-					{
-						PatternToCount[Crash.Pattern]++;
-					}
+				var CrashesFilteredWithDateQuery = FilteringQueryCrashes
+					// Filter be date
+					.Where( c => c.TimeOfCrash > FormData.DateFrom && c.TimeOfCrash < FormData.DateTo );
 
-					UniqueEpicIds.Add( Crash.EpicAccountId );
-				}
+				int CrashesFilteredWithDate = CrashesFilteredWithDateQuery
+					.Count();
 
-				// Select only patterns with at least 2 crashes.
-				var PatternToCountOrdered = PatternToCount.OrderByDescending( X => X.Value ).ToList();
-				var PatternAndCount = PatternToCountOrdered.Where( X => X.Value > MinNumberOfCrashes ).ToDictionary( X => X.Key, Y => Y.Value );
+				int CrashesYearToDateFiltered = FilteringQueryCrashes
+					// Year to date
+					.Where( c => c.TimeOfCrash > new DateTime( DateTime.UtcNow.Year, 1, 1 ) )
+					.Count();
 
-				// Get all buggs associated with crashes.
-				var RealBuggs = new Dictionary<string, Bugg>();
-				foreach (var PatternIt in PatternAndCount)
-				{
-					var RealBugg = BuggsRepo.Context.Buggs.Where( Bugg => Bugg.Pattern == PatternIt.Key ).FirstOrDefault();
-					if (RealBugg != null)
-					{
-						var CrashesForBugg = Crashes.Where( Crash => Crash.Pattern == PatternIt.Key ).ToList();
+				int AffectedUsersFiltered = FilteringQueryCrashes
+					.Select( c => c.EpicAccountId )
+					.Distinct()
+					.Count();
 
-						// Convert anonymous to full type;
-						var FullCrashesForBugg = new List<Crash>( CrashesForBugg.Count );
-						foreach (var Anon in CrashesForBugg)
-						{
-							FullCrashesForBugg.Add( new Crash()
-							{
-								ID = Anon.ID,
-								TimeOfCrash = Anon.TimeOfCrash,
-								BuildVersion = Anon.BuildVersion,
-								Jira = Anon.JIRA,
-								Platform = Anon.Platform,
-								FixCL = Anon.FixCL,
-								BuiltFromCL = Anon.BuiltFromCL,
-								Pattern = Anon.Pattern,
-								MachineId = Anon.MachineID,
-								Branch = Anon.Branch,
-								Description = Anon.Description,
-								RawCallStack = Anon.RawCallStack,
-							} );
-						}
+				int UniqueCrashesFiltered = FilteringQueryCrashes
+					.Select( c => c.Pattern )
+					.Distinct()
+					.Count();
 
-						RealBugg.PrepareBuggForJira( FullCrashesForBugg );
-						RealBuggs.Add( PatternIt.Key, RealBugg );
-					}
-					else
-					{
-						//FLogger.Global.WriteEvent( "Bugg for pattern " + PatternIt + " not found" );
-					}
-				}
 
+				int NumCrashes = FilteringQueryJoin.Count();
 
-				// Iterate all crashes and link EpicId with Bugg, and calculate #
-				foreach( var Crash in Crashes )
-				{
-					Bugg Bugg = null;
-					RealBuggs.TryGetValue( Crash.Pattern, out Bugg );
+				// Get all users
+				// SLOW
+				//var Users = FRepository.Get( BuggsRepo ).GetAnalyticsUsers().ToDictionary( X => X.EpicAccountId, Y => Y );
 
-					if( Bugg != null )
-					{
-						string EpicIdBugg = Crash.EpicAccountId + "%" + Bugg.Id;
-						bool bAdd = !EpicIdBuggToCSVRow.ContainsKey( EpicIdBugg );
-						if (bAdd)
-						{
-							var Row = new FCSVRow 
-							{ 
-								EpicId = Crash.EpicAccountId,
-								BuggId = Bugg.Id,
-								AffectedVersionsString = Bugg.GetAffectedVersions(),
-								NumberOfCrashes = 1,
-								//UserEmail = Users[Crash.EpicAccountId].UserEmail
-							};
-							EpicIdBuggToCSVRow.Add( EpicIdBugg, Row );
-						}
-						else
-						{
-							EpicIdBuggToCSVRow[EpicIdBugg].NumberOfCrashes++;
-						}
-					}
-				}
+				
 
 				// Export data to the file.
-				string CSVPathname = Path.Combine( Settings.Default.CrashReporterCSV, DateTime.UtcNow.ToString( "yyyy-MM-dd.HH-mm-ss" ) ) + ".csv";
+				string CSVPathname = Path.Combine( Settings.Default.CrashReporterCSV, DateTime.UtcNow.ToString( "yyyy-MM-dd.HH-mm-ss" ) );
+				CSVPathname += string
+					.Format( "__[{0}---{1}]__{2}", 
+					FormData.DateFrom.ToString( "yyyy-MM-dd" ), 
+					FormData.DateTo.ToString( "yyyy-MM-dd" ), 
+					NumCrashes ) 
+					+ ".csv";
+
 				string ServerPath = Server.MapPath( CSVPathname );
 				var CSVFile = new StreamWriter( ServerPath, true, Encoding.UTF8 );
 
-				CSVFile.WriteLine( "EpicId; BuggId; AffectedVersion; # of crashes;" );
-
-				foreach (var Row in EpicIdBuggToCSVRow)
+				using (FAutoScopedLogTimer ExportToCSVTimer = new FAutoScopedLogTimer( "ExportToCSV" ))
 				{
-					if (Row.Value.NumberOfCrashes > MinNumberOfCrashes)
+					CSVFile.WriteLine( "TimeOfCrash; EpicId; BuggId; EngineVersion;" );
+
+					foreach (var Row in FilteringQueryJoin)
 					{
-						CSVFile.WriteLine( "{0};{1};{2};{3};", Row.Value.EpicId, Row.Value.BuggId, Row.Value.AffectedVersionsString, Row.Value.NumberOfCrashes );
+						CSVFile.WriteLine( "{0};{1};{2};{3};", Row.TimeOfCrash, Row.EpicId, Row.BuggId, Row.EngineVersion );
 					}
-				}
 
-				CSVFile.Flush();
-				CSVFile.Close();
-				CSVFile = null;
+					CSVFile.Flush();
+					CSVFile.Close();
+					CSVFile = null;
+				}
 
-				List<FCSVRow> CSVRows = EpicIdBuggToCSVRow
-					.OrderByDescending( X => X.Value.NumberOfCrashes )
-					.Select( X => X.Value )
-					.Where( X => X.NumberOfCrashes > MinNumberOfCrashes )
+				List<FCSVRow> CSVRows = FilteringQueryJoin
+					.OrderByDescending( X => X.TimeOfCrash )
+					.Take( 100 )
+					.Select( c => new FCSVRow { TimeOfCrash = c.TimeOfCrash, EpicId = c.EpicId, BuggId = c.BuggId, EngineVersion = c.EngineVersion } )
 					.ToList();
 				
-				// Total # of unique anonymous crashes.
-				int TotalUniqueAnonymousCrashes = PatternToCount.Count;
-
-				// Total # of AFFECTED USERS (Anonymous) in timeframe
-				int TotalAffectedUsers = UniqueEpicIds.Count;
-
-				// Total # of ALL (Anonymous) crashes in timeframe
-				int TotalAnonymousCrashes = NumCrashes;
-
 				return new CSV_ViewModel
 				{
 					CSVRows = CSVRows,
 					CSVPathname = CSVPathname,
 					DateFrom = (long)( FormData.DateFrom - CrashesViewModel.Epoch ).TotalMilliseconds,
 					DateTo = (long)( FormData.DateTo - CrashesViewModel.Epoch ).TotalMilliseconds,
-					TotalAffectedUsers = TotalAffectedUsers,
-					TotalAnonymousCrashes = TotalAnonymousCrashes,
-					TotalUniqueAnonymousCrashes = TotalUniqueAnonymousCrashes
+					DateTimeFrom = FormData.DateFrom,
+					DateTimeTo = FormData.DateTo,
+					AffectedUsersFiltered = AffectedUsersFiltered,
+					UniqueCrashesFiltered = UniqueCrashesFiltered,
+					CrashesFilteredWithDate = CrashesFilteredWithDate,
+					TotalCrashes = TotalCrashes,
+					TotalCrashesYearToDate = TotalCrashesYearToDate,
 				};
 			}
 		}
diff --git a/Engine/Source/Programs/CrashReporter/CrashReportWebSite/Models/CSVViewModel.cs b/Engine/Source/Programs/CrashReporter/CrashReportWebSite/Models/CSVViewModel.cs
--- a/Engine/Source/Programs/CrashReporter/CrashReportWebSite/Models/CSVViewModel.cs
+++ b/Engine/Source/Programs/CrashReporter/CrashReportWebSite/Models/CSVViewModel.cs
@@ -17,20 +17,32 @@ namespace Tools.CrashReporter.CrashReportWebSite.Models
 		/// <summary>A container of sorted CSV rows.</summary>
 		public List<FCSVRow> CSVRows { get; set; }
 
-		/// <summary>The date of the earliest crash in a Bugg.</summary>
+		/// <summary>The date of the earliest crash.</summary>
 		public long DateFrom { get; set; }
 
-		/// <summary>The date of the most recent crash in a Bugg.</summary>
+		/// <summary>The date of the most recent crash.</summary>
 		public long DateTo { get; set; }
 
-		/// <summary></summary>
-		public int TotalAnonymousCrashes { get; set; }
+		/// <summary>The date of the earliest crash.</summary>
+		public DateTime DateTimeFrom { get; set; }
 
-		/// <summary></summary>
-		public int TotalUniqueAnonymousCrashes { get; set; }
+		/// <summary>The date of the most recent crash.</summary>
+		public DateTime DateTimeTo { get; set; }
 
-		/// <summary></summary>
-		public int TotalAffectedUsers { get; set; }
+		/// <summary>Unique Anonymous Crashes</summary>
+		public int UniqueCrashesFiltered { get; set; }
+
+		/// <summary>Total Affected Users</summary>
+		public int AffectedUsersFiltered { get; set; }
+
+		/// <summary>Total crashes.</summary>
+		public int TotalCrashes { get; set; }
+
+		/// <summary>Total crashes year to date.</summary>
+		public int TotalCrashesYearToDate { get; set; }
+
+		/// <summary>Total crashes.</summary>
+		public int CrashesFilteredWithDate { get; set; }
 
 		/// <summary>Pathname where the CSV file is saved.</summary>
 		public string CSVPathname { get; set; }
@@ -47,6 +59,12 @@ namespace Tools.CrashReporter.CrashReportWebSite.Models
 			return Path.GetDirectoryName( CSVPathname );
 		}
 
+		/// <summary>Public directory where the CSV files are saved.</summary>
+		public string GetCSVPublicDirectory()
+		{
+			return @"\\DEVWEB-02\CrashReporterCSV";
+		}
+
 		/// <summary>Time spent in generating this site, formatted as a string.</summary>
 		public string GenerationTime { get; set; }
 
diff --git a/Engine/Source/Programs/CrashReporter/CrashReportWebSite/Views/CSV/ViewCSV.ascx b/Engine/Source/Programs/CrashReporter/CrashReportWebSite/Views/CSV/ViewCSV.ascx
--- a/Engine/Source/Programs/CrashReporter/CrashReportWebSite/Views/CSV/ViewCSV.ascx
+++ b/Engine/Source/Programs/CrashReporter/CrashReportWebSite/Views/CSV/ViewCSV.ascx
@@ -29,21 +29,49 @@ table, table tr, table th, table td
 <table >
 	<tr>
 		<td>
-			TotalAnonymousCrashes:&nbsp;<%=Model.TotalAnonymousCrashes %>
+			Total Crashes:&nbsp;<%=Model.TotalCrashes %>
+		</td>
+	</tr>
+	<tr>
+		<td>
+			Total Crashes YearToDate:&nbsp;<%=Model.TotalCrashesYearToDate %>
 		</td>
 	</tr>
 
 	<tr>
 		<td>
-			TotalUniqueAnonymousCrashes:&nbsp;<%=Model.TotalUniqueAnonymousCrashes %>
+			<strong>
+			Filtered = CrashOrAssert, Anonymous, EpicId != 0, BuggId != 0,
+			</strong>
 		</td>
 	</tr>
 
 	<tr>
 		<td>
-			TotalAffectedUsers:&nbsp;<%=Model.TotalAffectedUsers %>
+			Total Unique Crashes:&nbsp;<%=Model.UniqueCrashesFiltered %>
+		</td>
+	</tr>
+	<tr>
+		<td>
+			Total Affected Users:&nbsp;<%=Model.AffectedUsersFiltered %>
 		</td>
 	</tr>
+		
+
+	<tr>
+		<td>
+			<strong>
+			Filtered = CrashOrAssert, Anonymous, EpicId != 0, BuggId != 0, TimeOfCrash within the time range, from <%=Model.DateTimeFrom.ToString()%> to <%=Model.DateTimeTo.ToString()%>
+			</strong>
+		</td>
+	</tr>
+
+	<tr>
+		<td>
+			Crashes:&nbsp;<%=Model.CrashesFilteredWithDate %>
+		</td>
+	</tr>
+
 
 	<tr>
 		<td>
@@ -56,14 +84,20 @@ table, table tr, table th, table td
 			Directory:&nbsp;<a href="<%=Model.GetCSVDirectory()%>" target="_blank"><%=Model.GetCSVDirectory()%></a>
 		</td>
 	</tr>
+
+	<tr>
+		<td title="Copy the link and open in the explorer">
+			Public Directory:&nbsp;<strong><%=Model.GetCSVPublicDirectory() %></strong>
+		</td>
+	</tr>
 </table>
 <br />
 
-<p>Displaying top 100 records, ordered by # of crashes</p>
+<p>Displaying top 100 records, ordered by the date of a crash</p>
 <table style="width: 100%">
 	<tr>
 		<th>
-			#
+			Time of crash
 		</th>
 		<th>
 			EpicId
@@ -75,7 +109,7 @@ table, table tr, table th, table td
 			BuggId
 		</th>
 		<th>
-			Versions Affected
+			Engine version
 		</th>
 	</tr>
 
@@ -90,7 +124,7 @@ table, table tr, table th, table td
 	%>
 	<tr>
 		<td>
-			<%=CSVRow.NumberOfCrashes%>
+			<%=CSVRow.TimeOfCrash.ToString()%>
 		</td>
 		<td>
 			<%=CSVRow.EpicId%>
@@ -102,7 +136,7 @@ table, table tr, table th, table td
 			<%=CSVRow.BuggId%>
 		</td>
 		<td>
-			<%=CSVRow.AffectedVersionsString%>
+			<%=CSVRow.EngineVersion%>
 		</td>		
 	</tr>
 	<%
