CrashReportWWW - CSV controller (fixed incorrect engine version for older crashes)

diff --git a/Engine/Source/Programs/CrashReporter/CrashReportWebSite/Controllers/CSVController.cs b/Engine/Source/Programs/CrashReporter/CrashReportWebSite/Controllers/CSVController.cs
--- a/Engine/Source/Programs/CrashReporter/CrashReportWebSite/Controllers/CSVController.cs
+++ b/Engine/Source/Programs/CrashReporter/CrashReportWebSite/Controllers/CSVController.cs
@@ -156,7 +156,12 @@ namespace Tools.CrashReporter.CrashReportWebSite.Controllers
 
 					foreach (var Row in FilteringQueryJoin)
 					{
-						CSVFile.WriteLine( "{0};{1};{2};{3};", Row.TimeOfCrash, Row.EpicId, Row.BuggId, Row.EngineVersion );
+						var BVParts = Row.EngineVersion.Split( new char[] { '.' }, StringSplitOptions.RemoveEmptyEntries );
+						if (BVParts.Length > 2 && BVParts[0] != "0")
+						{
+							string CleanEngineVersion = string.Format( "{0}.{1}.{2}", BVParts[0], BVParts[1], BVParts[2] );
+							CSVFile.WriteLine( "{0};{1};{2};{3};", Row.TimeOfCrash, Row.EpicId, Row.BuggId, CleanEngineVersion );
+						}
 					}
 
 					CSVFile.Flush();
diff --git a/Engine/Source/Programs/CrashReporter/CrashReportWebSite/Views/CSV/Index.aspx b/Engine/Source/Programs/CrashReporter/CrashReportWebSite/Views/CSV/Index.aspx
--- a/Engine/Source/Programs/CrashReporter/CrashReportWebSite/Views/CSV/Index.aspx
+++ b/Engine/Source/Programs/CrashReporter/CrashReportWebSite/Views/CSV/Index.aspx
@@ -57,7 +57,7 @@
 	</p>
 
 	 <p>
-		 <input type="submit" value="Search" class='SearchButton' />
+		 <input type="submit" value="Generate" class='SearchButton' />
 	 </p>
 	
 	<%} %>
