CrashReportWWW - CSV controller (a missing file)

diff --git a/Engine/Source/Programs/CrashReporter/CrashReportWebSite/Web.config b/Engine/Source/Programs/CrashReporter/CrashReportWebSite/Web.config
--- a/Engine/Source/Programs/CrashReporter/CrashReportWebSite/Web.config
+++ b/Engine/Source/Programs/CrashReporter/CrashReportWebSite/Web.config
@@ -96,6 +96,9 @@
       <setting name="CrashReporterVideos" serializeAs="String">
         <value>/CrashReporterVideos/</value>
       </setting>
+      <setting name="CrashReporterCSV" serializeAs="String">
+        <value>/CrashReporterCSV/</value>
+      </setting>
     </Tools.CrashReporter.CrashReportWebSite.Properties.Settings>
   </applicationSettings>
 </configuration>
\ No newline at end of file
