CrashFrontend - Removed from the callstack following entries: EngineCrashHandler

diff --git a/Engine/Source/Programs/CrashReporter/CrashReportWebSite/Models/CrashDataFormater.cs b/Engine/Source/Programs/CrashReporter/CrashReportWebSite/Models/CrashDataFormater.cs
--- a/Engine/Source/Programs/CrashReporter/CrashReportWebSite/Models/CrashDataFormater.cs
+++ b/Engine/Source/Programs/CrashReporter/CrashReportWebSite/Models/CrashDataFormater.cs
@@ -394,6 +394,7 @@ namespace Tools.CrashReporter.CrashReportWebSite.Models
 					"FMsg::Logf",
 					"ReportCrash",
 					"NewReportEnsure",
+					"EngineCrashHandler", // Generic crash handler for all platforms
 				} );
 
 				bool Contains = FuncsToRemove.Contains( FuncName, new CustomFuncComparer() );
