CrashFrontend - Ensure are also set if there is a call to NewReportEnsure

diff --git a/Engine/Source/Programs/CrashReporter/CrashReportWebSite/Models/CrashRepository.cs b/Engine/Source/Programs/CrashReporter/CrashReportWebSite/Models/CrashRepository.cs
--- a/Engine/Source/Programs/CrashReporter/CrashReportWebSite/Models/CrashRepository.cs
+++ b/Engine/Source/Programs/CrashReporter/CrashReportWebSite/Models/CrashRepository.cs
@@ -614,11 +614,15 @@ namespace Tools.CrashReporter.CrashReportWebSite.Models
 			NewCrash.CrashType = 1;
 			if (NewCrash.RawCallStack != null)
 			{
-				if (NewCrash.RawCallStack.Contains("FDebug::AssertFailed()"))
+				if (NewCrash.RawCallStack.Contains( "FDebug::AssertFailed" ))
 				{
 					NewCrash.CrashType = 2;
 				}
-				else if (NewCrash.RawCallStack.Contains("FDebug::EnsureFailed()"))
+				else if (NewCrash.RawCallStack.Contains( "FDebug::EnsureFailed" ))
+				{
+					NewCrash.CrashType = 3;
+				}
+				else if (NewCrash.RawCallStack.Contains( "NewReportEnsure" ))
 				{
 					NewCrash.CrashType = 3;
 				}
