CrashFrontend - Misc refactor

diff --git a/Engine/Source/Programs/CrashReporter/CrashReportWebSite/Models/CrashReport.cs b/Engine/Source/Programs/CrashReporter/CrashReportWebSite/Models/CrashReport.cs
--- a/Engine/Source/Programs/CrashReporter/CrashReportWebSite/Models/CrashReport.cs
+++ b/Engine/Source/Programs/CrashReporter/CrashReportWebSite/Models/CrashReport.cs
@@ -667,7 +667,7 @@ namespace Tools.CrashReporter.CrashReportWebSite.Models
 					Module = CallStack.GetModuleName();
 					try
 					{
-						foreach( CallStackEntry Entry in CallStack.CallStackEntries.Take( 64 ) )
+						foreach (CallStackEntry Entry in CallStack.CallStackEntries.Take( CallStackContainer.MaxLinesToParse ))
 						{
 							FunctionCall CurrentFunctionCall = new FunctionCall();
 
