CrashReportLinq - Updated RepairDatabase program (removed bad ensures from the database)

diff --git a/Engine/Source/Programs/CrashReporter/CrashReportWebSite/Queries/Fix-RepairDatabase.linq b/Engine/Source/Programs/CrashReporter/CrashReportWebSite/Queries/Fix-RepairDatabase.linq
--- a/Engine/Source/Programs/CrashReporter/CrashReportWebSite/Queries/Fix-RepairDatabase.linq
+++ b/Engine/Source/Programs/CrashReporter/CrashReportWebSite/Queries/Fix-RepairDatabase.linq
@@ -321,6 +321,43 @@ void SetNoPatternForInvalidPattern()
 	}
 }
 
+void SetNoPatternForEnsureForCrashes()
+{
+	WriteLine( string.Format( "EnsureButAssert" ) );
+	//UGameViewportClient::MouseLeave
+	var EnsureAssertBatch = new List<Crashes>(NUM_OPS_PER_BATCH);
+	DateTime StartDate = Date;
+	DateTime EndDate = Date.Add( Tick );
+	while( EndDate < DateTime.UtcNow )
+	{
+		var EnsureAssertCrashes = Crashes
+		.Where( c => c.TimeOfCrash > StartDate && c.TimeOfCrash <= EndDate )
+		.Where( c => c.CrashType == 1 || c.CrashType == 2 )
+		.Where( c => c.RawCallStack.Contains("FDebug::Ensure") || c.RawCallStack.Contains("NewReportEnsure") )
+		.ToList();
+	
+		WriteLine( string.Format( "EnsureAssertCrashes: {0} {1} -> {2}", EnsureAssertCrashes.Count, StartDate.ToString("yyyy-MM-dd"), EndDate.ToString("yyyy-MM-dd") ) );
+		StartDate = EndDate;
+	 	EndDate = EndDate.Add( Tick );
+	
+		for (int n = 0; n < EnsureAssertCrashes.Count; n ++)
+		{
+			EnsureAssertBatch.Add(EnsureAssertCrashes[n]);		
+			EnsureAssertCrashes[n].Pattern = null;
+			EnsureAssertCrashes[n].CrashType = 3;
+			EnsureAssertCrashes[n].Branch = "INVALID";
+			
+			if (EnsureAssertBatch.Count == NUM_OPS_PER_BATCH)
+			{
+				SumbitChanges(EnsureAssertBatch);
+			}
+		}
+	}
+	
+	// Last batch
+	SumbitChanges(EnsureAssertBatch);
+}
+
 void Main()
 {	
 	int NumCrashes = Crashes.Count();
@@ -330,6 +367,7 @@ void Main()
 	WriteLine( string.Format( "NumCrashes: {0} MaxCrashId: {1} Diff: {2}", NumCrashes, MaxCrashId, MaxCrashId-NumCrashes ) );
 	WriteLine( string.Format( "NumBuggs: {0} MaxBuggId: {1} Diff: {2}", NumBuggs, MaxBuggId, MaxBuggId-NumBuggs ) );
 	
+	SetNoPatternForEnsureForCrashes();
 	SetNoPatternForNoBranchAll();
 	SetNoPatternForInvalidPattern();
 	DeleteNoPatternAll();
