CrashFrontend - Added crash type to CSV report

diff --git a/Engine/Source/Programs/CrashReporter/CrashReportWebSite/Controllers/CSVController.cs b/Engine/Source/Programs/CrashReporter/CrashReportWebSite/Controllers/CSVController.cs
--- a/Engine/Source/Programs/CrashReporter/CrashReportWebSite/Controllers/CSVController.cs
+++ b/Engine/Source/Programs/CrashReporter/CrashReportWebSite/Controllers/CSVController.cs
@@ -31,6 +31,26 @@ namespace Tools.CrashReporter.CrashReportWebSite.Controllers
 		public string EngineVersion = "";
 		/// <summary>Time of a crash.</summary>
 		public DateTime TimeOfCrash;
+		/// <summary>Crash type.</summary>
+		public short CrashType;
+
+		/// <summary>Returns crash type as a string.</summary>
+		public string GetCrashTypeAsString()
+		{
+			if (CrashType == 1)
+			{
+				return "Crash";
+			}
+			else if (CrashType == 2)
+			{
+				return "Assert";
+			}
+			else if (CrashType == 3)
+			{
+				return "Ensure";
+			}
+			return "Unknown";
+		}
 	}
 
 	/// <summary>
@@ -86,6 +106,7 @@ namespace Tools.CrashReporter.CrashReportWebSite.Controllers
 							BuggId = bc.BuggId,
 							TimeOfCrash = c.TimeOfCrash.Value,
 							EngineVersion = c.BuildVersion,
+							CrashType = c.CrashType.Value,
 						}
 					);
 
@@ -152,7 +173,7 @@ namespace Tools.CrashReporter.CrashReportWebSite.Controllers
 
 				using (FAutoScopedLogTimer ExportToCSVTimer = new FAutoScopedLogTimer( "ExportToCSV" ))
 				{
-					CSVFile.WriteLine( "TimeOfCrash; EpicId; BuggId; EngineVersion;" );
+					CSVFile.WriteLine( "TimeOfCrash; EpicId; BuggId; EngineVersion; CrashType; " );
 
 					foreach (var Row in FilteringQueryJoin)
 					{
@@ -160,7 +181,12 @@ namespace Tools.CrashReporter.CrashReportWebSite.Controllers
 						if (BVParts.Length > 2 && BVParts[0] != "0")
 						{
 							string CleanEngineVersion = string.Format( "{0}.{1}.{2}", BVParts[0], BVParts[1], BVParts[2] );
-							CSVFile.WriteLine( "{0};{1};{2};{3};", Row.TimeOfCrash, Row.EpicId, Row.BuggId, CleanEngineVersion );
+							CSVFile.WriteLine( "{0};{1};{2};{3};{4};", 
+								Row.TimeOfCrash, 
+								Row.EpicId, 
+								Row.BuggId, 
+								CleanEngineVersion,
+								Row.CrashType == 1 ? "Crash" : "Assert" );
 						}
 					}
 
@@ -172,7 +198,7 @@ namespace Tools.CrashReporter.CrashReportWebSite.Controllers
 				List<FCSVRow> CSVRows = FilteringQueryJoin
 					.OrderByDescending( X => X.TimeOfCrash )
 					.Take( 100 )
-					.Select( c => new FCSVRow { TimeOfCrash = c.TimeOfCrash, EpicId = c.EpicId, BuggId = c.BuggId, EngineVersion = c.EngineVersion } )
+					.Select( c => new FCSVRow { TimeOfCrash = c.TimeOfCrash, EpicId = c.EpicId, BuggId = c.BuggId, EngineVersion = c.EngineVersion, CrashType = c.CrashType } )
 					.ToList();
 				
 				return new CSV_ViewModel
diff --git a/Engine/Source/Programs/CrashReporter/CrashReportWebSite/Views/CSV/ViewCSV.ascx b/Engine/Source/Programs/CrashReporter/CrashReportWebSite/Views/CSV/ViewCSV.ascx
--- a/Engine/Source/Programs/CrashReporter/CrashReportWebSite/Views/CSV/ViewCSV.ascx
+++ b/Engine/Source/Programs/CrashReporter/CrashReportWebSite/Views/CSV/ViewCSV.ascx
@@ -111,6 +111,9 @@ table, table tr, table th, table td
 		<th>
 			Engine version
 		</th>
+		<th>
+			Crash type
+		</th>
 	</tr>
 
 	<%
@@ -137,7 +140,10 @@ table, table tr, table th, table td
 		</td>
 		<td>
 			<%=CSVRow.EngineVersion%>
-		</td>		
+		</td>	
+		<td>
+			<%=CSVRow.GetCrashTypeAsString()%>
+		</td>	
 	</tr>
 	<%
 		}	
