CrashFrontend - CSV report exports following data from the database: GameName, TimeOfCrash, BuiltFromCL, PlatformName, EngineMode, MachineId, Module, EngineVersion, Jira, Branch, CrashType, EpicId, BuggId

diff --git a/Engine/Source/Programs/CrashReporter/CrashReportWebSite/Controllers/CSVController.cs b/Engine/Source/Programs/CrashReporter/CrashReportWebSite/Controllers/CSVController.cs
--- a/Engine/Source/Programs/CrashReporter/CrashReportWebSite/Controllers/CSVController.cs
+++ b/Engine/Source/Programs/CrashReporter/CrashReportWebSite/Controllers/CSVController.cs
@@ -20,19 +20,52 @@ namespace Tools.CrashReporter.CrashReportWebSite.Controllers
 	// so it could be easily downloaded
 	// Removed files older that 7 days
 
+	/*
+	GameName
+	TimeOfCrash
+	BuiltFromCL
+	PlatformName
+	EngineMode
+	MachineId
+	Module
+	BuildVersion
+	Jira
+	Branch
+	CrashType
+	EpicId
+	BuggId
+	*/
+
+
 	/// <summary>CSV data row.</summary>
 	public class FCSVRow
 	{
-		/// <summary>Epic Id.</summary>
-		public string EpicId = "";
-		/// <summary>Bugg Id.</summary>
-		public int BuggId = 0;
-		/// <summary>Engine version.</summary>
-		public string EngineVersion = "";
-		/// <summary>Time of a crash.</summary>
+		/// <summary>GameName.</summary>
+		public string GameName;
+		/// <summary>TimeOfCrash.</summary>
 		public DateTime TimeOfCrash;
-		/// <summary>Crash type.</summary>
+		/// <summary>BuiltFromCL.</summary>
+		public string BuiltFromCL;
+		/// <summary>PlatformName.</summary>
+		public string PlatformName;
+		/// <summary>EngineMode.</summary>
+		public string EngineMode;
+		/// <summary>MachineId.</summary>
+		public string MachineId;
+		/// <summary>Module.</summary>
+		public string Module;
+		/// <summary>BuildVersion.</summary>
+		public string BuildVersion;
+		/// <summary>Jira.</summary>
+		public string Jira;
+		/// <summary>Branch.</summary>
+		public string Branch;
+		/// <summary>CrashType.</summary>
 		public short CrashType;
+		/// <summary>EpicId.</summary>
+		public string EpicId;
+		/// <summary>BuggId.</summary>
+		public int BuggId;
 
 		/// <summary>Returns crash type as a string.</summary>
 		public string GetCrashTypeAsString()
@@ -83,7 +116,7 @@ namespace Tools.CrashReporter.CrashReportWebSite.Controllers
 				HashSet<string> UserNamesForUserGroup = FRepository.Get( BuggsRepo ).GetUserNamesFromGroupName( AnonymousGroup );
 
 				// Enable to narrow results and improve debugging performance.
-				//FormData.DateFrom = FormData.DateTo.AddDays( -1 );
+				FormData.DateFrom = FormData.DateTo.AddDays( -1 );
 				FormData.DateTo = FormData.DateTo.AddDays( 1 );
 
 				var FilteringQueryJoin = CrashRepo
@@ -102,11 +135,19 @@ namespace Tools.CrashReporter.CrashReportWebSite.Controllers
 						bc => bc.Crash,
 						( c, bc ) => new
 						{
-							EpicId = c.EpicAccountId,
-							BuggId = bc.BuggId,
+							GameName = c.GameName,
 							TimeOfCrash = c.TimeOfCrash.Value,
-							EngineVersion = c.BuildVersion,
+							BuiltFromCL = c.BuiltFromCL,
+							PlatformName = c.PlatformName,
+							EngineMode = c.EngineMode,
+							MachineId = c.MachineId,
+							Module = c.Module,
+							BuildVersion = c.BuildVersion,
+							Jira = c.Jira,
+							Branch = c.Branch,
 							CrashType = c.CrashType.Value,
+							EpicId = c.EpicAccountId,
+							BuggId = bc.BuggId,
 						}
 					);
 
@@ -157,8 +198,6 @@ namespace Tools.CrashReporter.CrashReportWebSite.Controllers
 				// SLOW
 				//var Users = FRepository.Get( BuggsRepo ).GetAnalyticsUsers().ToDictionary( X => X.EpicAccountId, Y => Y );
 
-				
-
 				// Export data to the file.
 				string CSVPathname = Path.Combine( Settings.Default.CrashReporterCSV, DateTime.UtcNow.ToString( "yyyy-MM-dd.HH-mm-ss" ) );
 				CSVPathname += string
@@ -173,20 +212,46 @@ namespace Tools.CrashReporter.CrashReportWebSite.Controllers
 
 				using (FAutoScopedLogTimer ExportToCSVTimer = new FAutoScopedLogTimer( "ExportToCSV" ))
 				{
-					CSVFile.WriteLine( "TimeOfCrash; EpicId; BuggId; EngineVersion; CrashType; " );
+					var RowType = FilteringQueryJoin.FirstOrDefault().GetType();
+
+					string AllProperties = "";
+					foreach( var Property in RowType.GetProperties() )
+					{
+						AllProperties += Property.Name;
+						AllProperties += "; ";
+					}
+
+					// Write header
+					CSVFile.WriteLine( AllProperties );
 
 					foreach (var Row in FilteringQueryJoin)
 					{
-						var BVParts = Row.EngineVersion.Split( new char[] { '.' }, StringSplitOptions.RemoveEmptyEntries );
+						var BVParts = Row.BuildVersion.Split( new char[] { '.' }, StringSplitOptions.RemoveEmptyEntries );
 						if (BVParts.Length > 2 && BVParts[0] != "0")
 						{
 							string CleanEngineVersion = string.Format( "{0}.{1}.{2}", BVParts[0], BVParts[1], BVParts[2] );
-							CSVFile.WriteLine( "{0};{1};{2};{3};{4};", 
-								Row.TimeOfCrash, 
-								Row.EpicId, 
-								Row.BuggId, 
+
+							string[] RowProperties = new string[]
+							{
+								Row.GameName,
+								Row.TimeOfCrash.ToString(),
+								Row.BuiltFromCL,
+								Row.PlatformName,
+								Row.EngineMode,
+								Row.MachineId,
+								Row.Module,
 								CleanEngineVersion,
-								Row.CrashType == 1 ? "Crash" : "Assert" );
+								Row.Jira,
+								Row.Branch,
+								Row.CrashType == 1 ? "Crash" : "Assert",
+								Row.EpicId,
+								Row.BuggId.ToString()
+							};
+
+							string JoinedLine = string.Join( "; ", RowProperties );
+							JoinedLine += "; ";
+
+							CSVFile.WriteLine( JoinedLine );
 						}
 					}
 
@@ -197,8 +262,23 @@ namespace Tools.CrashReporter.CrashReportWebSite.Controllers
 
 				List<FCSVRow> CSVRows = FilteringQueryJoin
 					.OrderByDescending( X => X.TimeOfCrash )
-					.Take( 100 )
-					.Select( c => new FCSVRow { TimeOfCrash = c.TimeOfCrash, EpicId = c.EpicId, BuggId = c.BuggId, EngineVersion = c.EngineVersion, CrashType = c.CrashType } )
+					.Take( 32 )
+					.Select( c => new FCSVRow 
+					{
+						GameName = c.GameName,
+						TimeOfCrash = c.TimeOfCrash,
+						BuiltFromCL = c.BuiltFromCL,
+						PlatformName = c.PlatformName,
+						EngineMode = c.EngineMode,
+						MachineId = c.MachineId,
+						Module = c.Module,
+						BuildVersion = c.BuildVersion,
+						Jira = c.Jira,
+						Branch = c.Branch,
+						CrashType = c.CrashType,
+						EpicId = c.EpicId,
+						BuggId = c.BuggId,
+					} )
 					.ToList();
 				
 				return new CSV_ViewModel
diff --git a/Engine/Source/Programs/CrashReporter/CrashReportWebSite/Views/CSV/ViewCSV.ascx b/Engine/Source/Programs/CrashReporter/CrashReportWebSite/Views/CSV/ViewCSV.ascx
--- a/Engine/Source/Programs/CrashReporter/CrashReportWebSite/Views/CSV/ViewCSV.ascx
+++ b/Engine/Source/Programs/CrashReporter/CrashReportWebSite/Views/CSV/ViewCSV.ascx
@@ -93,57 +93,43 @@ table, table tr, table th, table td
 </table>
 <br />
 
-<p>Displaying top 100 records, ordered by the date of a crash</p>
+<p>Displaying top 32 records, ordered by the date of a crash</p>
 <table style="width: 100%">
 	<tr>
-		<th>
-			Time of crash
-		</th>
-		<th>
-			EpicId
-		</th>
-		<%--<th>
-			User's email
-		</th>--%>
-		<th>
-			BuggId
-		</th>
-		<th>
-			Engine version
-		</th>
-		<th>
-			Crash type
-		</th>
+		<th>GameName</th>
+		<th>TimeOfCrash</th>
+		<th>BuiltFromCL</th>
+		<th>PlatformName</th>
+		<th>EngineMode</th>
+		<th>MachineId</th>
+		<th>Module</th>
+		<th>BuildVersion</th>
+		<th>Jira</th>
+		<th>Branch</th>
+		<th>CrashType</th>
+		<th>EpicId</th>
+		<th>BuggId</th>
 	</tr>
 
 	<%
 		for (int Index = 0; Index < Model.CSVRows.Count; Index++)
 		{
 			FCSVRow CSVRow = Model.CSVRows[Index];
-			if( Index > 100 )
-			{
-				break;
-			}
 	%>
 	<tr>
-		<td>
-			<%=CSVRow.TimeOfCrash.ToString()%>
-		</td>
-		<td>
-			<%=CSVRow.EpicId%>
-		</td>
-		<%--<td>
-			<%=CSVRow.UserEmail%>
-		</td>--%>
-		<td>
-			<%=CSVRow.BuggId%>
-		</td>
-		<td>
-			<%=CSVRow.EngineVersion%>
-		</td>	
-		<td>
-			<%=CSVRow.GetCrashTypeAsString()%>
-		</td>	
+		<td><%=CSVRow.GameName%></td>
+		<td><%=CSVRow.TimeOfCrash%></td>
+		<td><%=CSVRow.BuiltFromCL%></td>
+		<td><%=CSVRow.PlatformName%></td>
+		<td><%=CSVRow.EngineMode%></td>
+		<td><%=CSVRow.MachineId%></td>
+		<td><%=CSVRow.Module%></td>
+		<td><%=CSVRow.BuildVersion%></td>
+		<td><%=CSVRow.Jira%></td>
+		<td><%=CSVRow.Branch%></td>
+		<td><%=CSVRow.GetCrashTypeAsString()%></td>
+		<td><%=CSVRow.EpicId%></td>
+		<td><%=CSVRow.BuggId%></td>
 	</tr>
 	<%
 		}	
