CrashReportClient - Do not send user's name when using a binary release

diff --git a/Engine/Source/Programs/CrashReporter/CrashReportClient/Private/GenericErrorReport.cpp b/Engine/Source/Programs/CrashReporter/CrashReportClient/Private/GenericErrorReport.cpp
--- a/Engine/Source/Programs/CrashReporter/CrashReportClient/Private/GenericErrorReport.cpp
+++ b/Engine/Source/Programs/CrashReporter/CrashReportClient/Private/GenericErrorReport.cpp
@@ -73,9 +73,16 @@ bool FGenericErrorReport::SetUserComment(const FText& UserComment, bool bAllowTo
 		DynamicSignaturesNode->AppendChildNode(TEXT("Parameter3"), UserComment.ToString());
 	}
 	
+	FString MachineIDandUserID;
 	// Set global user name ID: will be added to the report
 	extern FCrashDescription& GetCrashDescription();
-	FString MachineIDandUserID = FString::Printf( TEXT( "!MachineId:%s!EpicAccountId:%s!Name:%s" ), *GetCrashDescription().MachineId, *GetCrashDescription().EpicAccountId, *GetCrashDescription().UserName );
+
+	MachineIDandUserID = FString::Printf( TEXT( "!MachineId:%s!EpicAccountId:%s" ), *GetCrashDescription().MachineId, *GetCrashDescription().EpicAccountId );
+
+	if (!FApp::IsEngineInstalled())
+	{
+		MachineIDandUserID += FString::Printf( TEXT( "!Name:%s" ), *GetCrashDescription().UserName );
+	}
 
 	// Add or update a user ID.
 	FXmlNode* Parameter4Node = DynamicSignaturesNode->FindChildNode(TEXT("Parameter4"));
