CrashFrontend - Fix for ensures being treated as crashes

diff --git a/Engine/Source/Programs/CrashReporter/CrashReportWebSite/Models/CrashRepository.cs b/Engine/Source/Programs/CrashReporter/CrashReportWebSite/Models/CrashRepository.cs
--- a/Engine/Source/Programs/CrashReporter/CrashReportWebSite/Models/CrashRepository.cs
+++ b/Engine/Source/Programs/CrashReporter/CrashReportWebSite/Models/CrashRepository.cs
@@ -618,7 +618,7 @@ namespace Tools.CrashReporter.CrashReportWebSite.Models
 				{
 					NewCrash.CrashType = 2;
 				}
-				else if (NewCrash.RawCallStack.Contains( "FDebug::EnsureFailed" ))
+				else if (NewCrash.RawCallStack.Contains( "FDebug::Ensure" ))
 				{
 					NewCrash.CrashType = 3;
 				}
