CrashFrontend - All ensures should have the pattern

diff --git a/Engine/Source/Programs/CrashReporter/CrashReportWebSite/Models/CrashDataFormater.cs b/Engine/Source/Programs/CrashReporter/CrashReportWebSite/Models/CrashDataFormater.cs
--- a/Engine/Source/Programs/CrashReporter/CrashReportWebSite/Models/CrashDataFormater.cs
+++ b/Engine/Source/Programs/CrashReporter/CrashReportWebSite/Models/CrashDataFormater.cs
@@ -253,11 +253,11 @@ namespace Tools.CrashReporter.CrashReportWebSite.Models
 			{
 			case 2:
 				bSkipping = true;
-				LineToSkipUpto = "FDebug::AssertFailed()";
+				LineToSkipUpto = "FDebug::AssertFailed";
 				break;
 			case 3:
 				bSkipping = true;
-				LineToSkipUpto = "FDebug::EnsureFailed()";
+				LineToSkipUpto = "FDebug::Ensure";
 				break;
 			}
 
diff --git a/Engine/Source/Programs/CrashReporter/CrashReportWebSite/Models/CrashRepository.cs b/Engine/Source/Programs/CrashReporter/CrashReportWebSite/Models/CrashRepository.cs
--- a/Engine/Source/Programs/CrashReporter/CrashReportWebSite/Models/CrashRepository.cs
+++ b/Engine/Source/Programs/CrashReporter/CrashReportWebSite/Models/CrashRepository.cs
@@ -147,7 +147,7 @@ namespace Tools.CrashReporter.CrashReportWebSite.Models
 		/// <returns>The crash with the requested id.</returns>
 		public Crash GetCrash( int Id )
 		{
-			using( FAutoScopedLogTimer LogTimer = new FAutoScopedLogTimer( this.GetType().ToString() + "(CrashId" + Id + ")" ) )
+			using( FAutoScopedLogTimer LogTimer = new FAutoScopedLogTimer( this.GetType().ToString() + "(CrashId=" + Id + ")" ) )
 			{
 				try
 				{
