CrashReportLinq - Updated RepairDatabase program

diff --git a/Engine/Source/Programs/CrashReporter/CrashReportWebSite/Queries/Fix-RepairDatabase.linq b/Engine/Source/Programs/CrashReporter/CrashReportWebSite/Queries/Fix-RepairDatabase.linq
--- a/Engine/Source/Programs/CrashReporter/CrashReportWebSite/Queries/Fix-RepairDatabase.linq
+++ b/Engine/Source/Programs/CrashReporter/CrashReportWebSite/Queries/Fix-RepairDatabase.linq
@@ -33,9 +33,10 @@ HashSet<int> GetInvalidFuncIds()
 }
 
 static Stopwatch Timer = Stopwatch.StartNew();
-static DateTime Date = DateTime.UtcNow.AddMonths(-1);
-static TimeSpan Tick = TimeSpan.FromDays(3);
-const int NUM_OPS_PER_BATCH = 32;
+static DateTime Date = DateTime.UtcNow.AddMonths(-1);   // From
+static DateTime UtcNow = DateTime.UtcNow.AddHours(-2); // To
+static TimeSpan Tick = TimeSpan.FromDays(1);
+const int NUM_OPS_PER_BATCH = 16;
 
 void WriteLine( string Line )
 {
@@ -52,6 +53,10 @@ void DeleteNoPatternBatch( List<Crashes> NoPatternBatch )
 		Crashes.Context.SubmitChanges();
 		NoPatternTotal += NoPatternBatch.Count;
 		WriteLine( string.Format( "NoPattern deleted: {0}", NoPatternTotal ) );
+		foreach (var Crash in NoPatternBatch)
+		{
+			WriteLine( string.Format( "NoPattern deleted, Id:{0}", Crash.Id ) );
+		}
 		NoPatternBatch.Clear();
 	}
 }
@@ -97,8 +102,8 @@ void DeleteBuggs( List<int> BuggsBatch )
 void DeleteNoPatternAll()
 {
 	var NoPatternCount = Crashes
-	.Where( c => c.Pattern == null )
-	.Where( c => c.TimeOfCrash> Date )
+	.Where( c => c.Pattern == null || c.Pattern == "" )
+	.Where( c => c.TimeOfCrash >= Date &&  c.TimeOfCrash <= UtcNow )
 	.Count();
 	WriteLine( string.Format( "NoPatternCount: {0}", NoPatternCount ) );
 	
@@ -106,19 +111,22 @@ void DeleteNoPatternAll()
 	
 	DateTime StartDate = Date;
 	DateTime EndDate = Date.Add( Tick );
-	while( EndDate < DateTime.UtcNow )
+	int Total = 0;
+	
+	while( EndDate <= UtcNow )
 	{
 		var NoPatternCrashes = Crashes
 		.Where (c => c.TimeOfCrash > StartDate && c.TimeOfCrash <= EndDate)
-		.Where( c => c.Pattern == null )
+		.Where( c => c.Pattern == null || c.Pattern == "" )
 		.ToList();
 	
-		WriteLine( string.Format( "NoPatternCrashes: {0} {1} -> {2}", NoPatternCrashes.Count, StartDate.ToString("yyyy-MM-dd"), EndDate.ToString("yyyy-MM-dd") ) );
+		Total += NoPatternCrashes.Count;
+		WriteLine( string.Format( "NoPatternCrashes: {0}/{1} {2} -> {3}", NoPatternCrashes.Count, Total, StartDate.ToString("yyyy-MM-dd"), EndDate.ToString("yyyy-MM-dd") ) );		
 		StartDate = EndDate;
 	 	EndDate = EndDate.Add( Tick );
 		
 		var CrashesIds = new HashSet<int>( NoPatternCrashes.Select (np => np.Id).ToList() );
-		int NumBuggCrass = Buggs_Crashes.Count (bc => bc.BuggId>0);
+		//int NumBuggCrass = Buggs_Crashes.Count (bc => bc.BuggId>0);
 		
 		// Unlink from the bugg.	
 		List<Buggs_Crashes> ToRemove = new List<Buggs_Crashes>();
@@ -180,23 +188,26 @@ void DeleteNoPatternAll()
 void SetNoPatternForNoBranchAll()
 {
 	var NoBranchCount = Crashes
-	.Where( c => c.Branch == null )
-	.Where (c => c.TimeOfCrash> Date)
+	.Where( c => c.Branch == null || c.Branch == "INVALID" )
+	.Where( c => c.TimeOfCrash >= Date &&  c.TimeOfCrash <= UtcNow )
 	.Count();
 	
 	WriteLine( string.Format( "NoBranch: {0}", NoBranchCount ) );
 	
 	var NoBranchBatch = new List<Crashes>(NUM_OPS_PER_BATCH);
 	DateTime StartDate = Date;
-	DateTime EndDate = Date.Add( Tick );
-	while( EndDate < DateTime.UtcNow )
+	DateTime EndDate = Date.Add( Tick );	
+	int Total = 0;
+	
+	while( EndDate <= UtcNow )
 	{
 		var NoBranchCrashes = Crashes
 		.Where (c => c.TimeOfCrash > StartDate && c.TimeOfCrash <= EndDate)
 		.Where( c => c.Branch == null )
 		.ToList();
 	
-		WriteLine( string.Format( "NoBranchCrashes: {0} {1} -> {2}", NoBranchCrashes.Count, StartDate.ToString("yyyy-MM-dd"), EndDate.ToString("yyyy-MM-dd") ) );
+		Total += NoBranchCrashes.Count;
+		WriteLine( string.Format( "NoBranchCrashes: {0}/{1} {2} -> {3}", NoBranchCrashes.Count, Total, StartDate.ToString("yyyy-MM-dd"), EndDate.ToString("yyyy-MM-dd") ) );		
 		StartDate = EndDate;
 	 	EndDate = EndDate.Add( Tick );
 	
@@ -204,7 +215,7 @@ void SetNoPatternForNoBranchAll()
 		{
 			NoBranchBatch.Add(NoBranchCrashes[n]);		
 			NoBranchCrashes[n].Pattern = null;
-			NoBranchCrashes[n].Branch = "INVALID";
+			NoBranchCrashes[n].Branch = null;
 			
 			if (NoBranchBatch.Count == NUM_OPS_PER_BATCH)
 			{
@@ -265,15 +276,18 @@ void SetNoPatternForInvalidPattern()
 	InvalidFuncIds.OrderBy (ifi => ifi);
 	
 	DateTime StartDate = Date;
-	DateTime EndDate = Date.Add( Tick );
-	while( EndDate < DateTime.UtcNow )
+	DateTime EndDate = Date.Add( Tick );	
+	int Total = 0;
+	
+	while( EndDate <= UtcNow )
 	{
 		var AllCrashes = Crashes
 		.Where (c => c.TimeOfCrash > StartDate && c.TimeOfCrash <= EndDate)
 		.Where (c => c.Pattern != null)
 		.ToList();
 		
-		WriteLine( string.Format( "AllCrashes: {0} {1} -> {2}", AllCrashes.Count, StartDate.ToString("yyyy-MM-dd"), EndDate.ToString("yyyy-MM-dd") ) );
+		Total += AllCrashes.Count;
+		WriteLine( string.Format( "AllCrashes: {0}/{1} {2} -> {3}", AllCrashes.Count, Total, StartDate.ToString("yyyy-MM-dd"), EndDate.ToString("yyyy-MM-dd") ) );		
 		StartDate = EndDate;
 	 	EndDate = EndDate.Add( Tick );
 	
@@ -327,8 +341,10 @@ void SetNoPatternForEnsureForCrashes()
 	//UGameViewportClient::MouseLeave
 	var EnsureAssertBatch = new List<Crashes>(NUM_OPS_PER_BATCH);
 	DateTime StartDate = Date;
-	DateTime EndDate = Date.Add( Tick );
-	while( EndDate < DateTime.UtcNow )
+	DateTime EndDate = Date.Add( Tick );	
+	int Total = 0;
+	
+	while( EndDate <= UtcNow )
 	{
 		var EnsureAssertCrashes = Crashes
 		.Where( c => c.TimeOfCrash > StartDate && c.TimeOfCrash <= EndDate )
@@ -336,7 +352,8 @@ void SetNoPatternForEnsureForCrashes()
 		.Where( c => c.RawCallStack.Contains("FDebug::Ensure") || c.RawCallStack.Contains("NewReportEnsure") )
 		.ToList();
 	
-		WriteLine( string.Format( "EnsureAssertCrashes: {0} {1} -> {2}", EnsureAssertCrashes.Count, StartDate.ToString("yyyy-MM-dd"), EndDate.ToString("yyyy-MM-dd") ) );
+		Total += EnsureAssertCrashes.Count;
+		WriteLine( string.Format( "EnsureAssertCrashes: {0}/{1} {2} -> {3}", EnsureAssertCrashes.Count, Total, StartDate.ToString("yyyy-MM-dd"), EndDate.ToString("yyyy-MM-dd") ) );		
 		StartDate = EndDate;
 	 	EndDate = EndDate.Add( Tick );
 	
@@ -345,7 +362,7 @@ void SetNoPatternForEnsureForCrashes()
 			EnsureAssertBatch.Add(EnsureAssertCrashes[n]);		
 			EnsureAssertCrashes[n].Pattern = null;
 			EnsureAssertCrashes[n].CrashType = 3;
-			EnsureAssertCrashes[n].Branch = "INVALID";
+			EnsureAssertCrashes[n].Branch = null;
 			
 			if (EnsureAssertBatch.Count == NUM_OPS_PER_BATCH)
 			{
@@ -358,19 +375,28 @@ void SetNoPatternForEnsureForCrashes()
 	SumbitChanges(EnsureAssertBatch);
 }
 
-void Main()
-{	
+void PrintStats()
+{
 	int NumCrashes = Crashes.Count();
 	int MaxCrashId = Crashes.Max (c => c.Id);
 	int NumBuggs = Buggs.Count();
 	int MaxBuggId = Buggs.Max(c => c.Id);
+	int NumBuggCrashes = Buggs_Crashes.Count();
 	WriteLine( string.Format( "NumCrashes: {0} MaxCrashId: {1} Diff: {2}", NumCrashes, MaxCrashId, MaxCrashId-NumCrashes ) );
 	WriteLine( string.Format( "NumBuggs: {0} MaxBuggId: {1} Diff: {2}", NumBuggs, MaxBuggId, MaxBuggId-NumBuggs ) );
+	WriteLine( string.Format( "NumBuggCrashes: {0}", NumBuggCrashes ) );
+}
+
+void Main()
+{	
+	PrintStats();
 	
 	SetNoPatternForEnsureForCrashes();
 	SetNoPatternForNoBranchAll();
 	SetNoPatternForInvalidPattern();
 	DeleteNoPatternAll();
+	
+	PrintStats();
 }
 
 // Define other methods and classes here
\ No newline at end of file
