CrashReportClient - Better logging

diff --git a/Engine/Source/Programs/CrashReporter/CrashReportClient/Private/CrashUpload.cpp b/Engine/Source/Programs/CrashReporter/CrashReportClient/Private/CrashUpload.cpp
--- a/Engine/Source/Programs/CrashReporter/CrashReportClient/Private/CrashUpload.cpp
+++ b/Engine/Source/Programs/CrashReporter/CrashReportClient/Private/CrashUpload.cpp
@@ -98,13 +98,14 @@ bool FCrashUpload::SendCheckReportRequest()
 {
 	FString XMLString;
 
-	UE_LOG(CrashReportClientLog, Log, TEXT("Sending HTTP request (checking report)"));
 	auto Request = CreateHttpRequest();
 	if (State == EUploadState::CheckingReport)
 	{
 		AssignReportIdToPostDataBuffer();
 		Request->SetURL(UrlPrefix / TEXT("CheckReport"));
 		Request->SetHeader(TEXT("Content-Type"), TEXT("text/plain; charset=us-ascii"));
+
+		UE_LOG( CrashReportClientLog, Log, TEXT( "Sending HTTP request: %s" ), *Request->GetURL() );
 	}
 	else
 	{
@@ -120,6 +121,8 @@ bool FCrashUpload::SendCheckReportRequest()
 
 		Request->SetURL(UrlPrefix / TEXT("CheckReportDetail"));
 		Request->SetHeader(TEXT("Content-Type"), TEXT("text/plain; charset=utf-8"));
+
+		UE_LOG( CrashReportClientLog, Log, TEXT( "Sending HTTP request: %s" ), *Request->GetURL() );
 	}
 
 	UE_LOG( CrashReportClientLog, Log, TEXT( "PostData Num: %i" ), PostData.Num() );
@@ -223,7 +226,6 @@ void FCrashUpload::CompressAndSendData()
 	CompressedDataRaw = nullptr;
 
 	// Set up request for upload
-	UE_LOG(CrashReportClientLog, Log, TEXT("Sending HTTP request (posting file)"));
 	auto Request = CreateHttpRequest();
 	Request->SetVerb(TEXT("POST"));
 	Request->SetHeader(TEXT("Content-Type"), TEXT("application/octet-stream"));
@@ -235,6 +237,7 @@ void FCrashUpload::CompressAndSendData()
 	Request->SetHeader(TEXT("CompressedSize"), TTypeToString<int32>::ToString(CompressedSize) );
 	Request->SetHeader(TEXT("UncompressedSize"), TTypeToString<int32>::ToString(UncompressedSize) );
 	Request->SetHeader(TEXT("NumberOfFiles"), TTypeToString<int32>::ToString(CurrentFileIndex) );
+	UE_LOG( CrashReportClientLog, Log, TEXT( "Sending HTTP request: %s" ), *Request->GetURL() );
 
 	if (Request->ProcessRequest())
 	{
@@ -269,12 +272,13 @@ void FCrashUpload::PostReportComplete()
 
 	AssignReportIdToPostDataBuffer();
 
-	UE_LOG(CrashReportClientLog, Log, TEXT("Sending HTTP request (posting \"Upload complete\")"));
+	
 	auto Request = CreateHttpRequest();
 	Request->SetVerb( TEXT( "POST" ) );
 	Request->SetURL(UrlPrefix / TEXT("UploadComplete"));
 	Request->SetHeader( TEXT( "Content-Type" ), TEXT( "text/plain; charset=us-ascii" ) );
 	Request->SetContent(PostData);
+	UE_LOG( CrashReportClientLog, Log, TEXT( "Sending HTTP request: %s" ), *Request->GetURL() );
 
 	if (Request->ProcessRequest())
 	{
@@ -438,7 +442,7 @@ void FCrashUpload::SendPingRequest()
 	auto Request = CreateHttpRequest();
 	Request->SetVerb(TEXT("GET"));
 	Request->SetURL(UrlPrefix / TEXT("Ping"));
-	UE_LOG(CrashReportClientLog, Log, TEXT("Sending HTTP request (pinging server)"));
+	UE_LOG( CrashReportClientLog, Log, TEXT( "Sending HTTP request: %s" ), *Request->GetURL() );
 
 	if (Request->ProcessRequest())
 	{
