Core - UE-15685/All Crashes are going to Anonymous

diff --git a/Engine/Source/Programs/CrashReporter/CrashReportClient/Private/GenericErrorReport.cpp b/Engine/Source/Programs/CrashReporter/CrashReportClient/Private/GenericErrorReport.cpp
--- a/Engine/Source/Programs/CrashReporter/CrashReportClient/Private/GenericErrorReport.cpp
+++ b/Engine/Source/Programs/CrashReporter/CrashReportClient/Private/GenericErrorReport.cpp
@@ -1,11 +1,11 @@
 // Copyright 1998-2015 Epic Games, Inc. All Rights Reserved.
 
 #include "CrashReportClientApp.h"
-
 #include "GenericErrorReport.h"
 #include "XmlFile.h"
 #include "CrashReportUtil.h"
 #include "CrashDescription.h"
+#include "EngineBuildSettings.h"
 
 // ----------------------------------------------------------------
 // Helpers
@@ -79,7 +79,8 @@ bool FGenericErrorReport::SetUserComment(const FText& UserComment, bool bAllowTo
 
 	MachineIDandUserID = FString::Printf( TEXT( "!MachineId:%s!EpicAccountId:%s" ), *GetCrashDescription().MachineId, *GetCrashDescription().EpicAccountId );
 
-	if (!FApp::IsEngineInstalled())
+	const bool bSendName = FEngineBuildSettings::IsInternalBuild() || FEngineBuildSettings::IsPerforceBuild() || FEngineBuildSettings::IsSourceDistribution();
+	if (bSendName)
 	{
 		MachineIDandUserID += FString::Printf( TEXT( "!Name:%s" ), *GetCrashDescription().UserName );
 	}
diff --git a/Engine/Source/Runtime/Core/Private/Windows/WindowsPlatformCrashContext.cpp b/Engine/Source/Runtime/Core/Private/Windows/WindowsPlatformCrashContext.cpp
--- a/Engine/Source/Runtime/Core/Private/Windows/WindowsPlatformCrashContext.cpp
+++ b/Engine/Source/Runtime/Core/Private/Windows/WindowsPlatformCrashContext.cpp
@@ -332,13 +332,6 @@ int32 ReportCrashUsingCrashReportClient(EXCEPTION_POINTERS* ExceptionInfo, const
 			CrashReportClientArguments += TEXT( " -Unattended" );
 		}
 
-		if( FApp::IsInstalled() )
-		{
-			// Temporary workaround for CrashReportClient being built in Development, not Shipping (TTP328030). The
-			// following ensures that logs are saved to the user directory when UE4 is installed.
-			CrashReportClientArguments += TEXT( " -Installed" );
-		}
-
 		CrashReportClientArguments += FString( TEXT( " -AppName=" ) ) + ReportInformation.wzApplicationName;
 
 		static const TCHAR CrashReportClientExeName[] = TEXT( "CrashReportClient.exe" );
