CrashReportClient - UECORE-102/ Include logs with Rocket crashes if checkbox enabled

diff --git a/Engine/Config/BaseEngine.ini b/Engine/Config/BaseEngine.ini
--- a/Engine/Config/BaseEngine.ini
+++ b/Engine/Config/BaseEngine.ini
@@ -2238,6 +2238,8 @@ DeployServerPort=53501
 [CrashReportClient]
 ;CrashReportReceiverIP="http://localhost:57005"
 CrashReportReceiverIP="http://crashreporter.epicgames.com:57005"
+bAllowToBeContacted=true
+bSendLogFile=true
 
 [UnrealHeaderTool]
 GeneratedCodeVersion=V1
diff --git a/Engine/Source/Programs/CrashReporter/CrashReportClient/Private/CrashReportClient.cpp b/Engine/Source/Programs/CrashReporter/CrashReportClient/Private/CrashReportClient.cpp
--- a/Engine/Source/Programs/CrashReporter/CrashReportClient/Private/CrashReportClient.cpp
+++ b/Engine/Source/Programs/CrashReporter/CrashReportClient/Private/CrashReportClient.cpp
@@ -27,7 +27,6 @@ FCrashReportClient::FCrashReportClient(const FPlatformErrorReport& InErrorReport
 	, Uploader( FCrashReportClientConfig::Get().GetReceiverAddress() )
 	, bBeginUploadCalled(false)
 	, bShouldWindowBeHidden(false)
-	, bAllowToBeContacted(true)
 	, bSendData(false)
 {
 
@@ -115,9 +114,14 @@ EVisibility FCrashReportClient::IsThrobberVisible() const
 	return IsProcessingCallstack() ? EVisibility::Visible : EVisibility::Hidden;
 }
 
-void FCrashReportClient::SCrashReportClient_OnCheckStateChanged( ECheckBoxState NewRadioState )
+void FCrashReportClient::AllowToBeContacted_OnCheckStateChanged( ECheckBoxState NewRadioState )
 {
-	bAllowToBeContacted = NewRadioState == ECheckBoxState::Checked;
+	FCrashReportClientConfig::Get().SetAllowToBeContacted( NewRadioState == ECheckBoxState::Checked );
+}
+
+void FCrashReportClient::SendLogFile_OnCheckStateChanged( ECheckBoxState NewRadioState )
+{
+	FCrashReportClientConfig::Get().SetSendLogFile( NewRadioState == ECheckBoxState::Checked );
 }
 
 void FCrashReportClient::StartTicker()
@@ -128,7 +132,7 @@ void FCrashReportClient::StartTicker()
 void FCrashReportClient::StoreCommentAndUpload()
 {
 	// Call upload even if the report is empty: pending reports will be sent if any
-	ErrorReport.SetUserComment(UserComment, bAllowToBeContacted);
+	ErrorReport.SetUserComment( UserComment, FCrashReportClientConfig::Get().GetAllowToBeContacted() );
 	StartTicker();
 }
 
diff --git a/Engine/Source/Programs/CrashReporter/CrashReportClient/Private/CrashReportClient.h b/Engine/Source/Programs/CrashReporter/CrashReportClient/Private/CrashReportClient.h
--- a/Engine/Source/Programs/CrashReporter/CrashReportClient/Private/CrashReportClient.h
+++ b/Engine/Source/Programs/CrashReporter/CrashReportClient/Private/CrashReportClient.h
@@ -113,7 +113,8 @@ public:
 	/** Whether the throbber should be visible while processing the callstack. */
 	EVisibility IsThrobberVisible() const;
 
-	void SCrashReportClient_OnCheckStateChanged( ECheckBoxState NewRadioState );
+	void AllowToBeContacted_OnCheckStateChanged( ECheckBoxState NewRadioState );
+	void SendLogFile_OnCheckStateChanged( ECheckBoxState NewRadioState );
 
 private:
 	/**
@@ -162,9 +163,6 @@ private:
 	/** Whether the main window should be hidden. */
 	bool bShouldWindowBeHidden;
 
-	/** Whether the user allowed us to be contacted. */
-	bool bAllowToBeContacted;
-
 	/** Whether we send the data. */
 	bool bSendData;
 
diff --git a/Engine/Source/Programs/CrashReporter/CrashReportClient/Private/CrashReportClientApp.cpp b/Engine/Source/Programs/CrashReporter/CrashReportClient/Private/CrashReportClientApp.cpp
--- a/Engine/Source/Programs/CrashReporter/CrashReportClient/Private/CrashReportClientApp.cpp
+++ b/Engine/Source/Programs/CrashReporter/CrashReportClient/Private/CrashReportClientApp.cpp
@@ -21,7 +21,7 @@
 #include "XmlFile.h"
 
 /** Default main window size */
-const FVector2D InitialWindowDimensions(640, 560);
+const FVector2D InitialWindowDimensions(740, 560);
 
 /** Average tick rate the app aims for */
 const float IdealTickRate = 30.f;
@@ -101,9 +101,24 @@ FCrashReportClientConfig::FCrashReportClientConfig()
 		CrashReportReceiverIP = TEXT( "http://crashreporter.epicgames.com:57005" );
 	}
 
+	GConfig->GetBool( TEXT( "CrashReportClient" ), TEXT( "bAllowToBeContacted" ), bAllowToBeContacted, GEngineIni );
+	GConfig->GetBool( TEXT( "CrashReportClient" ), TEXT( "bSendLogFile" ), bSendLogFile, GEngineIni );
+
 	UE_LOG( CrashReportClientLog, Log, TEXT( "CrashReportReceiverIP: %s" ), *CrashReportReceiverIP );
 }
 
+void FCrashReportClientConfig::SetAllowToBeContacted( bool bNewValue )
+{
+	bAllowToBeContacted = bNewValue;
+	GConfig->SetBool( TEXT( "CrashReportClient" ), TEXT( "bAllowToBeContacted" ), bAllowToBeContacted, GEngineIni );
+}
+
+void FCrashReportClientConfig::SetSendLogFile( bool bNewValue )
+{
+	bSendLogFile = bNewValue;
+	GConfig->SetBool( TEXT( "CrashReportClient" ), TEXT( "bSendLogFile" ), bSendLogFile, GEngineIni );
+}
+
 void RunCrashReportClient(const TCHAR* CommandLine)
 {
 	// Override the stack size for the thread pool.
diff --git a/Engine/Source/Programs/CrashReporter/CrashReportClient/Private/CrashReportUtil.h b/Engine/Source/Programs/CrashReporter/CrashReportClient/Private/CrashReportUtil.h
--- a/Engine/Source/Programs/CrashReporter/CrashReportClient/Private/CrashReportUtil.h
+++ b/Engine/Source/Programs/CrashReporter/CrashReportClient/Private/CrashReportUtil.h
@@ -27,12 +27,32 @@ struct FCrashReportClientConfig
 		return DiagnosticsFilename;
 	}
 
+	const bool& GetAllowToBeContacted() const
+	{
+		return bAllowToBeContacted;
+	}
+
+	const bool& GetSendLogFile() const
+	{
+		return bSendLogFile;
+	}
+
+	void SetAllowToBeContacted( bool bNewValue );
+
+	void SetSendLogFile( bool bNewValue );
+
 protected:
 	/** IP address of crash report receiver. */
 	FString CrashReportReceiverIP;
 
 	/** Filename to use when saving diagnostics report, if generated locally. */
 	FString DiagnosticsFilename;
+
+	/** Whether the user allowed us to be contacted. */
+	bool bAllowToBeContacted;
+
+	/** Whether the user allowed us to send the log file. */
+	bool bSendLogFile;
 };
 
 
diff --git a/Engine/Source/Programs/CrashReporter/CrashReportClient/Private/CrashUpload.cpp b/Engine/Source/Programs/CrashReporter/CrashReportClient/Private/CrashUpload.cpp
--- a/Engine/Source/Programs/CrashReporter/CrashReportClient/Private/CrashUpload.cpp
+++ b/Engine/Source/Programs/CrashReporter/CrashReportClient/Private/CrashUpload.cpp
@@ -176,7 +176,7 @@ void FCrashUpload::CompressAndSendData()
 	// Loop to keep trying files until a send succeeds or we run out of files
 	while (PendingFiles.Num() != 0)
 	{
-		FString PathOfFileToUpload = PendingFiles.Pop();
+		const FString PathOfFileToUpload = PendingFiles.Pop();
 		
 		if (FPlatformFileManager::Get().GetPlatformFile().FileSize(*PathOfFileToUpload) > MaxFileSizeToUpload)
 		{
@@ -190,6 +190,13 @@ void FCrashUpload::CompressAndSendData()
 			continue;
 		}
 
+		const bool bSkipLogFile = !FCrashReportClientConfig::Get().GetSendLogFile() && PathOfFileToUpload.EndsWith( TEXT( ".log" ) );
+		if (bSkipLogFile)
+		{
+			UE_LOG( CrashReportClientLog, Warning, TEXT( "Skipping the log file" ) );
+			continue;
+		}
+
 		UE_LOG(CrashReportClientLog, Log, TEXT("CompressAndSendData compressing %d bytes ('%s')"), PostData.Num(), *PathOfFileToUpload);
 		FString Filename = FPaths::GetCleanFilename(PathOfFileToUpload);
 		if (Filename == "diagnostics.txt")
diff --git a/Engine/Source/Programs/CrashReporter/CrashReportClient/Private/SCrashReportClient.cpp b/Engine/Source/Programs/CrashReporter/CrashReportClient/Private/SCrashReportClient.cpp
--- a/Engine/Source/Programs/CrashReporter/CrashReportClient/Private/SCrashReportClient.cpp
+++ b/Engine/Source/Programs/CrashReporter/CrashReportClient/Private/SCrashReportClient.cpp
@@ -172,7 +172,32 @@ void SCrashReportClient::Construct(const FArguments& InArgs, TSharedRef<FCrashRe
 
 			+SVerticalBox::Slot()
 			.AutoHeight()
-			.Padding( FMargin( 4, 12 ) )
+			.Padding( FMargin( 4, 12, 4, 4 ) )
+			[
+				SNew( SHorizontalBox )
+
+				+ SHorizontalBox::Slot()
+				.AutoWidth()
+				.VAlign( VAlign_Center )
+				[
+					SNew( SCheckBox )
+					.IsChecked( FCrashReportClientConfig::Get().GetSendLogFile() ? ECheckBoxState::Checked : ECheckBoxState::Unchecked )
+					.OnCheckStateChanged( CrashReportClient.ToSharedRef(), &FCrashReportClient::SendLogFile_OnCheckStateChanged )
+				]
+
+				+ SHorizontalBox::Slot()
+				.FillWidth( 1.0f )
+				.VAlign( VAlign_Center )
+				[
+					SNew( STextBlock )
+					.AutoWrapText( true )
+					.Text( LOCTEXT( "IncludeLogs", "Include log files with submission. I understand that logs contain some personal information such as my system and user name." ) )
+				]
+			]
+
+			+SVerticalBox::Slot()
+			.AutoHeight()
+			.Padding( FMargin( 4, 4 ) )
 			[
 				SNew(SHorizontalBox)
 
@@ -181,10 +206,11 @@ void SCrashReportClient::Construct(const FArguments& InArgs, TSharedRef<FCrashRe
 				.VAlign(VAlign_Center)
 				[
 					SNew(SCheckBox)
-					.IsChecked(ECheckBoxState::Checked)
-					.OnCheckStateChanged(CrashReportClient.ToSharedRef(), &FCrashReportClient::SCrashReportClient_OnCheckStateChanged)
+					.IsChecked( FCrashReportClientConfig::Get().GetAllowToBeContacted() ? ECheckBoxState::Checked : ECheckBoxState::Unchecked )
+					.OnCheckStateChanged(CrashReportClient.ToSharedRef(), &FCrashReportClient::AllowToBeContacted_OnCheckStateChanged)
 				]
 
+				
 				+SHorizontalBox::Slot()
 				.FillWidth(1.0f)
 				.VAlign(VAlign_Center)
