CrashFrontend -
UECORE-152
Crash Reporter: The page crashes when running a search query with < >

diff --git a/Engine/Source/Programs/CrashReporter/CrashReportWebSite/Models/CrashRepository.cs b/Engine/Source/Programs/CrashReporter/CrashReportWebSite/Models/CrashRepository.cs
--- a/Engine/Source/Programs/CrashReporter/CrashReportWebSite/Models/CrashRepository.cs
+++ b/Engine/Source/Programs/CrashReporter/CrashReportWebSite/Models/CrashRepository.cs
@@ -671,9 +671,10 @@ namespace Tools.CrashReporter.CrashReportWebSite.Models
 			var Results = ListAll();
 
 			// Grab Results 
-			if (!string.IsNullOrEmpty( FormData.SearchQuery ))
+			string QueryString = HttpUtility.HtmlDecode( FormData.SearchQuery );
+			if (!string.IsNullOrEmpty( QueryString ))
 			{
-				if (!string.IsNullOrEmpty( FormData.UsernameQuery ))
+				if (!string.IsNullOrEmpty( QueryString ))
 				{
 					//We only use SearchQuery now for CallStack searching - if there's a SearchQuery value and a Username value, we need to get rid of the 
 					//Username so that we can create a broader search range
diff --git a/Engine/Source/Programs/CrashReporter/CrashReportWebSite/Models/FormHelper.cs b/Engine/Source/Programs/CrashReporter/CrashReportWebSite/Models/FormHelper.cs
--- a/Engine/Source/Programs/CrashReporter/CrashReportWebSite/Models/FormHelper.cs
+++ b/Engine/Source/Programs/CrashReporter/CrashReportWebSite/Models/FormHelper.cs
@@ -5,6 +5,7 @@ using System.Collections.Generic;
 using System.Linq;
 using System.Web;
 using System.Web.Mvc;
+using System.Web.Helpers;
 
 namespace Tools.CrashReporter.CrashReportWebSite.Models
 {
@@ -98,6 +99,40 @@ namespace Tools.CrashReporter.CrashReportWebSite.Models
 			return false;
 		}
 
+		/// <summary>
+		/// Same as GetFormParameter, but ignores validation. Only used for search query.
+		/// </summary>
+		/// <param name="Request">The owning HttpRequest.</param>
+		/// <param name="Form">The form passed to the controller.</param>
+		/// <param name="Key">The parameter we wish to retrieve a value for.</param>
+		/// <param name="DefaultValue">The default value to set if the parameter is null or empty.</param>
+		/// <param name="Result">The value of the requested parameter.</param>
+		/// <returns>true if the parameter was found to be not null or empty.</returns>
+		private bool GetFormParameterUnvalidated( HttpRequestBase Request, FormCollection Form, string Key, string DefaultValue, out string Result )
+		{
+			var Collection = Request.Unvalidated;
+
+			Result = DefaultValue;
+
+			string Value = "";
+			if (Form.Count == 0)
+			{
+				Value = Collection.QueryString[Key];
+			}
+			else
+			{
+				Value = Form[Key];
+			}
+
+			if (!string.IsNullOrEmpty( Value ))
+			{
+				Result = Value.Trim();
+				return true;
+			}
+
+			return false;
+		}
+
 		/// <summary>
 		/// Attempt to read a form item as a date in milliseconds since 1970
 		/// </summary>
@@ -131,7 +166,7 @@ namespace Tools.CrashReporter.CrashReportWebSite.Models
 		public FormHelper( HttpRequestBase Request, FormCollection Form, string DefaultSortTerm )
 		{
 			// Set up Default values if there is no QueryString and set values to the Query string if it is there.
-			GetFormParameter( Request, Form, "SearchQuery", SearchQuery, out SearchQuery );
+			GetFormParameterUnvalidated( Request, Form, "SearchQuery", SearchQuery, out SearchQuery );
 			GetFormParameter( Request, Form, "UsernameQuery", UsernameQuery, out UsernameQuery );
 			GetFormParameter( Request, Form, "EpicIdOrMachineQuery", EpicIdOrMachineQuery, out EpicIdOrMachineQuery );
 			GetFormParameter( Request, Form, "JiraQuery", JiraQuery, out JiraQuery );
