CrashFrontend - UECORE-162/CRASH REPORTER: Please expand the Reports page to 200 crashgroups

diff --git a/Engine/Source/Programs/CrashReporter/CrashReportWebSite/Controllers/ReportsController.cs b/Engine/Source/Programs/CrashReporter/CrashReportWebSite/Controllers/ReportsController.cs
--- a/Engine/Source/Programs/CrashReporter/CrashReportWebSite/Controllers/ReportsController.cs
+++ b/Engine/Source/Programs/CrashReporter/CrashReportWebSite/Controllers/ReportsController.cs
@@ -136,8 +136,8 @@ namespace Tools.CrashReporter.CrashReportWebSite.Controllers
 					}
 				}
 				var PatternToCountOrdered = PatternToCount.OrderByDescending( X => X.Value ).ToList();
-				// The 100 top.
-				var PatternAndCount = PatternToCountOrdered.Take( 100 ).ToDictionary( X => X.Key, Y => Y.Value );
+				const int NumTopRecords = 200;
+				var PatternAndCount = PatternToCountOrdered.Take( NumTopRecords ).ToDictionary( X => X.Key, Y => Y.Value );
 
 				int TotalUniqueAnonymousCrashes = UniquePatterns.Count;
 
@@ -150,7 +150,7 @@ namespace Tools.CrashReporter.CrashReportWebSite.Controllers
 				HashSet<string> FoundJiras = new HashSet<string>();
 				Dictionary<string, List<Bugg>> JiraIDtoBugg = new Dictionary<string, List<Bugg>>();
 
-				List<Bugg> Buggs = new List<Bugg>( 100 );
+				List<Bugg> Buggs = new List<Bugg>( NumTopRecords );
 				foreach( var Top in PatternAndCount )
 				{
 					Bugg NewBugg = RealBuggs.Where( X => X.Pattern == Top.Key ).FirstOrDefault();
