CrashReportLinq -
UECORE-196/Problems with crash groups in Crash Reporter (Updated script, now removes empty buggs with no associated crashes)
UECORE-194/CRASH REPORTER: Ensure appearing in Top 100 Report (removed from the db, malformed data)

diff --git a/Engine/Source/Programs/CrashReporter/CrashReportWebSite/Queries/Fix-RepairDatabase.linq b/Engine/Source/Programs/CrashReporter/CrashReportWebSite/Queries/Fix-RepairDatabase.linq
--- a/Engine/Source/Programs/CrashReporter/CrashReportWebSite/Queries/Fix-RepairDatabase.linq
+++ b/Engine/Source/Programs/CrashReporter/CrashReportWebSite/Queries/Fix-RepairDatabase.linq
@@ -94,6 +94,9 @@ void DeleteBuggs( List<int> BuggsBatch )
 		Buggs.DeleteAllOnSubmit( Batch );
 		Buggs.Context.SubmitChanges();
 		WriteLine( string.Format( "DeleteBuggs deleted: {0}", BuggsTotal ) );
+		
+		WriteLine( "DeleteBuggs= " + string.Join( "; ", BuggsBatch ) );
+		
 		BuggsBatch.Clear();
 	}
 }
@@ -337,32 +340,68 @@ void SetNoPatternForInvalidPattern()
 
 void SetNoPatternForEnsureForCrashes()
 {
-	WriteLine( string.Format( "EnsureButAssert" ) );
-	//UGameViewportClient::MouseLeave
-	var EnsureAssertBatch = new List<Crashes>(NUM_OPS_PER_BATCH);
+	var EnsuresMarkedAsAssertsBatch = new List<Crashes>(NUM_OPS_PER_BATCH);
 	DateTime StartDate = Date;
 	DateTime EndDate = Date.Add( Tick );	
 	int Total = 0;
 	
 	while( EndDate <= UtcNow )
 	{
-		var EnsureAssertCrashes = Crashes
+		var EnsuresMarkedAsAsserts = Crashes
 		.Where( c => c.TimeOfCrash > StartDate && c.TimeOfCrash <= EndDate )
 		.Where( c => c.CrashType == 1 || c.CrashType == 2 )
 		.Where( c => c.RawCallStack.Contains("FDebug::Ensure") || c.RawCallStack.Contains("NewReportEnsure") )
 		.ToList();
 	
-		Total += EnsureAssertCrashes.Count;
-		WriteLine( string.Format( "EnsureAssertCrashes: {0}/{1} {2} -> {3}", EnsureAssertCrashes.Count, Total, StartDate.ToString("yyyy-MM-dd"), EndDate.ToString("yyyy-MM-dd") ) );		
+		Total += EnsuresMarkedAsAsserts.Count;
+		WriteLine( string.Format( "EnsuresMarkedAsAsserts: {0}/{1} {2} -> {3}", EnsuresMarkedAsAsserts.Count, Total, StartDate.ToString("yyyy-MM-dd"), EndDate.ToString("yyyy-MM-dd") ) );		
+		StartDate = EndDate;
+	 	EndDate = EndDate.Add( Tick );
+	
+		for (int n = 0; n < EnsuresMarkedAsAsserts.Count; n ++)
+		{
+			EnsuresMarkedAsAssertsBatch.Add(EnsuresMarkedAsAsserts[n]);		
+			EnsuresMarkedAsAsserts[n].Pattern = null;
+			EnsuresMarkedAsAsserts[n].CrashType = 3;
+			EnsuresMarkedAsAsserts[n].Branch = null;
+			
+			if (EnsuresMarkedAsAssertsBatch.Count == NUM_OPS_PER_BATCH)
+			{
+				SumbitChanges(EnsuresMarkedAsAssertsBatch);
+			}
+		}
+	}
+	
+	// Last batch
+	SumbitChanges(EnsuresMarkedAsAssertsBatch);
+}
+/*
+void RemoveErrorMessageForCrashesWithEnsure()
+{
+	var EnsureAssertBatch = new List<Crashes>(NUM_OPS_PER_BATCH);
+	DateTime StartDate = Date;
+	DateTime EndDate = Date.Add( Tick );	
+	int Total = 0;
+	
+	while( EndDate <= UtcNow )
+	{
+		var EnsuresMarkedAsAsserts = Crashes
+		.Where( c => c.TimeOfCrash > StartDate && c.TimeOfCrash <= EndDate )
+		.Where( c => c.CrashType == 1 )
+		.Where( c => c.Description.Contains("FDebug::Ensure") )
+		.ToList();
+	
+		Total += EnsuresMarkedAsAsserts.Count;
+		WriteLine( string.Format( "EnsureAssertCrashes: {0}/{1} {2} -> {3}", EnsuresMarkedAsAsserts.Count, Total, StartDate.ToString("yyyy-MM-dd"), EndDate.ToString("yyyy-MM-dd") ) );		
 		StartDate = EndDate;
 	 	EndDate = EndDate.Add( Tick );
 	
-		for (int n = 0; n < EnsureAssertCrashes.Count; n ++)
+		for (int n = 0; n < EnsuresMarkedAsAsserts.Count; n ++)
 		{
-			EnsureAssertBatch.Add(EnsureAssertCrashes[n]);		
-			EnsureAssertCrashes[n].Pattern = null;
-			EnsureAssertCrashes[n].CrashType = 3;
-			EnsureAssertCrashes[n].Branch = null;
+			EnsureAssertBatch.Add(EnsuresMarkedAsAsserts[n]);		
+			EnsuresMarkedAsAsserts[n].Pattern = null;
+			EnsuresMarkedAsAsserts[n].CrashType = 3;
+			EnsuresMarkedAsAsserts[n].Branch = null;
 			
 			if (EnsureAssertBatch.Count == NUM_OPS_PER_BATCH)
 			{
@@ -374,6 +413,49 @@ void SetNoPatternForEnsureForCrashes()
 	// Last batch
 	SumbitChanges(EnsureAssertBatch);
 }
+*/
+
+void FindAndDeleteEmptyBuggs()
+{
+	var EmptyBuggIds = new List<int>(NUM_OPS_PER_BATCH);
+	DateTime StartDate = Date;
+	DateTime EndDate = Date.Add( Tick );	
+	int Total = 0;
+	
+	while( EndDate <= UtcNow )
+	{
+		var BuggsDateRange = Buggs
+		.Where( c => c.TimeOfLastCrash > StartDate && c.TimeOfLastCrash <= EndDate )
+		.ToList();
+	
+		Total += BuggsDateRange.Count;
+		WriteLine( string.Format( "Buggs: {0}/{1} {2} -> {3}", BuggsDateRange.Count, Total, StartDate.ToString("yyyy-MM-dd"), EndDate.ToString("yyyy-MM-dd") ) );		
+		StartDate = EndDate;
+	 	EndDate = EndDate.Add( Tick );
+	
+		foreach( var Bugg in BuggsDateRange )
+		{
+			// Check if this bugg has any crashes
+			var CrashesForBugg = Buggs_Crashes
+				.Where (bc => bc.BuggId==Bugg.Id)
+				.Count ();
+				
+			// Remove the bugg.
+			if (CrashesForBugg==0)
+			{
+				EmptyBuggIds.Add(Bugg.Id);
+			}
+		}
+	
+		if (EmptyBuggIds.Count == NUM_OPS_PER_BATCH)
+		{
+			DeleteBuggs(EmptyBuggIds);
+		}
+	}
+	
+	// Last batch
+	DeleteBuggs(EmptyBuggIds);
+}
 
 void PrintStats()
 {
@@ -390,11 +472,12 @@ void PrintStats()
 void Main()
 {	
 	PrintStats();
-	
+
 	SetNoPatternForEnsureForCrashes();
 	SetNoPatternForNoBranchAll();
 	SetNoPatternForInvalidPattern();
 	DeleteNoPatternAll();
+	FindAndDeleteEmptyBuggs();
 	
 	PrintStats();
 }
