CrashReportLinq - Added simple script to fix a bugg based on the id

diff --git a/Engine/Source/Programs/CrashReporter/CrashReportWebSite/Queries/FixBuggWithId.linq b/Engine/Source/Programs/CrashReporter/CrashReportWebSite/Queries/FixBuggWithId.linq
new file mode 100644
--- /dev/null
+++ b/Engine/Source/Programs/CrashReporter/CrashReportWebSite/Queries/FixBuggWithId.linq
@@ -0,0 +1,24 @@
+<Query Kind="Program">
+  <Connection>
+    <ID>e40cf789-b18a-4b65-bae6-a0f5440c1ce2</ID>
+    <Persist>true</Persist>
+    <Server>db-09</Server>
+    <Database>CrashReport</Database>
+    <IsProduction>true</IsProduction>
+  </Connection>
+</Query>
+
+void Main()
+{
+	var Remove = Buggs_Crashes.Where (bc => bc.BuggId==53425);
+	Buggs.Where (b => b.Id==53425).First ().CrashType=3;
+	Buggs.Context.SubmitChanges();
+	foreach( var it in Remove )
+	{
+		it.Crash.CrashType=3;
+		Debug.WriteLine( "{0}", it.CrashId );
+	}
+	Buggs_Crashes.Context.SubmitChanges();
+}
+
+// Define other methods and classes here
