CrashReportLinq - A few more scripts

diff --git a/Engine/Source/Programs/CrashReporter/CrashReportWebSite/Queries/(UpdateCrashesByPattern).linq b/Engine/Source/Programs/CrashReporter/CrashReportWebSite/Queries/(UpdateCrashesByPattern).linq
new file mode 100644
--- /dev/null
+++ b/Engine/Source/Programs/CrashReporter/CrashReportWebSite/Queries/(UpdateCrashesByPattern).linq
@@ -0,0 +1,255 @@
+<Query Kind="SQL">
+  <Connection>
+    <ID>e40cf789-b18a-4b65-bae6-a0f5440c1ce2</ID>
+    <Persist>true</Persist>
+    <Server>db-09</Server>
+    <Database>CrashReport</Database>
+    <IsProduction>true</IsProduction>
+  </Connection>
+</Query>
+
+ALTER PROCEDURE [dbo].[UpdateCrashesByPattern]
+AS
+
+BEGIN
+	-- SET NOCOUNT ON added to prevent extra result sets from
+
+	-- interfering with SELECT statements.
+	SET NOCOUNT ON;
+
+	
+	-- Update crash user names
+
+	UPDATE [dbo].[Crashes]
+	SET [UserName] = CASE WHEN c.[UserName] is NULL THEN u.[UserName] ELSE c.[UserName] END
+
+	FROM [dbo].[Crashes] c
+	LEFT JOIN [dbo].[Users] u on (u.[id] = c.[UserNameId]);
+
+	
+	--Create Buggs
+
+	MERGE Buggs AS B
+	USING 
+
+	(
+		SELECT 
+		
+			   [TTPID]
+			  ,[Pattern]
+			  ,[NumberOfCrashes]
+			  ,[NumberOfUsers]
+			  
+			  ,[TimeOfFirstCrash]
+			  ,[TimeOfLastCrash]
+			  
+			  ,[BuildVersion]
+			  ,[CrashType]
+
+			  ,[Status]
+			  ,[FixedChangeList]
+			  
+		FROM
+
+		(
+			SELECT COUNT(1) as NumberOfCrashes
+
+				  , max([TimeOfCrash]) as TimeOfLastCrash
+				  , min([TimeOfCrash]) as TimeOfFirstCrash
+
+				  , max([BuildVersion]) as BuildVersion
+
+				  , max([CrashType]) as CrashType
+				  , count(distinct [UserName]) as NumberOfUsers
+
+				  , Max([Status]) as [Status]
+				  , Max(TTPID) as TTPID
+
+				  , Max(FixedChangeList) as FixedChangeList
+				  , Pattern
+
+			  FROM [dbo].[Crashes] c 
+			  WHERE
+
+				Pattern is not NULL AND Pattern not like ''
+			  group by Pattern
+
+		  ) as s
+		 Where s.NumberOfCrashes > 1
+
+	) AS C 
+	ON (B.Pattern = C.Pattern)
+
+	WHEN NOT MATCHED BY TARGET
+		THEN INSERT (
+
+				[TTPID]
+			  ,[Pattern]
+			  ,[NumberOfCrashes]
+			  ,[NumberOfUsers]
+
+			  ,[TimeOfFirstCrash]
+			  ,[TimeOfLastCrash]
+
+			  ,[BuildVersion]
+			  ,[CrashType]
+			  
+			  ,[Status]
+		  	  ,[FixedChangeList]
+
+			  ) 
+		Values(
+
+				C.[TTPID]
+			  , C.[Pattern]
+			  , C.[NumberOfCrashes]
+			  , C.[NumberOfUsers]
+
+			  , C.[TimeOfFirstCrash]
+			  , C.[TimeOfLastCrash]
+
+			  , C.[BuildVersion]
+			  , C.[CrashType]
+			  
+			  , C.[Status]
+			  , C.[FixedChangeList]
+
+			  ) 
+
+
+	WHEN MATCHED 
+		THEN UPDATE SET B.[NumberOfCrashes] = C.[NumberOfCrashes]
+
+						, B.[TimeOfLastCrash] = C.[TimeOfLastCrash]
+						, B.[BuildVersion] = C.[BuildVersion]
+
+						, B.[CrashType] = C.[CrashType]
+
+						, B.[NumberOfUsers] = C.[NumberOfUsers]
+
+
+	OUTPUT $action, Inserted.*, Deleted.*;
+		
+
+		
+		
+
+	/****** Join Buggs and Crashes  ******/
+	MERGE dbo.Buggs_Crashes BC 
+
+	USING 
+	( 
+
+		SELECT  b.Id as BuggId, c.[Id] as CrashId
+		FROM [dbo].[Crashes] c
+
+		Join [dbo].Buggs b on (b.Pattern = c.Pattern)
+		group by b.Id, c.Id
+
+	 ) AS C
+	 ON BC.BuggId = C.BuggId AND BC.CrashId = C.CrashId
+
+	 WHEN NOT MATCHED BY TARGET
+		THEN INSERT 
+
+			(BuggId, CrashId)
+		VALUES 
+
+			(C.BuggId, C.CrashId)	
+	 WHEN MATCHED
+
+		THEN UPDATE SET
+			BC.BuggId = C.BuggId, 
+
+			BC.CrashId = C.CrashId
+	OUTPUT $action, Inserted.*, Deleted.*;
+
+	
+	
+
+	/****** Join Buggs_Users and Crashes  ******/
+	MERGE dbo.Buggs_Users AS BU
+
+	USING
+		(
+
+		  SELECT  b.Id as BuggId, c.[UserName] as UserName
+		  FROM [dbo].[Crashes] c
+
+		  Join [dbo].Buggs b on (b.Pattern = c.Pattern)
+		  Group by b.Id, c.UserName
+
+		) AS C
+	ON BU.BuggId = C.BuggId AND BU.UserName = C.UserName
+
+	 WHEN NOT MATCHED BY TARGET
+		THEN INSERT 
+
+			(BuggId, UserName) 
+		VALUES 
+
+			(C.BuggId, C.UserName)	
+	 WHEN MATCHED
+
+		THEN UPDATE SET
+			BU.BuggId = C.BuggId, 
+
+			BU.UserName = C.UserName
+	OUTPUT $action, Inserted.*, Deleted.*;
+
+	
+		  
+
+		
+	  
+
+	/****** Join Buggs and UserGroups  ******/
+	MERGE dbo.Buggs_UserGroups AS BUG
+
+	USING
+	(
+
+	  Select BuggId, UserGroupId
+	  From [dbo].[Buggs_Users] bu
+
+	  Join dbo.Users u on (u.UserName = bu.UserName)
+	  Group by BuggId, UserGroupId
+
+	 ) AS BU
+	 ON  BUG.BuggId = BU.BuggId AND BUG.UserGroupId = BU.UserGroupId
+
+	 WHEN NOT MATCHED BY TARGET
+		THEN INSERT 
+
+			(BuggId, UserGroupId) 
+		VALUES 
+
+			(BU.BuggId, BU.UserGroupId)	
+	 WHEN MATCHED
+
+		THEN UPDATE SET
+			BUG.BuggId = BU.BuggId, 
+
+			BUG.UserGroupId = BU.UserGroupId
+	OUTPUT $action, Inserted.*, Deleted.*;
+
+	 
+	/*** Update Crashes to match the Bugg Status, TTPID, and Fixed Change List ****/
+
+	-- This handles new crashes that enter the system; it has the side effect of making the Bugg authoritative in this case but that's how they should be used.
+	update C
+
+	SET 
+		C.TTPID = B.TTPID,
+
+		C.FixedChangeList = B.FixedChangeList,
+		C.[Status] = B.[Status]
+
+	From 
+		dbo.Crashes as C
+
+		Join Buggs_Crashes as BC ON (c.Id = BC.CrashId)
+		Join Buggs as B ON (B.Id =BC.BuggId AND C.Id = BC.CrashId) 
+
+	 
+END
\ No newline at end of file
diff --git a/Engine/Source/Programs/CrashReporter/CrashReportWebSite/Queries/BranchVersionTestQuery.linq b/Engine/Source/Programs/CrashReporter/CrashReportWebSite/Queries/BranchVersionTestQuery.linq
new file mode 100644
--- /dev/null
+++ b/Engine/Source/Programs/CrashReporter/CrashReportWebSite/Queries/BranchVersionTestQuery.linq
@@ -0,0 +1,46 @@
+<Query Kind="Program">
+  <Connection>
+    <ID>e40cf789-b18a-4b65-bae6-a0f5440c1ce2</ID>
+    <Persist>true</Persist>
+    <Server>db-09</Server>
+    <Database>CrashReport</Database>
+    <IsProduction>true</IsProduction>
+  </Connection>
+  <Output>DataGrids</Output>
+</Query>
+
+void Main()
+{
+	var BranchStartsWith = Crashes
+		.Where( n => n.TimeOfCrash > DateTime.Now.AddMonths( -3 ) )
+		.Where( n => n.Branch.StartsWith( "UE4" ) )
+		.Select( n => n.Branch )
+		.Distinct()
+		.ToList();
+					
+	var BranchContains = Crashes
+		.Where( n => n.TimeOfCrash > DateTime.Now.AddMonths( -3 ) )
+		.Where( n => n.Branch.Contains( "UE4" ) )
+		.Select( n => n.Branch )
+		.Distinct()
+		.ToList();
+					
+	var BuildVersions = Crashes.Where( c=> c.TimeOfCrash > DateTime.Now.AddMonths(-3))
+	.Select( c => c.BuildVersion )
+	.Distinct()
+	.ToList();
+	
+	List<string> DistinctBuildVersions = new List<string>();
+	
+	foreach (var BuildVersion in BuildVersions)
+	{
+		var BVParts = BuildVersion.Split( new char[] { '.' }, StringSplitOptions.RemoveEmptyEntries );
+		if (BVParts.Length > 2 && BVParts[0] != "0")
+		{
+			string CleanBV = string.Format( "{0}.{1}.{2}", BVParts[0], BVParts[1], BVParts[2] );
+			DistinctBuildVersions.Add(CleanBV);
+		}
+	}
+}
+
+// Define other methods and classes here
diff --git a/Engine/Source/Programs/CrashReporter/CrashReportWebSite/Queries/DailyCrashesByGame.linq b/Engine/Source/Programs/CrashReporter/CrashReportWebSite/Queries/DailyCrashesByGame.linq
new file mode 100644
--- /dev/null
+++ b/Engine/Source/Programs/CrashReporter/CrashReportWebSite/Queries/DailyCrashesByGame.linq
@@ -0,0 +1,54 @@
+<Query Kind="Program">
+  <Connection>
+    <ID>e40cf789-b18a-4b65-bae6-a0f5440c1ce2</ID>
+    <Persist>true</Persist>
+    <Server>db-09</Server>
+    <Database>CrashReport</Database>
+    <IsProduction>true</IsProduction>
+  </Connection>
+  <Output>DataGrids</Output>
+</Query>
+
+void Main()
+{
+	DateTime StartDate = new DateTime( 2014,4,1 );
+	DateTime EndDate = new DateTime(2014,4,10);
+	TimeSpan TimeSpan = TimeSpan.FromDays(1);
+	
+	while (StartDate < EndDate)
+	{
+		DateTime TickEndDate = StartDate+TimeSpan;
+		Debug.WriteLine( "StartDate: {0} EndDate: {1}", StartDate, TickEndDate );
+		var Filters = Crashes
+		.Where (c => c.TimeOfCrash > StartDate && c.TimeOfCrash < TickEndDate)
+		.Where (c => c.CrashType == 1 || c.CrashType == 2)
+		//.Where (c => c.UserName == "Anonymous")
+		//.Select (c => c.GameName)
+		.GroupBy (c => c.GameName)
+		.OrderByDescending (c => c.Count ());
+		
+		var Countt = Crashes
+		.Where (c => c.TimeOfCrash > StartDate && c.TimeOfCrash < TickEndDate)
+		.Where (c => c.CrashType == 1 || c.CrashType == 2)
+		//.Where (c => c.UserName == "Anonymous")
+		//.Select (c => c.GameName)
+		.Count ();
+		
+		int Total = 0;
+		int Top = 0;
+		foreach( var Element in Filters )
+		{
+			Debug.WriteLine( "	GameName: {0} Crashes: {1}", Element.Key, Element.Count () );
+			Top++;
+			Total += Element.Count ();
+			if (Top > 10)
+			{
+				break;
+			}
+		}
+		Debug.WriteLine( "	Total: {0}/{1}", Total, Countt );
+		StartDate += TimeSpan;
+	}
+}
+
+// Define other methods and classes here
