MinidumpDiagnostics - Added a simple crash handler that prints the callstack to the log

diff --git a/Engine/Source/Developer/CrashDebugHelper/Private/CrashDebugPDBCache.cpp b/Engine/Source/Developer/CrashDebugHelper/Private/CrashDebugPDBCache.cpp
--- a/Engine/Source/Developer/CrashDebugHelper/Private/CrashDebugPDBCache.cpp
+++ b/Engine/Source/Developer/CrashDebugHelper/Private/CrashDebugPDBCache.cpp
@@ -325,7 +325,7 @@ FPDBCacheEntryRef FPDBCache::ReadPDBCacheEntry( const FString& Directory )
 	else
 	{
 		// Something wrong.
-		check( 0 );
+		checkf( 0, TEXT( "Invalid symbol cache entry: %s" ), *EntryDirectory );
 	}
 
 	return NewEntry.ToSharedRef();
diff --git a/Engine/Source/Programs/CrashReporter/MinidumpDiagnostics/Private/Windows/MinidumpDiagnosticsMainWindows.cpp b/Engine/Source/Programs/CrashReporter/MinidumpDiagnostics/Private/Windows/MinidumpDiagnosticsMainWindows.cpp
--- a/Engine/Source/Programs/CrashReporter/MinidumpDiagnostics/Private/Windows/MinidumpDiagnosticsMainWindows.cpp
+++ b/Engine/Source/Programs/CrashReporter/MinidumpDiagnostics/Private/Windows/MinidumpDiagnosticsMainWindows.cpp
@@ -14,6 +14,19 @@ IMPLEMENT_APPLICATION( MinidumpDiagnostics, "MinidumpDiagnostics" )
  */
 int32 EmptyCrashHandler( LPEXCEPTION_POINTERS ExceptionInfo )
 {
+	// Simple crash handler that prints the callstack to the log.
+	const SIZE_T StackTraceSize = 65535;
+	ANSICHAR* StackTrace = (ANSICHAR*)GMalloc->Malloc( StackTraceSize );
+	StackTrace[0] = 0;
+	FPlatformStackWalk::StackWalkAndDump( StackTrace, StackTraceSize, 0, ExceptionInfo->ContextRecord );
+
+	FCString::Strncat( GErrorHist, TEXT( "\r\n\r\n" ), ARRAY_COUNT( GErrorHist ) );
+	FCString::Strncat( GErrorHist, ANSI_TO_TCHAR( StackTrace ), ARRAY_COUNT( GErrorHist ) );
+
+	GMalloc->Free( StackTrace );
+
+	GError->HandleError();
+
 	return EXCEPTION_EXECUTE_HANDLER;
 }
 
