CrashReportLinq - Remove assert message when the crash type is crash

diff --git a/Engine/Source/Programs/CrashReporter/CrashReportWebSite/Queries/DailyCrashesByGame.linq b/Engine/Source/Programs/CrashReporter/CrashReportWebSite/Queries/DailyCrashesByGame.linq
--- a/Engine/Source/Programs/CrashReporter/CrashReportWebSite/Queries/DailyCrashesByGame.linq
+++ b/Engine/Source/Programs/CrashReporter/CrashReportWebSite/Queries/DailyCrashesByGame.linq
@@ -11,8 +11,8 @@
 
 void Main()
 {
-	DateTime StartDate = new DateTime( 2014,4,1 );
-	DateTime EndDate = new DateTime(2014,4,10);
+	DateTime StartDate = new DateTime( 2015,5,1 );
+	DateTime EndDate = new DateTime(2015,5,22);
 	TimeSpan TimeSpan = TimeSpan.FromDays(1);
 	
 	while (StartDate < EndDate)
@@ -51,4 +51,4 @@ void Main()
 	}
 }
 
-// Define other methods and classes here
+// Define other methods and classes here
\ No newline at end of file
diff --git a/Engine/Source/Programs/CrashReporter/CrashReportWebSite/Queries/Fix-RepairDatabase.linq b/Engine/Source/Programs/CrashReporter/CrashReportWebSite/Queries/Fix-RepairDatabase.linq
--- a/Engine/Source/Programs/CrashReporter/CrashReportWebSite/Queries/Fix-RepairDatabase.linq
+++ b/Engine/Source/Programs/CrashReporter/CrashReportWebSite/Queries/Fix-RepairDatabase.linq
@@ -33,7 +33,7 @@ HashSet<int> GetInvalidFuncIds()
 }
 
 static Stopwatch Timer = Stopwatch.StartNew();
-static DateTime Date = DateTime.UtcNow.AddMonths(-1);   // From
+static DateTime Date = DateTime.UtcNow.AddMonths(-6);   // From
 static DateTime UtcNow = DateTime.UtcNow.AddHours(-2); // To
 static TimeSpan Tick = TimeSpan.FromDays(1);
 const int NUM_OPS_PER_BATCH = 16;
@@ -338,6 +338,44 @@ void SetNoPatternForInvalidPattern()
 	}
 }
 
+void ClearErrorMessageForCrashesWithAssertMessage()
+{
+	//Assertion failed:
+	
+	var CrashesBatch = new List<Crashes>(NUM_OPS_PER_BATCH);
+	DateTime StartDate = Date;
+	DateTime EndDate = Date.Add( Tick );	
+	int Total = 0;
+	
+	while( EndDate <= UtcNow )
+	{
+		var CrashesWithAssertMessage = Crashes
+		.Where( c => c.TimeOfCrash > StartDate && c.TimeOfCrash <= EndDate )
+		.Where( c => c.CrashType == 1 )
+		.Where( c => c.Summary.Contains("Assertion failed:") )
+		.ToList();
+	
+		Total += CrashesWithAssertMessage.Count;
+		WriteLine( string.Format( "CrashesWithAssertMessage: {0}/{1} {2} -> {3}", CrashesWithAssertMessage.Count, Total, StartDate.ToString("yyyy-MM-dd"), EndDate.ToString("yyyy-MM-dd") ) );		
+		StartDate = EndDate;
+	 	EndDate = EndDate.Add( Tick );
+	
+		for (int n = 0; n < CrashesWithAssertMessage.Count; n ++)
+		{
+			CrashesBatch.Add(CrashesWithAssertMessage[n]);		
+			CrashesWithAssertMessage[n].Summary = "";
+			
+			if (CrashesBatch.Count == NUM_OPS_PER_BATCH)
+			{
+				SumbitChanges(CrashesBatch);
+			}
+		}
+	}
+	
+	// Last batch
+	SumbitChanges(CrashesBatch);
+}
+
 void SetNoPatternForEnsureForCrashes()
 {
 	var EnsuresMarkedAsAssertsBatch = new List<Crashes>(NUM_OPS_PER_BATCH);
@@ -445,12 +483,12 @@ void FindAndDeleteEmptyBuggs()
 			{
 				EmptyBuggIds.Add(Bugg.Id);
 			}
-		}
-	
-		if (EmptyBuggIds.Count == NUM_OPS_PER_BATCH)
-		{
-			DeleteBuggs(EmptyBuggIds);
-		}
+			
+			if (EmptyBuggIds.Count == NUM_OPS_PER_BATCH)
+			{
+				DeleteBuggs(EmptyBuggIds);
+			}
+		}		
 	}
 	
 	// Last batch
@@ -473,6 +511,7 @@ void Main()
 {	
 	PrintStats();
 
+	ClearErrorMessageForCrashesWithAssertMessage();
 	SetNoPatternForEnsureForCrashes();
 	SetNoPatternForNoBranchAll();
 	SetNoPatternForInvalidPattern();
diff --git a/Engine/Source/Programs/CrashReporter/CrashReportWebSite/Queries/FixBuggWithId.linq b/Engine/Source/Programs/CrashReporter/CrashReportWebSite/Queries/FixBuggWithId.linq
--- a/Engine/Source/Programs/CrashReporter/CrashReportWebSite/Queries/FixBuggWithId.linq
+++ b/Engine/Source/Programs/CrashReporter/CrashReportWebSite/Queries/FixBuggWithId.linq
@@ -10,12 +10,13 @@
 
 void Main()
 {
-	var Remove = Buggs_Crashes.Where (bc => bc.BuggId==53425);
-	Buggs.Where (b => b.Id==53425).First ().CrashType=3;
+	const int BuggId = 53425;
+	var Remove = Buggs_Crashes.Where (bc => bc.BuggId==BuggId);
+	Buggs.Where (b => b.Id==BuggId).First ().CrashType=3; // Ensure
 	Buggs.Context.SubmitChanges();
 	foreach( var it in Remove )
 	{
-		it.Crash.CrashType=3;
+		it.Crash.CrashType=3; // Ensure
 		Debug.WriteLine( "{0}", it.CrashId );
 	}
 	Buggs_Crashes.Context.SubmitChanges();
