CrashReportLinq - Tweaks

diff --git a/Engine/Source/Programs/CrashReporter/CrashReportWebSite/Queries/Fix-RepairDatabase.linq b/Engine/Source/Programs/CrashReporter/CrashReportWebSite/Queries/Fix-RepairDatabase.linq
--- a/Engine/Source/Programs/CrashReporter/CrashReportWebSite/Queries/Fix-RepairDatabase.linq
+++ b/Engine/Source/Programs/CrashReporter/CrashReportWebSite/Queries/Fix-RepairDatabase.linq
@@ -33,7 +33,7 @@ HashSet<int> GetInvalidFuncIds()
 }
 
 static Stopwatch Timer = Stopwatch.StartNew();
-static DateTime Date = DateTime.UtcNow.AddMonths(-6);   // From
+static DateTime Date = DateTime.UtcNow.AddMonths(-1);   // From
 static DateTime UtcNow = DateTime.UtcNow.AddHours(-2); // To
 static TimeSpan Tick = TimeSpan.FromDays(1);
 const int NUM_OPS_PER_BATCH = 16;
@@ -340,8 +340,6 @@ void SetNoPatternForInvalidPattern()
 
 void ClearErrorMessageForCrashesWithAssertMessage()
 {
-	//Assertion failed:
-	
 	var CrashesBatch = new List<Crashes>(NUM_OPS_PER_BATCH);
 	DateTime StartDate = Date;
 	DateTime EndDate = Date.Add( Tick );	
diff --git a/Engine/Source/Programs/CrashReporter/CrashReportWebSite/Queries/FixBuggWithId.linq b/Engine/Source/Programs/CrashReporter/CrashReportWebSite/Queries/FixBuggWithId.linq
--- a/Engine/Source/Programs/CrashReporter/CrashReportWebSite/Queries/FixBuggWithId.linq
+++ b/Engine/Source/Programs/CrashReporter/CrashReportWebSite/Queries/FixBuggWithId.linq
@@ -10,14 +10,24 @@
 
 void Main()
 {
-	const int BuggId = 53425;
-	var Remove = Buggs_Crashes.Where (bc => bc.BuggId==BuggId);
+	const int BuggId = 56230;
+	var BuggCrashes = Buggs_Crashes.Where (bc => bc.BuggId==BuggId);
 	Buggs.Where (b => b.Id==BuggId).First ().CrashType=3; // Ensure
 	Buggs.Context.SubmitChanges();
-	foreach( var it in Remove )
+	var NumToChange = BuggCrashes.Count();
+	int Index = 0;
+	foreach( var It in BuggCrashes )
 	{
-		it.Crash.CrashType=3; // Ensure
-		Debug.WriteLine( "{0}", it.CrashId );
+		Index ++;
+		if (It.Crash.CrashType != 3)
+		{
+			It.Crash.CrashType=3; // Ensure
+			Debug.WriteLine( "{0} {1,3}/{2,3}", It.CrashId, Index, NumToChange );
+			if (Index % 16 == 0)
+			{
+				Buggs_Crashes.Context.SubmitChanges();
+			}
+		}
 	}
 	Buggs_Crashes.Context.SubmitChanges();
 }
