MinidumpDiagnostics - Updated simple crash handler

diff --git a/Engine/Source/Programs/CrashReporter/MinidumpDiagnostics/Private/Windows/MinidumpDiagnosticsMainWindows.cpp b/Engine/Source/Programs/CrashReporter/MinidumpDiagnostics/Private/Windows/MinidumpDiagnosticsMainWindows.cpp
--- a/Engine/Source/Programs/CrashReporter/MinidumpDiagnostics/Private/Windows/MinidumpDiagnosticsMainWindows.cpp
+++ b/Engine/Source/Programs/CrashReporter/MinidumpDiagnostics/Private/Windows/MinidumpDiagnosticsMainWindows.cpp
@@ -10,11 +10,10 @@ IMPLEMENT_APPLICATION( MinidumpDiagnostics, "MinidumpDiagnostics" )
 
 	
 /** 
- * A null crash handler to suppress error report generation
+ * A simple crash handler that prints the callstack to the log
  */
-int32 EmptyCrashHandler( LPEXCEPTION_POINTERS ExceptionInfo )
+int32 SimpleCrashHandler( LPEXCEPTION_POINTERS ExceptionInfo )
 {
-	// Simple crash handler that prints the callstack to the log.
 	const SIZE_T StackTraceSize = 65535;
 	ANSICHAR* StackTrace = (ANSICHAR*)GMalloc->Malloc( StackTraceSize );
 	StackTrace[0] = 0;
@@ -27,6 +26,8 @@ int32 EmptyCrashHandler( LPEXCEPTION_POINTERS ExceptionInfo )
 
 	GError->HandleError();
 
+	FPlatformMisc::RequestExit( true );
+
 	return EXCEPTION_EXECUTE_HANDLER;
 }
 
@@ -48,13 +49,10 @@ int32 GuardedMainWrapper(int32 ArgC, TCHAR* ArgV[])
 {
 	int32 ReturnCode = 0;
 
-#if !PLATFORM_MAC
 	if (FPlatformMisc::IsDebuggerPresent())
-#endif
 	{
 		ReturnCode = GuardedMain( ArgC, ArgV );
 	}
-#if !PLATFORM_MAC
 	else
 	{
 		__try
@@ -63,11 +61,10 @@ int32 GuardedMainWrapper(int32 ArgC, TCHAR* ArgV[])
 			ReturnCode = GuardedMain( ArgC, ArgV );
 			GIsGuarded = 0;
 		}
-		__except( EmptyCrashHandler( GetExceptionInformation() ) )
+		__except( SimpleCrashHandler( GetExceptionInformation() ) )
 		{
 		}
 	}
-#endif
 
 	FEngineLoop::AppPreExit();
 	FEngineLoop::AppExit();
