CrashReportClient - UECORE-160/Process 'ensures' with the crash report client

diff --git a/Engine/Source/Programs/CrashReporter/CrashReportClient/Private/CrashReportClientUnattended.cpp b/Engine/Source/Programs/CrashReporter/CrashReportClient/Private/CrashReportClientUnattended.cpp
--- a/Engine/Source/Programs/CrashReporter/CrashReportClient/Private/CrashReportClientUnattended.cpp
+++ b/Engine/Source/Programs/CrashReporter/CrashReportClient/Private/CrashReportClientUnattended.cpp
@@ -6,9 +6,11 @@
 FCrashReportClientUnattended::FCrashReportClientUnattended(const FPlatformErrorReport& ErrorReport)
 	: Uploader( FCrashReportClientConfig::Get().GetReceiverAddress() )
 {
-	// Prevent uploader waiting for a diagnosis file
-	Uploader.LocalDiagnosisSkipped();	
-	Uploader.BeginUpload(ErrorReport);
+	// Process the report synchronously
+	ErrorReport.DiagnoseReport();
+	auto DiagnosticsFilePath = ErrorReport.GetReportDirectory() / FCrashReportClientConfig::Get().GetDiagnosticsFilename();
+	Uploader.LocalDiagnosisComplete( FPaths::FileExists( DiagnosticsFilePath ) ? DiagnosticsFilePath : TEXT( "" ) );
+	Uploader.BeginUpload( ErrorReport );
 	StartTicker();
 }
 
diff --git a/Engine/Source/Programs/CrashReporter/CrashReportClient/Private/MainLoopTiming.cpp b/Engine/Source/Programs/CrashReporter/CrashReportClient/Private/MainLoopTiming.cpp
--- a/Engine/Source/Programs/CrashReporter/CrashReportClient/Private/MainLoopTiming.cpp
+++ b/Engine/Source/Programs/CrashReporter/CrashReportClient/Private/MainLoopTiming.cpp
@@ -8,14 +8,15 @@
 
 FMainLoopTiming::FMainLoopTiming(float InIdealTickRate, EMainLoopOptions::Type Options)
 	: IdealFrameTime(1.f / InIdealTickRate)
-	, ActualDeltaTime(0)
-	, LastTime(FPlatformTime::Seconds())
 	, bTickSlate(Options & EMainLoopOptions::UsingSlate)
 {
 }
 
 void FMainLoopTiming::Tick()
 {
+	static double ActualDeltaTime = IdealFrameTime;
+	static double LastTime = FPlatformTime::Seconds();
+
 	// Tick app logic
 	FTaskGraphInterface::Get().ProcessThreadUntilIdle(ENamedThreads::GameThread);
 	FTicker::GetCoreTicker().Tick(ActualDeltaTime);
@@ -29,9 +30,6 @@ void FMainLoopTiming::Tick()
 	}
 #endif // !CRASH_REPORT_UNATTENDED_ONLY
 
-	static const float TickRate = 30.f;
-	static const float TickInterval = 1.f / TickRate;
-
 	// Sleep Throttling
 	// Copied from Community Portal - should be shared
 	FPlatformProcess::Sleep(FMath::Max<float>(0, IdealFrameTime - (FPlatformTime::Seconds() - LastTime)));
diff --git a/Engine/Source/Programs/CrashReporter/CrashReportClient/Private/MainLoopTiming.h b/Engine/Source/Programs/CrashReporter/CrashReportClient/Private/MainLoopTiming.h
--- a/Engine/Source/Programs/CrashReporter/CrashReportClient/Private/MainLoopTiming.h
+++ b/Engine/Source/Programs/CrashReporter/CrashReportClient/Private/MainLoopTiming.h
@@ -38,12 +38,6 @@ private:
 	/** Interval between each Tick that we're aiming for */
 	float IdealFrameTime;
 
-	/** Actual time between Tick and previous Tick */
-	double ActualDeltaTime;
-
-	/** Absolute time of previous Tick */
-	double LastTime;
-
 	/** Should Slate's Tick be called? */
 	bool bTickSlate;
 };
