[NO NON-PII COMMIT MESSAGE]

diff --git a/Engine/Source/Programs/CrashReporter/CrashReportClient/Private/CrashReportClientUnattended.cpp b/Engine/Source/Programs/CrashReporter/CrashReportClient/Private/CrashReportClientUnattended.cpp
--- a/Engine/Source/Programs/CrashReporter/CrashReportClient/Private/CrashReportClientUnattended.cpp
+++ b/Engine/Source/Programs/CrashReporter/CrashReportClient/Private/CrashReportClientUnattended.cpp
@@ -2,6 +2,7 @@
 
 #include "CrashReportClientApp.h"
 #include "CrashReportClientUnattended.h"
+#include "CrashReportUtil.h"
 
 FCrashReportClientUnattended::FCrashReportClientUnattended(const FPlatformErrorReport& ErrorReport)
 	: Uploader( FCrashReportClientConfig::Get().GetReceiverAddress() )
