CrashReport - Frontend - UECORE-176/CRASH REPORTER: Automatically include "liveissue" label with CopyToJira

diff --git a/Engine/Source/Programs/CrashReporter/CrashReportWebSite/Models/CrashReport.cs b/Engine/Source/Programs/CrashReporter/CrashReportWebSite/Models/CrashReport.cs
--- a/Engine/Source/Programs/CrashReporter/CrashReportWebSite/Models/CrashReport.cs
+++ b/Engine/Source/Programs/CrashReporter/CrashReportWebSite/Models/CrashReport.cs
@@ -305,7 +305,7 @@ namespace Tools.CrashReporter.CrashReportWebSite.Models
 				Fields.Add( "summary", "[CrashReport] " + ToJiraSummary );						// Call Stack, Line 1
 				Fields.Add( "description", string.Join( "\r\n", ToJiraDescriptions ) );			// Description
 				Fields.Add( "issuetype", new Dictionary<string, object> { { "id", "1" } } );	// Bug
-				Fields.Add( "labels", new string[] { "crash" } );								// <label>crash</label>
+				Fields.Add( "labels", new string[] { "crash", "liveissue" } );					// <label>crash, live issue</label>
 				Fields.Add( "customfield_11500", ToJiraFirstCLAffected );						// Changelist # / Found Changelist
 				Fields.Add( "environment", LatestOSAffected );									// Platform
 
