CrashReporter-Frontend - UECORE-159/Crash Reporter: Dashboard update to graph binary crashes per Build Version

diff --git a/Engine/Source/Programs/CrashReporter/CrashReportWebSite/Controllers/BuggsController.cs b/Engine/Source/Programs/CrashReporter/CrashReportWebSite/Controllers/BuggsController.cs
--- a/Engine/Source/Programs/CrashReporter/CrashReportWebSite/Controllers/BuggsController.cs
+++ b/Engine/Source/Programs/CrashReporter/CrashReportWebSite/Controllers/BuggsController.cs
@@ -127,14 +127,18 @@ namespace Tools.CrashReporter.CrashReportWebSite.Controllers
 
 				Crashes = NewBugg.GetCrashes();
 
-				using( FAutoScopedLogTimer GetCrashesTimer = new FAutoScopedLogTimer( "Bugg.PrepareBuggForJira" ) )
+				using (FAutoScopedLogTimer GetCrashesTimer = new FAutoScopedLogTimer( "Bugg.PrepareBuggForJira" ))
 				{
-					NewBugg.PrepareBuggForJira( Crashes );
-
-					if( BuggIDToBeAddedToJira != 0 )
+					if (Crashes.Count > 0)
 					{
-						NewBugg.CopyToJira();
-					}			
+						NewBugg.PrepareBuggForJira( Crashes );
+
+						if (BuggIDToBeAddedToJira != 0)
+						{
+							NewBugg.CopyToJira();
+						}
+					}
+
 				}
 
 				using( FAutoScopedLogTimer JiraResultsTimer = new FAutoScopedLogTimer( "Bugg.GrabJira" ) )
diff --git a/Engine/Source/Programs/CrashReporter/CrashReportWebSite/Controllers/DashboardController.cs b/Engine/Source/Programs/CrashReporter/CrashReportWebSite/Controllers/DashboardController.cs
--- a/Engine/Source/Programs/CrashReporter/CrashReportWebSite/Controllers/DashboardController.cs
+++ b/Engine/Source/Programs/CrashReporter/CrashReportWebSite/Controllers/DashboardController.cs
@@ -25,26 +25,22 @@ namespace Tools.CrashReporter.CrashReportWebSite.Controllers
 		/// </summary>
 		public readonly int UserId;
 
-
 		/// <summary>
-		/// Constructor
+		/// 
 		/// </summary>
-		/// <param name="Crash">Crash</param>
-		public FCrashMinimal( Crash Crash )
-		{
-			TimeOfCrash = Crash.TimeOfCrash.Value;
-			UserId = Crash.UserNameId.Value;
-		}
+		public readonly string EngineVersion;
 
 		/// <summary>
 		/// 
 		/// </summary>
 		/// <param name="InTimeOfCrash"></param>
 		/// <param name="InUserId"></param>
-		public FCrashMinimal( DateTime InTimeOfCrash, int InUserId )
+		/// <param name="InEngineVersion"></param>
+		public FCrashMinimal( DateTime InTimeOfCrash, int InUserId, string InEngineVersion )
 		{
 			TimeOfCrash = InTimeOfCrash;
 			UserId = InUserId;
+			EngineVersion = InEngineVersion;
 		}
 	}
 
@@ -98,6 +94,43 @@ namespace Tools.CrashReporter.CrashReportWebSite.Controllers
 			}
 		}
 
+		/// <summary>
+		/// Return a dictionary of crashes per version grouped by week.
+		/// </summary>
+		/// <param name="Crashes">A set of crashes to interrogate.</param>
+		/// <param name="EngineVersion">Engine version</param>
+		/// <param name="AnonymousID">Anonymous id</param>
+		/// <returns>A dictionary of week vs. crash count.</returns>
+		public Dictionary<DateTime, int> GetWeeklyCountsByVersion( List<FCrashMinimal> Crashes, string EngineVersion, int AnonymousID )
+		{
+			using (FAutoScopedLogTimer LogTimer = new FAutoScopedLogTimer( this.GetType().ToString() + "(EngineVersion=" + EngineVersion + ")" ))
+			{
+				Dictionary<DateTime, int> Results = new Dictionary<DateTime, int>();
+
+				Crashes = Crashes
+					.Where( Crash => Crash.UserId == AnonymousID )
+					.Where( Crash => Crash.EngineVersion.Contains( EngineVersion ) )
+					.ToList();
+
+				try
+				{
+					Results =
+					(
+						from CrashDetail in Crashes
+						group CrashDetail by CrashDetail.TimeOfCrash.AddDays( -(int)CrashDetail.TimeOfCrash.DayOfWeek ).Date into VersionCount
+						orderby VersionCount.Key
+						select new { Count = VersionCount.Count(), Date = VersionCount.Key }
+					).ToDictionary( x => x.Date, y => y.Count );
+				}
+				catch (Exception Ex)
+				{
+					Debug.WriteLine( "Exception in GeWeeklyCountsByVersion: " + Ex.ToString() );
+				}
+
+				return Results;
+			}
+		}
+
 		/// <summary>
 		/// Return a dictionary of crashes per group grouped by day.
 		/// </summary>
@@ -137,6 +170,43 @@ namespace Tools.CrashReporter.CrashReportWebSite.Controllers
 			}
 		}
 
+		/// <summary>
+		/// Return a dictionary of crashes per version grouped by day.
+		/// </summary>
+		/// <param name="Crashes">A set of crashes to interrogate.</param>
+		/// <param name="EngineVersion">Engine version</param>
+		/// <param name="AnonymousID">Anonymous id</param>
+		/// <returns>A dictionary of day vs. crash count.</returns>
+		public Dictionary<DateTime, int> GetDailyCountsByVersion( List<FCrashMinimal> Crashes, string EngineVersion, int AnonymousID )
+		{
+			using (FAutoScopedLogTimer LogTimer = new FAutoScopedLogTimer( this.GetType().ToString() + "(EngineVersion=" + EngineVersion + ")" ))
+			{
+				Dictionary<DateTime, int> Results = new Dictionary<DateTime, int>();
+
+				Crashes = Crashes
+					.Where( Crash => Crash.UserId == AnonymousID )
+					.Where( Crash => Crash.EngineVersion.Contains( EngineVersion ) )
+					.ToList();
+				
+				try
+				{
+					Results =
+					(
+						from CrashDetail in Crashes
+						group CrashDetail by CrashDetail.TimeOfCrash.Date into VersionCount
+						orderby VersionCount.Key
+						select new { Count = VersionCount.Count(), Date = VersionCount.Key }
+					).ToDictionary( x => x.Date, y => y.Count );
+				}
+				catch (Exception Ex)
+				{
+					Debug.WriteLine( "Exception in GetDailyCountsByVersion: " + Ex.ToString() );
+				}
+
+				return Results;
+			}
+		}
+
 		/// <summary>
 		/// The main view of the dash board.
 		/// </summary>
@@ -150,36 +220,75 @@ namespace Tools.CrashReporter.CrashReportWebSite.Controllers
 
 				FAutoScopedLogTimer LogTimerSQL = new FAutoScopedLogTimer( "CrashesFilterByDate", "", "" );
 
+				// Get engine versions from the last 6 months.
+				var TempVersions = CrashRepository.GetVersions();
+				List<string> EngineUE4Versions = new List<string>();
+				foreach (var Version in TempVersions)
+				{
+					if( Version.StartsWith("4.") )
+					{
+						EngineUE4Versions.Add( Version );
+					}
+				}
+
+				// Only 4 latest version.
+				EngineUE4Versions = EngineUE4Versions.OrderByDescending( item => item ).Take( 5 ).ToList();
+
 				IQueryable<Crash> CrashesInTimeFrame = _Crashes.ListAll()
 					.Where( MyCrash => MyCrash.TimeOfCrash >= AfewMonthsAgo && MyCrash.TimeOfCrash <= Today.AddDays( 1 ) );
 				//IEnumerable<Crash> Crashes = FRepository.Get().Crashes.FilterByDate( FRepository.Get().Crashes.ListAll(), AfewMonthsAgo, Today );
-				var VMinimalCrashes = CrashesInTimeFrame.Select( Crash => new { TimeOfCrash = Crash.TimeOfCrash.Value, UserID = Crash.UserNameId.Value } ).ToList();
+				var VMinimalCrashes = CrashesInTimeFrame.Select( Crash => new 
+				{ 
+					TimeOfCrash = Crash.TimeOfCrash.Value, 
+					UserID = Crash.UserNameId.Value,
+					EngineVersion = Crash.BuildVersion,
+				} )
+				.ToList();
 
 				List<FCrashMinimal> MinimalCrashes = new List<FCrashMinimal>( VMinimalCrashes.Count );
 				foreach( var Anotype in VMinimalCrashes )
 				{
-					MinimalCrashes.Add( new FCrashMinimal( Anotype.TimeOfCrash, Anotype.UserID ) );
+					MinimalCrashes.Add( new FCrashMinimal( Anotype.TimeOfCrash, Anotype.UserID, Anotype.EngineVersion ) );
 				}
 				LogTimerSQL.Dispose();
 
+				HashSet<int> AnonumousIDs = FRepository.Get( _Crashes ).GetUserIdsFromUserGroup( "Anonymous" );
+				int AnonymousID = AnonumousIDs.First();
+
 				int GeneralUserGroupId = FRepository.Get( _Crashes ).FindOrAddGroup( "General" );
 				int CoderUserGroupId = FRepository.Get( _Crashes ).FindOrAddGroup( "Coder" );
 				int EngineQAUserGroupId = FRepository.Get( _Crashes ).FindOrAddGroup( "EngineQA" );
 				int GameQAUserGroupId = FRepository.Get( _Crashes ).FindOrAddGroup( "GameQA" );
-				int AnonymousUserGroupId = FRepository.Get( _Crashes ).FindOrAddGroup( "Anonymous" );
 
-				Dictionary<DateTime, int> GeneralResults = GetWeeklyCountsByGroup( MinimalCrashes, GeneralUserGroupId );
-				Dictionary<DateTime, int> CoderResults = GetWeeklyCountsByGroup( MinimalCrashes, CoderUserGroupId );
-				Dictionary<DateTime, int> EngineQAResults = GetWeeklyCountsByGroup( MinimalCrashes, EngineQAUserGroupId );
-				Dictionary<DateTime, int> GameQAResults = GetWeeklyCountsByGroup( MinimalCrashes, GameQAUserGroupId );
-				Dictionary<DateTime, int> AnonymousResults = GetWeeklyCountsByGroup( MinimalCrashes, AnonymousUserGroupId );
-				Dictionary<DateTime, int> AllResults = GetWeeklyCountsByGroup( MinimalCrashes, AllUserGroupId );
+				// Weekly
+				Dictionary<DateTime, int> WeeklyGeneralResults = GetWeeklyCountsByGroup( MinimalCrashes, GeneralUserGroupId );
+				Dictionary<DateTime, int> WeeklyCoderResults = GetWeeklyCountsByGroup( MinimalCrashes, CoderUserGroupId );
+				Dictionary<DateTime, int> WeeklyEngineQAResults = GetWeeklyCountsByGroup( MinimalCrashes, EngineQAUserGroupId );
+				Dictionary<DateTime, int> WeeklyGameQAResults = GetWeeklyCountsByGroup( MinimalCrashes, GameQAUserGroupId );
+
+				Dictionary<string, Dictionary<DateTime, int>> WeeklyEngineVersionResults = new Dictionary<string, Dictionary<DateTime, int>>();
+				foreach (string UE4Version in EngineUE4Versions)
+				{
+					Dictionary<DateTime, int> Results = GetWeeklyCountsByVersion( MinimalCrashes, UE4Version, AnonymousID );
+					WeeklyEngineVersionResults.Add( UE4Version, Results );
+				}
 
+				Dictionary<DateTime, int> WeeklyAllResults = GetWeeklyCountsByGroup( MinimalCrashes, AllUserGroupId );
+				
+				// Daily
 				Dictionary<DateTime, int> DailyGeneralResults = GetDailyCountsByGroup( MinimalCrashes, GeneralUserGroupId );
 				Dictionary<DateTime, int> DailyCoderResults = GetDailyCountsByGroup( MinimalCrashes, CoderUserGroupId );
 				Dictionary<DateTime, int> DailyEngineQAResults = GetDailyCountsByGroup( MinimalCrashes, EngineQAUserGroupId );
 				Dictionary<DateTime, int> DailyGameQAResults = GetDailyCountsByGroup( MinimalCrashes, GameQAUserGroupId );
-				Dictionary<DateTime, int> DailyAnonymousResults = GetDailyCountsByGroup( MinimalCrashes, AnonymousUserGroupId );
+
+
+				Dictionary<string, Dictionary<DateTime, int>> DailyEngineVersionResults = new Dictionary<string, Dictionary<DateTime, int>>();
+				foreach(string UE4Version in EngineUE4Versions)
+				{
+					Dictionary<DateTime, int> Results = GetDailyCountsByVersion( MinimalCrashes, UE4Version, AnonymousID );
+					DailyEngineVersionResults.Add( UE4Version, Results );
+				}
+
 				Dictionary<DateTime, int> DailyAllResults = GetDailyCountsByGroup( MinimalCrashes, AllUserGroupId );
 
 				// Get daily buggs stats.
@@ -196,33 +305,40 @@ namespace Tools.CrashReporter.CrashReportWebSite.Controllers
 
 				string CrashesByWeek = "";
 
-				foreach( KeyValuePair<DateTime, int> Result in AllResults )
+				foreach( KeyValuePair<DateTime, int> WeeklyResult in WeeklyAllResults )
 				{
 					int GeneralCrashes = 0;
-					GeneralResults.TryGetValue( Result.Key, out GeneralCrashes );
+					WeeklyGeneralResults.TryGetValue( WeeklyResult.Key, out GeneralCrashes );
 
 					int CoderCrashes = 0;
-					CoderResults.TryGetValue( Result.Key, out CoderCrashes );
+					WeeklyCoderResults.TryGetValue( WeeklyResult.Key, out CoderCrashes );
 
 					int EngineQACrashes = 0;
-					EngineQAResults.TryGetValue( Result.Key, out EngineQACrashes );
+					WeeklyEngineQAResults.TryGetValue( WeeklyResult.Key, out EngineQACrashes );
 
 					int GameQACrashes = 0;
-					GameQAResults.TryGetValue( Result.Key, out GameQACrashes );
+					WeeklyGameQAResults.TryGetValue( WeeklyResult.Key, out GameQACrashes );
 
-					int AnonymousCrashes = 0;
-					AnonymousResults.TryGetValue( Result.Key, out AnonymousCrashes );
+					string AnonymousLine = "";
+					foreach (var VersionCrashes in WeeklyEngineVersionResults)
+					{
+						int EngineVersionCrashes = 0;
+						VersionCrashes.Value.TryGetValue( WeeklyResult.Key, out EngineVersionCrashes );
+
+						AnonymousLine += EngineVersionCrashes;
+						AnonymousLine += ", ";
+					}
 
-					int Year = Result.Key.Year;
-					int Month = Result.Key.AddMonths( -1 ).Month;
-					if( Result.Key.Month == 13 || Result.Key.Month == 1 )
+					int Year = WeeklyResult.Key.Year;
+					int Month = WeeklyResult.Key.AddMonths( -1 ).Month;
+					if( WeeklyResult.Key.Month == 13 || WeeklyResult.Key.Month == 1 )
 					{
 						Month = 0;
 					}
 
-					int Day = Result.Key.Day + 6;
+					int Day = WeeklyResult.Key.Day + 6;
 
-					string Line = "[new Date(" + Year + ", " + Month + ", " + Day + "), " + GeneralCrashes + ", " + CoderCrashes + ", " + EngineQACrashes + ", " + GameQACrashes + ", " + AnonymousCrashes + ", " + Result.Value + "], ";
+					string Line = "[new Date(" + Year + ", " + Month + ", " + Day + "), " + GeneralCrashes + ", " + CoderCrashes + ", " + EngineQACrashes + ", " + GameQACrashes + ", " + AnonymousLine + WeeklyResult.Value + "], ";
 					CrashesByWeek += Line;
 				}
 
@@ -243,8 +359,15 @@ namespace Tools.CrashReporter.CrashReportWebSite.Controllers
 					int GameQACrashes = 0;
 					DailyGameQAResults.TryGetValue( DailyResult.Key, out GameQACrashes );
 
-					int AnonymousCrashes = 0;
-					DailyAnonymousResults.TryGetValue( DailyResult.Key, out AnonymousCrashes );
+					string AnonymousLine = "";
+					foreach (var VersionCrashes in DailyEngineVersionResults)
+					{
+						int EngineVersionCrashes = 0;
+						VersionCrashes.Value.TryGetValue( DailyResult.Key, out EngineVersionCrashes );
+
+						AnonymousLine += EngineVersionCrashes;
+						AnonymousLine +=", ";
+					}
 
 					int Year = DailyResult.Key.Year;
 					int Month = DailyResult.Key.AddMonths( -1 ).Month;
@@ -255,7 +378,7 @@ namespace Tools.CrashReporter.CrashReportWebSite.Controllers
 
 					int Day = DailyResult.Key.Day;
 
-					string Line = "[new Date(" + Year + ", " + Month + ", " + Day + "), " + GeneralCrashes + ", " + CoderCrashes + ", " + EngineQACrashes + ", " + GameQACrashes + ", " + AnonymousCrashes + ", " + DailyResult.Value + "], ";
+					string Line = "[new Date(" + Year + ", " + Month + ", " + Day + "), " + GeneralCrashes + ", " + CoderCrashes + ", " + EngineQACrashes + ", " + GameQACrashes + ", " + AnonymousLine + DailyResult.Value + "], ";
 					CrashesByDay += Line;
 				}
 
@@ -279,7 +402,13 @@ namespace Tools.CrashReporter.CrashReportWebSite.Controllers
 
 				BuggsByDay = BuggsByDay.TrimEnd( ", ".ToCharArray() );
 
-				var ResultDashboard = new DashboardViewModel { CrashesByWeek = CrashesByWeek, CrashesByDay = CrashesByDay, BuggsByDay = BuggsByDay };
+				var ResultDashboard = new DashboardViewModel
+				{
+					CrashesByWeek = CrashesByWeek,
+					CrashesByDay = CrashesByDay,
+					BuggsByDay = BuggsByDay,
+					EngineVersions = EngineUE4Versions,
+				};
 				ResultDashboard.GenerationTime = LogTimer.GetElapsedSeconds().ToString( "F2" );
 				return View( "Index", ResultDashboard );
 			}
diff --git a/Engine/Source/Programs/CrashReporter/CrashReportWebSite/Controllers/ReportsController.cs b/Engine/Source/Programs/CrashReporter/CrashReportWebSite/Controllers/ReportsController.cs
--- a/Engine/Source/Programs/CrashReporter/CrashReportWebSite/Controllers/ReportsController.cs
+++ b/Engine/Source/Programs/CrashReporter/CrashReportWebSite/Controllers/ReportsController.cs
@@ -60,7 +60,6 @@ namespace Tools.CrashReporter.CrashReportWebSite.Controllers
 				int AnonymousGroupID = FRepository.Get( BuggsRepo ).FindOrAddGroup( AnonumousGroup );
 				HashSet<int> AnonumousIDs = FRepository.Get( BuggsRepo ).GetUserIdsFromUserGroup( AnonumousGroup );
 				int AnonymousID = AnonumousIDs.First();
-				HashSet<string> UserNamesForUserGroup = FRepository.Get( BuggsRepo ).GetUserNamesFromGroupName( AnonumousGroup );
 
 				var Crashes = CrashRepo
 					.FilterByDate( CrashRepo.ListAll(), FormData.DateFrom, FormData.DateTo )
diff --git a/Engine/Source/Programs/CrashReporter/CrashReportWebSite/Models/BuggsViewModel.cs b/Engine/Source/Programs/CrashReporter/CrashReportWebSite/Models/BuggsViewModel.cs
--- a/Engine/Source/Programs/CrashReporter/CrashReportWebSite/Models/BuggsViewModel.cs
+++ b/Engine/Source/Programs/CrashReporter/CrashReportWebSite/Models/BuggsViewModel.cs
@@ -82,7 +82,7 @@ namespace Tools.CrashReporter.CrashReportWebSite.Models
 		/// </summary>
 		public BuggsViewModel()
 		{
-			VersionNames = CrashRepository.GetVersions();
+			VersionNames = CrashRepository.GetVersionsAsListItems();
 		}
 	}
 }
\ No newline at end of file
diff --git a/Engine/Source/Programs/CrashReporter/CrashReportWebSite/Models/CrashRepository.cs b/Engine/Source/Programs/CrashReporter/CrashReportWebSite/Models/CrashRepository.cs
--- a/Engine/Source/Programs/CrashReporter/CrashReportWebSite/Models/CrashRepository.cs
+++ b/Engine/Source/Programs/CrashReporter/CrashReportWebSite/Models/CrashRepository.cs
@@ -66,7 +66,7 @@ namespace Tools.CrashReporter.CrashReportWebSite.Models
 		/// <summary>
 		/// Retrieves a list of distinct UE4 Branches from the CrashRepository
 		/// </summary>
-		public static List<SelectListItem> GetBranches()
+		public static List<SelectListItem> GetBranchesAsListItems()
 		{
 			using (FAutoScopedLogTimer LogTimer = new FAutoScopedLogTimer( "CrashRepository.GetBranches" ))
 			{
@@ -97,10 +97,11 @@ namespace Tools.CrashReporter.CrashReportWebSite.Models
 
 		private static DateTime LastVersionDate = DateTime.UtcNow.AddDays( -1 );
 		private static List<SelectListItem> VersionsAsSelectList = null;
+		private static HashSet<string> DistinctBuildVersions = null;
 		/// <summary>
 		/// Retrieves a list of distinct UE4 Versions from the CrashRepository
 		/// </summary>
-		public static List<SelectListItem> GetVersions()
+		public static List<SelectListItem> GetVersionsAsListItems()
 		{
 			using (FAutoScopedLogTimer LogTimer = new FAutoScopedLogTimer( "CrashRepository.GetVersions" ))
 			{
@@ -115,7 +116,7 @@ namespace Tools.CrashReporter.CrashReportWebSite.Models
 					.Distinct()
 					.ToList();
 
-					HashSet<string> DistinctBuildVersions = new HashSet<string>();
+					DistinctBuildVersions = new HashSet<string>();
 
 					foreach (var BuildVersion in BuildVersions)
 					{
@@ -139,6 +140,15 @@ namespace Tools.CrashReporter.CrashReportWebSite.Models
 				return VersionsAsSelectList;
 			}
 		}
+
+		/// <summary>
+		/// Retrieves a list of distinct UE4 Versions from the CrashRepository
+		/// </summary>
+		public static HashSet<string> GetVersions()
+		{
+			GetVersionsAsListItems();
+			return DistinctBuildVersions;
+		}
 		
 		/// <summary>
 		/// Gets the crash from an id.
diff --git a/Engine/Source/Programs/CrashReporter/CrashReportWebSite/Models/CrashesViewModel.cs b/Engine/Source/Programs/CrashReporter/CrashReportWebSite/Models/CrashesViewModel.cs
--- a/Engine/Source/Programs/CrashReporter/CrashReportWebSite/Models/CrashesViewModel.cs
+++ b/Engine/Source/Programs/CrashReporter/CrashReportWebSite/Models/CrashesViewModel.cs
@@ -93,8 +93,8 @@ namespace Tools.CrashReporter.CrashReportWebSite.Models
 		{
 			DateTime FromDate = DateTime.Today.AddDays( -7 ).ToUniversalTime();
 			DateTime ToDate = DateTime.Today.ToUniversalTime();
-			BranchNames = CrashRepository.GetBranches();
-			VersionNames = CrashRepository.GetVersions();
+			BranchNames = CrashRepository.GetBranchesAsListItems();
+			VersionNames = CrashRepository.GetVersionsAsListItems();
 			DateFrom = (long)( FromDate - Epoch ).TotalMilliseconds;
 			DateTo = (long)( ToDate - Epoch ).TotalMilliseconds;
 			CrashType = "CrashesAsserts";
diff --git a/Engine/Source/Programs/CrashReporter/CrashReportWebSite/Models/DashboardViewModel.cs b/Engine/Source/Programs/CrashReporter/CrashReportWebSite/Models/DashboardViewModel.cs
--- a/Engine/Source/Programs/CrashReporter/CrashReportWebSite/Models/DashboardViewModel.cs
+++ b/Engine/Source/Programs/CrashReporter/CrashReportWebSite/Models/DashboardViewModel.cs
@@ -1,5 +1,10 @@
 ﻿// Copyright 1998-2015 Epic Games, Inc. All Rights Reserved.
 
+using System;
+using System.Collections.Generic;
+using System.Linq;
+using System.Web.Mvc;
+
 namespace Tools.CrashReporter.CrashReportWebSite.Models
 {
 	/// <summary>
@@ -16,5 +21,8 @@ namespace Tools.CrashReporter.CrashReportWebSite.Models
 
 		/// <summary>Time spent in generating this site, formatted as a string.</summary>
 		public string GenerationTime { get; set; }
+
+		/// <summary>Engine versions.</summary>
+		public List<string> EngineVersions;
 	}
 }
\ No newline at end of file
diff --git a/Engine/Source/Programs/CrashReporter/CrashReportWebSite/Views/Dashboard/Index.aspx b/Engine/Source/Programs/CrashReporter/CrashReportWebSite/Views/Dashboard/Index.aspx
--- a/Engine/Source/Programs/CrashReporter/CrashReportWebSite/Views/Dashboard/Index.aspx
+++ b/Engine/Source/Programs/CrashReporter/CrashReportWebSite/Views/Dashboard/Index.aspx
@@ -50,7 +50,16 @@ table
 				data.addColumn('number', 'Coder Crashes');
 				data.addColumn('number', 'EngineQA Crashes');
 				data.addColumn('number', 'GameQA Crashes');
-				data.addColumn('number', 'Anonymous Crashes');
+				
+				<%
+				foreach (var Version in Model.EngineVersions )
+				{
+				%>
+					data.addColumn('number', '<%=Version%> - ');
+				<% 
+				}
+				%>
+
 				data.addColumn('number', 'All Crashes');
 				data.addRows([<%=Model.CrashesByWeek%>]);
 				var chart = new google.visualization.AnnotatedTimeLine(document.getElementById('weekly_chart'));
@@ -69,7 +78,16 @@ table
 				data.addColumn('number', 'Coder Crashes');
 				data.addColumn('number', 'EngineQA Crashes');
 				data.addColumn('number', 'GameQA Crashes');
-				data.addColumn('number', 'Anonymous Crashes');
+				
+				<%
+				foreach (var Version in Model.EngineVersions )
+				{
+				%>
+					data.addColumn('number', '<%=Version%> - ');
+				<% 
+				}
+				%>
+
 				data.addColumn('number', 'All Crashes');
 				data.addRows([<%=Model.CrashesByDay%>]);
 				var chart = new google.visualization.AnnotatedTimeLine(document.getElementById('daily_chart'));
