CrashReportLinq - Less restrictive for removing invalid crashes

diff --git a/Engine/Source/Programs/CrashReporter/CrashReportWebSite/Queries/Fix-RepairDatabase.linq b/Engine/Source/Programs/CrashReporter/CrashReportWebSite/Queries/Fix-RepairDatabase.linq
--- a/Engine/Source/Programs/CrashReporter/CrashReportWebSite/Queries/Fix-RepairDatabase.linq
+++ b/Engine/Source/Programs/CrashReporter/CrashReportWebSite/Queries/Fix-RepairDatabase.linq
@@ -33,8 +33,8 @@ HashSet<int> GetInvalidFuncIds()
 }
 
 static Stopwatch Timer = Stopwatch.StartNew();
-static DateTime Date = DateTime.UtcNow.AddMonths(-1);   // From
-static DateTime UtcNow = DateTime.UtcNow.AddHours(-2); // To
+static DateTime Date = DateTime.UtcNow.AddDays(-30);   // From
+static DateTime UtcNow = DateTime.UtcNow.AddHours(-4); // To
 static TimeSpan Tick = TimeSpan.FromDays(1);
 const int NUM_OPS_PER_BATCH = 16;
 
@@ -321,7 +321,7 @@ void SetNoPatternForInvalidPattern()
 				}
 			}
 		
-			bool bInvalid = NumInvalidInPattern>0 && (PatternIds.Count - NumInvalidInPattern) < 4;
+			bool bInvalid = NumInvalidInPattern > (PatternIds.Count / 2 + 1);
 			if (bInvalid)
 			{
 				CrashIndex++;
