CrashReportLinq - Tweaks

diff --git a/Engine/Source/Programs/CrashReporter/CrashReportWebSite/Queries/BranchVersionTestQuery.linq b/Engine/Source/Programs/CrashReporter/CrashReportWebSite/Queries/BranchVersionTestQuery.linq
--- a/Engine/Source/Programs/CrashReporter/CrashReportWebSite/Queries/BranchVersionTestQuery.linq
+++ b/Engine/Source/Programs/CrashReporter/CrashReportWebSite/Queries/BranchVersionTestQuery.linq
@@ -10,22 +10,11 @@
 </Query>
 
 void Main()
-{
-	var BranchStartsWith = Crashes
-		.Where( n => n.TimeOfCrash > DateTime.Now.AddMonths( -3 ) )
-		.Where( n => n.Branch.StartsWith( "UE4" ) )
-		.Select( n => n.Branch )
-		.Distinct()
-		.ToList();
-					
-	var BranchContains = Crashes
-		.Where( n => n.TimeOfCrash > DateTime.Now.AddMonths( -3 ) )
-		.Where( n => n.Branch.Contains( "UE4" ) )
-		.Select( n => n.Branch )
-		.Distinct()
-		.ToList();
-					
-	var BuildVersions = Crashes.Where( c=> c.TimeOfCrash > DateTime.Now.AddMonths(-3))
+{			
+	var BuildVersions = Crashes
+	.Where( c=> c.TimeOfCrash > DateTime.Now.AddMonths(-3))
+	.Where( n => n.BuildVersion.StartsWith( "4." ) )
+	.Where( n => n.Branch.Contains( "UE4-Release" ) )
 	.Select( c => c.BuildVersion )
 	.Distinct()
 	.ToList();
diff --git a/Engine/Source/Programs/CrashReporter/CrashReportWebSite/Queries/FixBuggWithId.linq b/Engine/Source/Programs/CrashReporter/CrashReportWebSite/Queries/FixBuggWithId.linq
--- a/Engine/Source/Programs/CrashReporter/CrashReportWebSite/Queries/FixBuggWithId.linq
+++ b/Engine/Source/Programs/CrashReporter/CrashReportWebSite/Queries/FixBuggWithId.linq
@@ -10,7 +10,7 @@
 
 void Main()
 {
-	const int BuggId = 56230;
+	const int BuggId = 43365;
 	var BuggCrashes = Buggs_Crashes.Where (bc => bc.BuggId==BuggId);
 	Buggs.Where (b => b.Id==BuggId).First ().CrashType=3; // Ensure
 	Buggs.Context.SubmitChanges();
@@ -28,8 +28,15 @@ void Main()
 				Buggs_Crashes.Context.SubmitChanges();
 			}
 		}
+		else
+		{
+			if (Index % 16 == 0)
+			{
+				Debug.WriteLine( "{0,3}/{1,3}", Index, NumToChange );
+			}
+		}
 	}
 	Buggs_Crashes.Context.SubmitChanges();
 }
 
-// Define other methods and classes here
+// Define other methods and classes here
\ No newline at end of file
