CrashReportClient - Added CloseWithoutSending, PII is only send if check box "Allow to be.." is checked, username is only send when running non-launcher build, added options to view the folder with the crash data

diff --git a/Engine/Source/Programs/CrashReporter/CrashReportClient/Private/CrashDescription.cpp b/Engine/Source/Programs/CrashReporter/CrashReportClient/Private/CrashDescription.cpp
--- a/Engine/Source/Programs/CrashReporter/CrashReportClient/Private/CrashDescription.cpp
+++ b/Engine/Source/Programs/CrashReporter/CrashReportClient/Private/CrashDescription.cpp
@@ -7,8 +7,10 @@
 #include "XmlFile.h"
 #include "CrashDescription.h"
 #include "CrashReportAnalytics.h"
+#include "CrashReportUtil.h"
 #include "Runtime/Analytics/Analytics/Public/Analytics.h"
 #include "Runtime/Analytics/Analytics/Public/Interfaces/IAnalyticsProvider.h"
+#include "EngineBuildSettings.h"
 
 FCrashDescription::FCrashDescription() :
 	CrashDescriptionVersion( ECrashDescVersions::VER_1_NewCrashFormat ),
@@ -19,7 +21,7 @@ FCrashDescription::FCrashDescription() :
 	bHasVideoFile( false ),
 	bHasCompleteData( false )
 {
-	InitializeIDs();
+	UpdateIDs();
 }
 
 /** Unescapes a specified XML string, naive implementation. */
@@ -43,7 +45,7 @@ FCrashDescription::FCrashDescription( FString WERXMLFilepath ) :
 	bHasVideoFile( false ),
 	bHasCompleteData( false )
 {
-	InitializeIDs();
+	UpdateIDs();
 
 	// This is for the current system that uses files from Windows Error Reporting.
 	// Will be replaced with the unified version soon.
@@ -208,16 +210,31 @@ void FCrashDescription::InitializeEngineVersion()
 	EngineVersion = FEngineVersion( Major, Minor, Patch, BuiltFromCL, BranchName );
 }
 
-void FCrashDescription::InitializeIDs()
+void FCrashDescription::UpdateIDs()
 {
-	MachineId = FPlatformMisc::GetMachineId().ToString( EGuidFormats::Digits );
+	const bool bAddPersonalData = FCrashReportClientConfig::Get().GetAllowToBeContacted();
+	if (bAddPersonalData)
+	{
+		MachineId = FPlatformMisc::GetMachineId().ToString( EGuidFormats::Digits );
 
-	// The Epic ID can be looked up from this ID.
-	EpicAccountId = FPlatformMisc::GetEpicAccountId();
+		// The Epic ID can be looked up from this ID.
+		EpicAccountId = FPlatformMisc::GetEpicAccountId();
 
-	// Remove periods from user names to match AutoReporter user names
-	// The name prefix is read by CrashRepository.AddNewCrash in the website code
-	UserName = FString( FPlatformProcess::UserName() ).Replace( TEXT( "." ), TEXT( "" ) );
+		// Don't add real user name if running from the launcher.
+		const bool bAddUserName = FEngineBuildSettings::IsInternalBuild() || FEngineBuildSettings::IsPerforceBuild() || FEngineBuildSettings::IsSourceDistribution();
+		if (bAddUserName)
+		{
+			// Remove periods from user names to match AutoReporter user names
+			// The name prefix is read by CrashRepository.AddNewCrash in the website code
+			UserName = FString( FPlatformProcess::UserName() ).Replace( TEXT( "." ), TEXT( "" ) );
+		}
+	}
+	else
+	{
+		MachineId.Empty();
+		EpicAccountId.Empty();
+		UserName.Empty();
+	}
 }
 
 void FCrashDescription::SendAnalytics()
@@ -226,7 +243,6 @@ void FCrashDescription::SendAnalytics()
 	FCrashReportAnalytics::Initialize();
 
 	IAnalyticsProvider& Analytics = FCrashReportAnalytics::GetProvider();
-	//Analytics.SetUserID( MachineId );
 
 	TArray<FAnalyticsEventAttribute> CrashAttributes;
 
@@ -243,6 +259,7 @@ void FCrashDescription::SendAnalytics()
 	CrashAttributes.Add( FAnalyticsEventAttribute( TEXT( "BuiltFromCL" ), BuiltFromCL ) );
 	CrashAttributes.Add( FAnalyticsEventAttribute( TEXT( "BranchName" ), BranchName ) );
 
+	// @see UpdateIDs()
 	CrashAttributes.Add( FAnalyticsEventAttribute( TEXT( "MachineID" ), MachineId ) );
 	CrashAttributes.Add( FAnalyticsEventAttribute( TEXT( "UserName" ), UserName ) );
 	CrashAttributes.Add( FAnalyticsEventAttribute( TEXT( "EpicAccountId" ), EpicAccountId ) );
diff --git a/Engine/Source/Programs/CrashReporter/CrashReportClient/Private/CrashDescription.h b/Engine/Source/Programs/CrashReporter/CrashReportClient/Private/CrashDescription.h
--- a/Engine/Source/Programs/CrashReporter/CrashReportClient/Private/CrashDescription.h
+++ b/Engine/Source/Programs/CrashReporter/CrashReportClient/Private/CrashDescription.h
@@ -9,9 +9,9 @@
  *	Must contain the same properties as ...\CrashReportServer\CrashReportCommon\CrashDescription.cs.
  *	Contains all usable information about the crash. 
  *	
+ *	At this moment only used to sent the analytics.
  */
 
-// @TODO yrx 2014-08-27 add samples data
 struct FCrashDescription
 {
 	/** 
@@ -161,8 +161,6 @@ struct FCrashDescription
 	 */
 	TArray<FString> ErrorMessage;
 
-	// @TODO yrx 2014-09-10 Add misc data
-
 	/**
 	 *	Generic platform memory stats.
 	 *	@ADD-TO-THE-DATABASE
@@ -227,8 +225,8 @@ struct FCrashDescription
 	/** Initializes engine version from the separate components. */
 	void InitializeEngineVersion();
 
-	/** Initializes following properties: UserNAme, MachineID and EpicAccountID. */
-	void InitializeIDs();
+	/** Updates following properties: UserName, MachineID and EpicAccountID. */
+	void UpdateIDs();
 
 	/** Sends this crash for analytics. */
 	void SendAnalytics();
diff --git a/Engine/Source/Programs/CrashReporter/CrashReportClient/Private/CrashReportClient.cpp b/Engine/Source/Programs/CrashReporter/CrashReportClient/Private/CrashReportClient.cpp
--- a/Engine/Source/Programs/CrashReporter/CrashReportClient/Private/CrashReportClient.cpp
+++ b/Engine/Source/Programs/CrashReporter/CrashReportClient/Private/CrashReportClient.cpp
@@ -51,6 +51,12 @@ FCrashReportClient::~FCrashReportClient()
 	}
 }
 
+FReply FCrashReportClient::CloseWithoutSending()
+{
+	GIsRequestingExit = true;
+	return FReply::Handled();
+}
+
 FReply FCrashReportClient::Submit()
 {
 	bSendData = true;
@@ -117,6 +123,9 @@ EVisibility FCrashReportClient::IsThrobberVisible() const
 void FCrashReportClient::AllowToBeContacted_OnCheckStateChanged( ECheckBoxState NewRadioState )
 {
 	FCrashReportClientConfig::Get().SetAllowToBeContacted( NewRadioState == ECheckBoxState::Checked );
+
+	// Refresh PII based on the bAllowToBeContacted flag.
+	GetCrashDescription().UpdateIDs();
 }
 
 void FCrashReportClient::SendLogFile_OnCheckStateChanged( ECheckBoxState NewRadioState )
@@ -131,7 +140,7 @@ void FCrashReportClient::StartTicker()
 
 void FCrashReportClient::StoreCommentAndUpload()
 {
-	// Call upload even if the report is empty: pending reports will be sent if any
+	// Write user's comment
 	ErrorReport.SetUserComment( UserComment, FCrashReportClientConfig::Get().GetAllowToBeContacted() );
 	StartTicker();
 }
@@ -171,6 +180,11 @@ FString FCrashReportClient::GetCrashedAppName() const
 	return GetCrashDescription().GameName;
 }
 
+FString FCrashReportClient::GetCrashDirectory() const
+{
+	return ErrorReport.GetReportDirectory();
+}
+
 void FCrashReportClient::FinalizeDiagnoseReportWorker( FText ReportText )
 {
 	DiagnosticText = FCrashReportUtil::FormatDiagnosticText( ReportText, GetCrashDescription().MachineId, GetCrashDescription().EpicAccountId, GetCrashDescription().UserName );
diff --git a/Engine/Source/Programs/CrashReporter/CrashReportClient/Private/CrashReportClient.h b/Engine/Source/Programs/CrashReporter/CrashReportClient/Private/CrashReportClient.h
--- a/Engine/Source/Programs/CrashReporter/CrashReportClient/Private/CrashReportClient.h
+++ b/Engine/Source/Programs/CrashReporter/CrashReportClient/Private/CrashReportClient.h
@@ -58,6 +58,9 @@ public:
 	/** Destructor. */
 	virtual ~FCrashReportClient();
 
+	/** Closes the crash report client without sending any data. Except the startup analytics. */
+	FReply CloseWithoutSending();
+
 	/**
 	 * Respond to the user pressing Submit
 	 * @return Whether the request was handled
@@ -88,6 +91,12 @@ public:
 	 */
 	FString GetCrashedAppName() const;
 
+
+	/**
+	 * @return the full path of the crash directory.
+	 */
+	FString GetCrashDirectory() const;
+
 	/**
 	 * Handle the user updating the user comment text
 	 * @param Comment Text provided by the user
diff --git a/Engine/Source/Programs/CrashReporter/CrashReportClient/Private/CrashReportClientApp.cpp b/Engine/Source/Programs/CrashReporter/CrashReportClient/Private/CrashReportClientApp.cpp
--- a/Engine/Source/Programs/CrashReporter/CrashReportClient/Private/CrashReportClientApp.cpp
+++ b/Engine/Source/Programs/CrashReporter/CrashReportClient/Private/CrashReportClientApp.cpp
@@ -78,7 +78,7 @@ FPlatformErrorReport LoadErrorReport()
 	ErrorReport.FindFirstReportFileWithExtension( XMLWerFilename, TEXT( ".xml" ) );
 
 	extern FCrashDescription& GetCrashDescription();
-	GetCrashDescription() = FCrashDescription( ReportDirectoryAbsolutePath / XMLWerFilename );
+	GetCrashDescription() = FCrashDescription( ReportDirectoryAbsolutePath / XMLWerFilename );	
 
 #if CRASH_REPORT_UNATTENDED_ONLY
 	return ErrorReport;
@@ -158,6 +158,7 @@ void RunCrashReportClient(const TCHAR* CommandLine)
 
 	if (bUnattended)
 	{
+		// In the unattended mode we don't send any PII.
 		ErrorReport.SetUserComment( NSLOCTEXT( "CrashReportClient", "UnattendedMode", "Sent in the unattended mode" ), false );
 		FCrashReportClientUnattended CrashReportClient( ErrorReport );
 
diff --git a/Engine/Source/Programs/CrashReporter/CrashReportClient/Private/CrashReportClientStyle.cpp b/Engine/Source/Programs/CrashReporter/CrashReportClient/Private/CrashReportClientStyle.cpp
--- a/Engine/Source/Programs/CrashReporter/CrashReportClient/Private/CrashReportClientStyle.cpp
+++ b/Engine/Source/Programs/CrashReporter/CrashReportClient/Private/CrashReportClientStyle.cpp
@@ -27,6 +27,7 @@ void FCrashReportClientStyle::Shutdown()
 #define TTF_FONT(RelativePath, ...) FSlateFontInfo(ContentFromEngine(TEXT(RelativePath), TEXT(".ttf")), __VA_ARGS__)
 #define IMAGE_BRUSH( RelativePath, ... ) FSlateImageBrush(ContentFromEngine(TEXT(RelativePath), TEXT(".png")), __VA_ARGS__)
 #define BOX_BRUSH( RelativePath, ... ) FSlateBoxBrush(ContentFromEngine(TEXT(RelativePath), TEXT(".png")), __VA_ARGS__)
+#define BORDER_BRUSH( RelativePath, ... ) FSlateBorderBrush(ContentFromEngine(RelativePath, TEXT(".png")), __VA_ARGS__)
 
 namespace
 {
@@ -86,6 +87,27 @@ TSharedRef< FSlateStyleSet > FCrashReportClientStyle::Create()
 		Style.Set( "NormalEditableTextBox", NormalEditableTextBoxStyle );
 	}
 
+	// RichText
+	const FTextBlockStyle CrashReportDataStyle = FTextBlockStyle()
+		.SetFont( TTF_FONT( "Testing/Fonts/Roboto-Italic", 9 ) )
+		.SetColorAndOpacity( FSlateColor( FLinearColor::White * 0.5f ) )
+		.SetShadowOffset( FVector2D::ZeroVector )
+		.SetShadowColorAndOpacity( FLinearColor::Black );
+
+	Style.Set( "CrashReportDataStyle", CrashReportDataStyle );
+
+	FButtonStyle DarkHyperlinkButton = FButtonStyle()
+		.SetNormal( BORDER_BRUSH( "Old/HyperlinkDotted", FMargin( 0, 0, 0, 3 / 16.0f ), FSlateColor( FLinearColor::White * 0.5f ) ) )
+		.SetPressed( FSlateNoResource() )
+		.SetHovered( BORDER_BRUSH( "Old/HyperlinkUnderline", FMargin( 0, 0, 0, 3 / 16.0f ), FSlateColor( FLinearColor::White * 0.5f ) ) );
+
+	const FHyperlinkStyle DarkHyperlink = FHyperlinkStyle()
+		.SetUnderlineStyle( DarkHyperlinkButton )
+		.SetTextStyle( CrashReportDataStyle )
+		.SetPadding( FMargin( 0.0f ) );
+
+	Style.Set( "RichText.Hyperlink", DarkHyperlink );
+
 	return StyleRef;
 }
 
diff --git a/Engine/Source/Programs/CrashReporter/CrashReportClient/Private/CrashReportUtil.h b/Engine/Source/Programs/CrashReporter/CrashReportClient/Private/CrashReportUtil.h
--- a/Engine/Source/Programs/CrashReporter/CrashReportClient/Private/CrashReportUtil.h
+++ b/Engine/Source/Programs/CrashReporter/CrashReportClient/Private/CrashReportUtil.h
@@ -48,7 +48,11 @@ protected:
 	/** Filename to use when saving diagnostics report, if generated locally. */
 	FString DiagnosticsFilename;
 
-	/** Whether the user allowed us to be contacted. */
+	/**
+	 *	Whether the user allowed us to be contacted. 
+	 *	If true the following properties are retrieved from the system: UserName (for non-launcher build), MachineID and EpicAccountID.
+	 *	Otherwise they will be empty.
+	 */
 	bool bAllowToBeContacted;
 
 	/** Whether the user allowed us to send the log file. */
diff --git a/Engine/Source/Programs/CrashReporter/CrashReportClient/Private/GenericErrorReport.cpp b/Engine/Source/Programs/CrashReporter/CrashReportClient/Private/GenericErrorReport.cpp
--- a/Engine/Source/Programs/CrashReporter/CrashReportClient/Private/GenericErrorReport.cpp
+++ b/Engine/Source/Programs/CrashReporter/CrashReportClient/Private/GenericErrorReport.cpp
@@ -43,6 +43,47 @@ FGenericErrorReport::FGenericErrorReport(const FString& Directory)
 
 bool FGenericErrorReport::SetUserComment(const FText& UserComment, bool bAllowToBeContacted)
 {
+	// Load the file and remove all PII if bAllowToBeContacted is set to false.
+	const bool bRemovePersonalData = !bAllowToBeContacted;
+	if( bRemovePersonalData )
+	{
+		FString CrashContextFilename;
+		const bool bFound = FindFirstReportFileWithExtension( CrashContextFilename, TEXT( ".runtime-xml" ) );
+		if (bFound)
+		{
+			const FString CrashContextPath = GetReportDirectory() / CrashContextFilename;
+
+			// Cannot use FXMLFile because runtime-xml may contain multi line elements.
+			FString CrashContextString;
+			FFileHelper::LoadFileToString( CrashContextString, *CrashContextPath );
+
+			// Remove UserName, MachineId and EpicAccountId;
+			TArray<FString> Lines;
+			CrashContextString.ParseIntoArrayLines( Lines, true );
+
+			for( auto& It : Lines )
+			{
+				if (It.Contains( TEXT( "<UserName>" ) ))
+				{
+					It = TEXT( "<UserName></UserName>" );
+				}
+				else if (It.Contains( TEXT( "<MachineId>" ) ))
+				{
+					It = TEXT( "<MachineId></MachineId>" );
+				}
+				else if (It.Contains( TEXT( "<EpicAccountId>" ) ))
+				{
+					It = TEXT( "<EpicAccountId></EpicAccountId>" );
+				}
+			}
+
+			const FString FixedCrashContextString = FString::Join( Lines, TEXT( "\n" ) );
+			FFileHelper::SaveStringToFile( FixedCrashContextString, *CrashContextPath );
+		}
+
+		// #YRX_Crash: 2015-07-01 Move to FGenericCrashContext
+	}
+
 	// Find .xml file
 	FString XmlFilename;
 	if (!FindFirstReportFileWithExtension(XmlFilename, TEXT(".xml")))
@@ -50,7 +91,7 @@ bool FGenericErrorReport::SetUserComment(const FText& UserComment, bool bAllowTo
 		return false;
 	}
 
-	FString XmlFilePath = ReportDirectory / XmlFilename;
+	FString XmlFilePath = GetReportDirectory() / XmlFilename;
 	// FXmlFile's constructor loads the file to memory, closes the file and parses the data
 	FXmlFile XmlFile(XmlFilePath);
 	FXmlNode* DynamicSignaturesNode = XmlFile.IsValid() ?
diff --git a/Engine/Source/Programs/CrashReporter/CrashReportClient/Private/GenericErrorReport.h b/Engine/Source/Programs/CrashReporter/CrashReportClient/Private/GenericErrorReport.h
--- a/Engine/Source/Programs/CrashReporter/CrashReportClient/Private/GenericErrorReport.h
+++ b/Engine/Source/Programs/CrashReporter/CrashReportClient/Private/GenericErrorReport.h
@@ -38,7 +38,7 @@ public:
 	}
 
 	/**
-	 * Write the provided comment into the error report
+	 * Write the provided comment into the error report, also remove any PII if bAllowToBeContacted is set to false
 	 * @param UserComment Information provided by the user to add to the report
 	 * @return Whether the comment was successfully written to the report
 	 */
diff --git a/Engine/Source/Programs/CrashReporter/CrashReportClient/Private/SCrashReportClient.cpp b/Engine/Source/Programs/CrashReporter/CrashReportClient/Private/SCrashReportClient.cpp
--- a/Engine/Source/Programs/CrashReporter/CrashReportClient/Private/SCrashReportClient.cpp
+++ b/Engine/Source/Programs/CrashReporter/CrashReportClient/Private/SCrashReportClient.cpp
@@ -20,6 +20,15 @@ static void OnBrowserLinkClicked(const FSlateHyperlinkRun::FMetadata& Metadata,
 	}
 }
 
+static void OnViewCrashDirectory( const FSlateHyperlinkRun::FMetadata& Metadata, TSharedRef<SWidget> ParentWidget )
+{
+	const FString* UrlPtr = Metadata.Find( TEXT( "href" ) );
+	if (UrlPtr)
+	{
+		FPlatformProcess::ExploreFolder( **UrlPtr );
+	}
+}
+
 void SCrashReportClient::Construct(const FArguments& InArgs, TSharedRef<FCrashReportClient> Client)
 {
 	CrashReportClient = Client;
@@ -27,10 +36,15 @@ void SCrashReportClient::Construct(const FArguments& InArgs, TSharedRef<FCrashRe
 	auto CrashedAppName = CrashReportClient->GetCrashedAppName();
 
 	// Set the text displaying the name of the crashed app, if available
-	FText CrashedAppText = CrashedAppName.IsEmpty() ?
+	const FText CrashedAppText = CrashedAppName.IsEmpty() ?
 		LOCTEXT("CrashedAppNotFound", "An Unreal process has crashed") :
 		LOCTEXT("CrashedApp", "The following process has crashed: ");
 
+	const FText CrashReportDataText = FText::Format( LOCTEXT(
+		"CrashReportData",
+		"Crash reports comprise diagnostics files (<a id=\"browser\" href=\"{0}\" style=\"Richtext.Hyperlink\">click here to view directory</>) and the following summary information: " ),
+		FText::FromString( CrashReportClient->GetCrashDirectory()) );
+
 	ChildSlot
 	[
 		SNew(SBorder)
@@ -130,12 +144,13 @@ void SCrashReportClient::Construct(const FArguments& InArgs, TSharedRef<FCrashRe
 
 						+ SOverlay::Slot()
 						[
-							SNew(STextBlock)
-							.Margin(FMargin(4, 2, 0, 8))
-							.Font(FSlateFontInfo(FPaths::EngineContentDir() / TEXT("Slate/Testing/Fonts/Roboto-Italic.ttf"), 9))
-							.Text(LOCTEXT("CrashReportData", "Crash report data:"))
-							.ColorAndOpacity(FSlateColor(FLinearColor::White * 0.5f))
-							.AutoWrapText(false)
+							SNew( SRichTextBlock )
+							.Margin( FMargin( 4, 2, 0, 8 ) )
+							.TextStyle( &FCrashReportClientStyle::Get().GetWidgetStyle<FTextBlockStyle>( "CrashReportDataStyle" ) )
+							.Text( CrashReportDataText )
+							.AutoWrapText( true )
+							.DecoratorStyleSet( &FCrashReportClientStyle::Get() )
+							+ SRichTextBlock::HyperlinkDecorator( TEXT( "browser" ), FSlateHyperlinkRun::FOnClick::CreateStatic( &OnViewCrashDirectory, AsShared() ) )
 						]
 					]
 
@@ -237,6 +252,19 @@ void SCrashReportClient::Construct(const FArguments& InArgs, TSharedRef<FCrashRe
 					SNew(SSpacer)
 				]
 
+				
+				+ SHorizontalBox::Slot()
+				.HAlign( HAlign_Center )
+				.VAlign( VAlign_Center )
+				.AutoWidth()
+				.Padding( FMargin( 0 ) )
+				[
+					SNew( SButton )
+					.ContentPadding( FMargin( 8, 2 ) )
+					.Text( LOCTEXT( "CloseWithoutSending", "Close Without Sending" ) )
+					.OnClicked( Client, &FCrashReportClient::CloseWithoutSending )
+				]
+
 				+SHorizontalBox::Slot()
 				.HAlign(HAlign_Center)
 				.VAlign(VAlign_Center)
diff --git a/Engine/Source/Runtime/Core/Private/GenericPlatform/GenericPlatformContext.cpp b/Engine/Source/Runtime/Core/Private/GenericPlatform/GenericPlatformContext.cpp
--- a/Engine/Source/Runtime/Core/Private/GenericPlatform/GenericPlatformContext.cpp
+++ b/Engine/Source/Runtime/Core/Private/GenericPlatform/GenericPlatformContext.cpp
@@ -100,6 +100,8 @@ void FGenericCrashContext::SerializeContentToBuffer()
 	AddCrashProperty( TEXT( "IsPerforceBuild" ), (int32)NCachedCrashContextProperties::bIsPerforceBuild );
 	AddCrashProperty( TEXT( "IsSourceDistribution" ), (int32)NCachedCrashContextProperties::bIsSourceDistribution );
 
+	// Added 'editor/game open time till crash' 
+
 	// Add common crash properties.
 	AddCrashProperty( TEXT( "GameName" ), FApp::GetGameName() );
 	AddCrashProperty(TEXT("ExecutableName"), *NCachedCrashContextProperties::ExecutableName);
@@ -115,9 +117,10 @@ void FGenericCrashContext::SerializeContentToBuffer()
 
 	AddCrashProperty( TEXT( "IsUE4Release" ), (int32)NCachedCrashContextProperties::bIsUE4Release);
 
+	const bool bSetUserName = NCachedCrashContextProperties::bIsInternalBuild || NCachedCrashContextProperties::bIsPerforceBuild || NCachedCrashContextProperties::bIsSourceDistribution;
 	// Remove periods from user names to match AutoReporter user names
 	// The name prefix is read by CrashRepository.AddNewCrash in the website code
-	AddCrashProperty(TEXT("UserName"), *NCachedCrashContextProperties::UserName.Replace(TEXT("."), TEXT("")));
+	AddCrashProperty( TEXT( "UserName" ), bSetUserName ? *NCachedCrashContextProperties::UserName.Replace( TEXT( "." ), TEXT( "" ) ) : TEXT( "" ) );
 
 	AddCrashProperty(TEXT("BaseDir"), *NCachedCrashContextProperties::BaseDir);
 	AddCrashProperty(TEXT("RootDir"), *NCachedCrashContextProperties::RootDir);
diff --git a/Engine/Source/Runtime/Core/Public/Windows/WindowsPlatformCrashContext.h b/Engine/Source/Runtime/Core/Public/Windows/WindowsPlatformCrashContext.h
--- a/Engine/Source/Runtime/Core/Public/Windows/WindowsPlatformCrashContext.h
+++ b/Engine/Source/Runtime/Core/Public/Windows/WindowsPlatformCrashContext.h
@@ -18,7 +18,7 @@ struct CORE_API FWindowsPlatformCrashContext : public FGenericCrashContext
 
 	virtual void AddPlatformSpecificProperties() override
 	{
-		AddCrashProperty( TEXT( "Platform.IsRunningWindows" ), 1 );
+		AddCrashProperty( TEXT( "PlatformIsRunningWindows" ), 1 );
 	}
 };
 
