CrashReportClient - Better formatting the diagnostics text

diff --git a/Engine/Source/Programs/CrashReporter/CrashReportClient/Private/CrashReportClient.cpp b/Engine/Source/Programs/CrashReporter/CrashReportClient/Private/CrashReportClient.cpp
--- a/Engine/Source/Programs/CrashReporter/CrashReportClient/Private/CrashReportClient.cpp
+++ b/Engine/Source/Programs/CrashReporter/CrashReportClient/Private/CrashReportClient.cpp
@@ -37,7 +37,7 @@ FCrashReportClient::FCrashReportClient(const FPlatformErrorReport& InErrorReport
 	}
 	else if( !DiagnosticText.IsEmpty() )
 	{
-		DiagnosticText = FCrashReportUtil::FormatDiagnosticText( DiagnosticText, GetCrashDescription().MachineId, GetCrashDescription().EpicAccountId, GetCrashDescription().UserName );
+		FormattedDiagnosticText = FCrashReportUtil::FormatDiagnosticText( DiagnosticText, GetCrashDescription().MachineId, GetCrashDescription().EpicAccountId, GetCrashDescription().UserName );
 	}
 }
 
@@ -86,7 +86,7 @@ FReply FCrashReportClient::CopyCallstack()
 
 FText FCrashReportClient::GetDiagnosticText() const
 {
-	return DiagnosticText;
+	return FormattedDiagnosticText;
 }
 
 void FCrashReportClient::UserCommentChanged(const FText& Comment, ETextCommit::Type CommitType)
@@ -126,6 +126,9 @@ void FCrashReportClient::AllowToBeContacted_OnCheckStateChanged( ECheckBoxState
 
 	// Refresh PII based on the bAllowToBeContacted flag.
 	GetCrashDescription().UpdateIDs();
+
+	// Update diagnostics text.
+	FormattedDiagnosticText = FCrashReportUtil::FormatDiagnosticText( DiagnosticText, GetCrashDescription().MachineId, GetCrashDescription().EpicAccountId, GetCrashDescription().UserName );
 }
 
 void FCrashReportClient::SendLogFile_OnCheckStateChanged( ECheckBoxState NewRadioState )
@@ -187,7 +190,7 @@ FString FCrashReportClient::GetCrashDirectory() const
 
 void FCrashReportClient::FinalizeDiagnoseReportWorker( FText ReportText )
 {
-	DiagnosticText = FCrashReportUtil::FormatDiagnosticText( ReportText, GetCrashDescription().MachineId, GetCrashDescription().EpicAccountId, GetCrashDescription().UserName );
+	FormattedDiagnosticText = FCrashReportUtil::FormatDiagnosticText( ReportText, GetCrashDescription().MachineId, GetCrashDescription().EpicAccountId, GetCrashDescription().UserName );
 
 	auto DiagnosticsFilePath = ErrorReport.GetReportDirectory() / FCrashReportClientConfig::Get().GetDiagnosticsFilename();
 	Uploader.LocalDiagnosisComplete(FPaths::FileExists(DiagnosticsFilePath) ? DiagnosticsFilePath : TEXT(""));
diff --git a/Engine/Source/Programs/CrashReporter/CrashReportClient/Private/CrashReportClient.h b/Engine/Source/Programs/CrashReporter/CrashReportClient/Private/CrashReportClient.h
--- a/Engine/Source/Programs/CrashReporter/CrashReportClient/Private/CrashReportClient.h
+++ b/Engine/Source/Programs/CrashReporter/CrashReportClient/Private/CrashReportClient.h
@@ -157,6 +157,9 @@ private:
 	/** Exception and call-stack to show, valid once diagnosis task is complete */
 	FText DiagnosticText;
 
+	/** Formatted diagnostics crash reporter data. */
+	FText FormattedDiagnosticText;
+
 	/** Background worker to get a callstack from the report */
 	FAsyncTask<FDiagnoseReportWorker>* DiagnoseReportTask;
 
