CrashReportClient - Analitycs are now configurable through the config, enabled by default. MachineID is always send, required to identify the number of unique users that are experiencing a crash.

diff --git a/Engine/Source/Programs/CrashReporter/CrashReportClient/Private/CrashDescription.cpp b/Engine/Source/Programs/CrashReporter/CrashReportClient/Private/CrashDescription.cpp
--- a/Engine/Source/Programs/CrashReporter/CrashReportClient/Private/CrashDescription.cpp
+++ b/Engine/Source/Programs/CrashReporter/CrashReportClient/Private/CrashDescription.cpp
@@ -215,8 +215,6 @@ void FCrashDescription::UpdateIDs()
 	const bool bAddPersonalData = FCrashReportClientConfig::Get().GetAllowToBeContacted();
 	if (bAddPersonalData)
 	{
-		MachineId = FPlatformMisc::GetMachineId().ToString( EGuidFormats::Digits );
-
 		// The Epic ID can be looked up from this ID.
 		EpicAccountId = FPlatformMisc::GetEpicAccountId();
 
@@ -231,46 +229,51 @@ void FCrashDescription::UpdateIDs()
 	}
 	else
 	{
-		MachineId.Empty();
 		EpicAccountId.Empty();
 		UserName.Empty();
 	}
+
+	MachineId = FPlatformMisc::GetMachineId().ToString( EGuidFormats::Digits );
 }
 
 void FCrashDescription::SendAnalytics()
 {
-	// Connect the crash report client analytics provider.
-	FCrashReportAnalytics::Initialize();
+	bool bSendAnalyticsForCrash = FCrashReportClientConfig::Get().GetSendAnalyticsForCrash();
+	if (bSendAnalyticsForCrash)
+	{
+		// Connect the crash report client analytics provider.
+		FCrashReportAnalytics::Initialize();
 
-	IAnalyticsProvider& Analytics = FCrashReportAnalytics::GetProvider();
+		IAnalyticsProvider& Analytics = FCrashReportAnalytics::GetProvider();
 
-	TArray<FAnalyticsEventAttribute> CrashAttributes;
+		TArray<FAnalyticsEventAttribute> CrashAttributes;
 
-	CrashAttributes.Add( FAnalyticsEventAttribute( TEXT( "bHasCompleteData" ), bHasCompleteData ) );
-	CrashAttributes.Add( FAnalyticsEventAttribute( TEXT( "CrashDescriptionVersion" ), (int32)CrashDescriptionVersion ) );
-	CrashAttributes.Add( FAnalyticsEventAttribute( TEXT( "ReportName" ), ReportName ) );
+		CrashAttributes.Add( FAnalyticsEventAttribute( TEXT( "bHasCompleteData" ), bHasCompleteData ) );
+		CrashAttributes.Add( FAnalyticsEventAttribute( TEXT( "CrashDescriptionVersion" ), (int32)CrashDescriptionVersion ) );
+		CrashAttributes.Add( FAnalyticsEventAttribute( TEXT( "ReportName" ), ReportName ) );
 
-	//	AppID = GameName
-	CrashAttributes.Add( FAnalyticsEventAttribute( TEXT( "GameName" ), GameName ) );
+		//	AppID = GameName
+		CrashAttributes.Add( FAnalyticsEventAttribute( TEXT( "GameName" ), GameName ) );
 
-	//	AppVersion = EngineVersion
-	CrashAttributes.Add( FAnalyticsEventAttribute( TEXT( "EngineVersion" ), EngineVersion.ToString() ) );
-	CrashAttributes.Add( FAnalyticsEventAttribute( TEXT( "BuildVersion" ), BuildVersion ) );
-	CrashAttributes.Add( FAnalyticsEventAttribute( TEXT( "BuiltFromCL" ), BuiltFromCL ) );
-	CrashAttributes.Add( FAnalyticsEventAttribute( TEXT( "BranchName" ), BranchName ) );
+		//	AppVersion = EngineVersion
+		CrashAttributes.Add( FAnalyticsEventAttribute( TEXT( "EngineVersion" ), EngineVersion.ToString() ) );
+		CrashAttributes.Add( FAnalyticsEventAttribute( TEXT( "BuildVersion" ), BuildVersion ) );
+		CrashAttributes.Add( FAnalyticsEventAttribute( TEXT( "BuiltFromCL" ), BuiltFromCL ) );
+		CrashAttributes.Add( FAnalyticsEventAttribute( TEXT( "BranchName" ), BranchName ) );
 
-	// @see UpdateIDs()
-	CrashAttributes.Add( FAnalyticsEventAttribute( TEXT( "MachineID" ), MachineId ) );
-	CrashAttributes.Add( FAnalyticsEventAttribute( TEXT( "UserName" ), UserName ) );
-	CrashAttributes.Add( FAnalyticsEventAttribute( TEXT( "EpicAccountId" ), EpicAccountId ) );
+		// @see UpdateIDs()
+		CrashAttributes.Add( FAnalyticsEventAttribute( TEXT( "MachineID" ), MachineId ) );
+		CrashAttributes.Add( FAnalyticsEventAttribute( TEXT( "UserName" ), UserName ) );
+		CrashAttributes.Add( FAnalyticsEventAttribute( TEXT( "EpicAccountId" ), EpicAccountId ) );
 
-	CrashAttributes.Add( FAnalyticsEventAttribute( TEXT( "Platform" ), Platform ) );
-	CrashAttributes.Add( FAnalyticsEventAttribute( TEXT( "TimeOfCrash" ), TimeOfCrash.GetTicks() ) );
-	CrashAttributes.Add( FAnalyticsEventAttribute( TEXT( "EngineMode" ), EngineMode ) );
-	CrashAttributes.Add( FAnalyticsEventAttribute( TEXT( "LanguageLCID" ), LanguageLCID ) );
+		CrashAttributes.Add( FAnalyticsEventAttribute( TEXT( "Platform" ), Platform ) );
+		CrashAttributes.Add( FAnalyticsEventAttribute( TEXT( "TimeOfCrash" ), TimeOfCrash.GetTicks() ) );
+		CrashAttributes.Add( FAnalyticsEventAttribute( TEXT( "EngineMode" ), EngineMode ) );
+		CrashAttributes.Add( FAnalyticsEventAttribute( TEXT( "LanguageLCID" ), LanguageLCID ) );
 
-	Analytics.RecordEvent( TEXT( "CrashReportClient.ReportCrash" ), CrashAttributes );
+		Analytics.RecordEvent( TEXT( "CrashReportClient.ReportCrash" ), CrashAttributes );
 
-	// Shutdown analytics.
-	FCrashReportAnalytics::Shutdown();
+		// Shutdown analytics.
+		FCrashReportAnalytics::Shutdown();
+	}
 }
diff --git a/Engine/Source/Programs/CrashReporter/CrashReportClient/Private/CrashReportClientApp.cpp b/Engine/Source/Programs/CrashReporter/CrashReportClient/Private/CrashReportClientApp.cpp
--- a/Engine/Source/Programs/CrashReporter/CrashReportClient/Private/CrashReportClientApp.cpp
+++ b/Engine/Source/Programs/CrashReporter/CrashReportClient/Private/CrashReportClientApp.cpp
@@ -103,6 +103,13 @@ FCrashReportClientConfig::FCrashReportClientConfig()
 
 	GConfig->GetBool( TEXT( "CrashReportClient" ), TEXT( "bAllowToBeContacted" ), bAllowToBeContacted, GEngineIni );
 	GConfig->GetBool( TEXT( "CrashReportClient" ), TEXT( "bSendLogFile" ), bSendLogFile, GEngineIni );
+	if (!GConfig->GetBool( TEXT( "CrashReportClient" ), TEXT( "bSendAnalyticsForCrash" ), bSendAnalyticsForCrash, GEngineIni ))
+	{
+		// Be default send.
+		bSendAnalyticsForCrash = true;
+		// Add to the config.
+		FCrashReportClientConfig::Get().SetSendAnalyticsForCrash( bSendAnalyticsForCrash );
+	}
 
 	UE_LOG( CrashReportClientLog, Log, TEXT( "CrashReportReceiverIP: %s" ), *CrashReportReceiverIP );
 }
@@ -119,6 +126,12 @@ void FCrashReportClientConfig::SetSendLogFile( bool bNewValue )
 	GConfig->SetBool( TEXT( "CrashReportClient" ), TEXT( "bSendLogFile" ), bSendLogFile, GEngineIni );
 }
 
+void FCrashReportClientConfig::SetSendAnalyticsForCrash( bool bNewValue )
+{
+	bSendAnalyticsForCrash = bNewValue;
+	GConfig->SetBool( TEXT( "CrashReportClient" ), TEXT( "bSendAnalyticsForCrash" ), bSendAnalyticsForCrash, GEngineIni );
+}
+
 void RunCrashReportClient(const TCHAR* CommandLine)
 {
 	// Override the stack size for the thread pool.
diff --git a/Engine/Source/Programs/CrashReporter/CrashReportClient/Private/CrashReportUtil.h b/Engine/Source/Programs/CrashReporter/CrashReportClient/Private/CrashReportUtil.h
--- a/Engine/Source/Programs/CrashReporter/CrashReportClient/Private/CrashReportUtil.h
+++ b/Engine/Source/Programs/CrashReporter/CrashReportClient/Private/CrashReportUtil.h
@@ -37,9 +37,14 @@ struct FCrashReportClientConfig
 		return bSendLogFile;
 	}
 
-	void SetAllowToBeContacted( bool bNewValue );
+	const bool& GetSendAnalyticsForCrash() const
+	{
+		return bSendAnalyticsForCrash;
+	}
 
+	void SetAllowToBeContacted( bool bNewValue );
 	void SetSendLogFile( bool bNewValue );
+	void SetSendAnalyticsForCrash( bool bNewValue );
 
 protected:
 	/** IP address of crash report receiver. */
@@ -50,13 +55,23 @@ protected:
 
 	/**
 	 *	Whether the user allowed us to be contacted. 
-	 *	If true the following properties are retrieved from the system: UserName (for non-launcher build), MachineID and EpicAccountID.
+	 *	If true the following properties are retrieved from the system: UserName (for non-launcher build) and EpicAccountID.
 	 *	Otherwise they will be empty.
 	 */
 	bool bAllowToBeContacted;
 
 	/** Whether the user allowed us to send the log file. */
 	bool bSendLogFile;
+
+	/**
+	 *	Whether the user allowed us to send the analytics for the crash.
+	 *	Enabled by default.
+	 *	Can be changed in the config file located here, for Windows:
+	 *	[REDACTED_USER_HOME]\AppData\Local\CrashReportClient\Saved\Config\Windows\Engine.ini
+	 *	
+	 *	This will be improved in the next release.
+	 */
+	bool bSendAnalyticsForCrash;
 };
 
 
diff --git a/Engine/Source/Programs/CrashReporter/CrashReportClient/Private/GenericErrorReport.cpp b/Engine/Source/Programs/CrashReporter/CrashReportClient/Private/GenericErrorReport.cpp
--- a/Engine/Source/Programs/CrashReporter/CrashReportClient/Private/GenericErrorReport.cpp
+++ b/Engine/Source/Programs/CrashReporter/CrashReportClient/Private/GenericErrorReport.cpp
@@ -57,7 +57,7 @@ bool FGenericErrorReport::SetUserComment(const FText& UserComment, bool bAllowTo
 			FString CrashContextString;
 			FFileHelper::LoadFileToString( CrashContextString, *CrashContextPath );
 
-			// Remove UserName, MachineId and EpicAccountId;
+			// Remove UserName and EpicAccountId;
 			TArray<FString> Lines;
 			CrashContextString.ParseIntoArrayLines( Lines, true );
 
@@ -67,10 +67,6 @@ bool FGenericErrorReport::SetUserComment(const FText& UserComment, bool bAllowTo
 				{
 					It = TEXT( "<UserName></UserName>" );
 				}
-				else if (It.Contains( TEXT( "<MachineId>" ) ))
-				{
-					It = TEXT( "<MachineId></MachineId>" );
-				}
 				else if (It.Contains( TEXT( "<EpicAccountId>" ) ))
 				{
 					It = TEXT( "<EpicAccountId></EpicAccountId>" );
