CrashReportClient - Various privacy concerns (Analytics via send button)

diff --git a/Engine/Source/Programs/CrashReporter/CrashReportClient/Private/CrashDescription.cpp b/Engine/Source/Programs/CrashReporter/CrashReportClient/Private/CrashDescription.cpp
--- a/Engine/Source/Programs/CrashReporter/CrashReportClient/Private/CrashDescription.cpp
+++ b/Engine/Source/Programs/CrashReporter/CrashReportClient/Private/CrashDescription.cpp
@@ -238,9 +238,6 @@ void FCrashDescription::UpdateIDs()
 
 void FCrashDescription::SendAnalytics()
 {
-	bool bSendAnalyticsForCrash = FCrashReportClientConfig::Get().GetSendAnalyticsForCrash();
-	if (bSendAnalyticsForCrash)
-	{
 		// Connect the crash report client analytics provider.
 		FCrashReportAnalytics::Initialize();
 
@@ -276,4 +273,3 @@ void FCrashDescription::SendAnalytics()
 		// Shutdown analytics.
 		FCrashReportAnalytics::Shutdown();
 	}
-}
diff --git a/Engine/Source/Programs/CrashReporter/CrashReportClient/Private/CrashReportClient.cpp b/Engine/Source/Programs/CrashReporter/CrashReportClient/Private/CrashReportClient.cpp
--- a/Engine/Source/Programs/CrashReporter/CrashReportClient/Private/CrashReportClient.cpp
+++ b/Engine/Source/Programs/CrashReporter/CrashReportClient/Private/CrashReportClient.cpp
@@ -59,6 +59,12 @@ FReply FCrashReportClient::CloseWithoutSending()
 
 FReply FCrashReportClient::Submit()
 {
+	if (ErrorReport.HasFilesToUpload())
+	{
+		// Send analytics.
+		GetCrashDescription().SendAnalytics();
+	}
+
 	bSendData = true;
 	StoreCommentAndUpload();
 	bShouldWindowBeHidden = true;
diff --git a/Engine/Source/Programs/CrashReporter/CrashReportClient/Private/CrashReportClientApp.cpp b/Engine/Source/Programs/CrashReporter/CrashReportClient/Private/CrashReportClientApp.cpp
--- a/Engine/Source/Programs/CrashReporter/CrashReportClient/Private/CrashReportClientApp.cpp
+++ b/Engine/Source/Programs/CrashReporter/CrashReportClient/Private/CrashReportClientApp.cpp
@@ -103,13 +103,6 @@ FCrashReportClientConfig::FCrashReportClientConfig()
 
 	GConfig->GetBool( TEXT( "CrashReportClient" ), TEXT( "bAllowToBeContacted" ), bAllowToBeContacted, GEngineIni );
 	GConfig->GetBool( TEXT( "CrashReportClient" ), TEXT( "bSendLogFile" ), bSendLogFile, GEngineIni );
-	if (!GConfig->GetBool( TEXT( "CrashReportClient" ), TEXT( "bSendAnalyticsForCrash" ), bSendAnalyticsForCrash, GEngineIni ))
-	{
-		// Be default send.
-		bSendAnalyticsForCrash = true;
-		// Add to the config.
-		FCrashReportClientConfig::Get().SetSendAnalyticsForCrash( bSendAnalyticsForCrash );
-	}
 
 	UE_LOG( CrashReportClientLog, Log, TEXT( "CrashReportReceiverIP: %s" ), *CrashReportReceiverIP );
 }
@@ -126,12 +119,6 @@ void FCrashReportClientConfig::SetSendLogFile( bool bNewValue )
 	GConfig->SetBool( TEXT( "CrashReportClient" ), TEXT( "bSendLogFile" ), bSendLogFile, GEngineIni );
 }
 
-void FCrashReportClientConfig::SetSendAnalyticsForCrash( bool bNewValue )
-{
-	bSendAnalyticsForCrash = bNewValue;
-	GConfig->SetBool( TEXT( "CrashReportClient" ), TEXT( "bSendAnalyticsForCrash" ), bSendAnalyticsForCrash, GEngineIni );
-}
-
 void RunCrashReportClient(const TCHAR* CommandLine)
 {
 	// Override the stack size for the thread pool.
@@ -162,13 +149,6 @@ void RunCrashReportClient(const TCHAR* CommandLine)
 	FPlatformErrorReport::Init();
 	auto ErrorReport = LoadErrorReport();
 	
-	if( ErrorReport.HasFilesToUpload() )
-	{
-		// Send analytics.
-		extern FCrashDescription& GetCrashDescription();
-		GetCrashDescription().SendAnalytics();
-	}
-
 	if (bUnattended)
 	{
 		// In the unattended mode we don't send any PII.
diff --git a/Engine/Source/Programs/CrashReporter/CrashReportClient/Private/CrashReportClientUnattended.cpp b/Engine/Source/Programs/CrashReporter/CrashReportClient/Private/CrashReportClientUnattended.cpp
--- a/Engine/Source/Programs/CrashReporter/CrashReportClient/Private/CrashReportClientUnattended.cpp
+++ b/Engine/Source/Programs/CrashReporter/CrashReportClient/Private/CrashReportClientUnattended.cpp
@@ -2,7 +2,6 @@
 
 #include "CrashReportClientApp.h"
 #include "CrashReportClientUnattended.h"
-#include "CrashReportUtil.h"
 
 FCrashReportClientUnattended::FCrashReportClientUnattended(const FPlatformErrorReport& ErrorReport)
 	: Uploader( FCrashReportClientConfig::Get().GetReceiverAddress() )
diff --git a/Engine/Source/Programs/CrashReporter/CrashReportClient/Private/CrashReportUtil.h b/Engine/Source/Programs/CrashReporter/CrashReportClient/Private/CrashReportUtil.h
--- a/Engine/Source/Programs/CrashReporter/CrashReportClient/Private/CrashReportUtil.h
+++ b/Engine/Source/Programs/CrashReporter/CrashReportClient/Private/CrashReportUtil.h
@@ -37,14 +37,8 @@ struct FCrashReportClientConfig
 		return bSendLogFile;
 	}
 
-	const bool& GetSendAnalyticsForCrash() const
-	{
-		return bSendAnalyticsForCrash;
-	}
-
 	void SetAllowToBeContacted( bool bNewValue );
 	void SetSendLogFile( bool bNewValue );
-	void SetSendAnalyticsForCrash( bool bNewValue );
 
 protected:
 	/** IP address of crash report receiver. */
@@ -62,16 +56,6 @@ protected:
 
 	/** Whether the user allowed us to send the log file. */
 	bool bSendLogFile;
-
-	/**
-	 *	Whether the user allowed us to send the analytics for the crash.
-	 *	Enabled by default.
-	 *	Can be changed in the config file located here, for Windows:
-	 *	[REDACTED_USER_HOME]\AppData\Local\CrashReportClient\Saved\Config\Windows\Engine.ini
-	 *	
-	 *	This will be improved in the next release.
-	 */
-	bool bSendAnalyticsForCrash;
 };
 
 
