CrashReportClient - Better formatting the diagnostics text (UI fix) #jira UE-18020

diff --git a/Engine/Source/Programs/CrashReporter/CrashReportClient/Private/CrashReportClient.cpp b/Engine/Source/Programs/CrashReporter/CrashReportClient/Private/CrashReportClient.cpp
--- a/Engine/Source/Programs/CrashReporter/CrashReportClient/Private/CrashReportClient.cpp
+++ b/Engine/Source/Programs/CrashReporter/CrashReportClient/Private/CrashReportClient.cpp
@@ -196,6 +196,7 @@ FString FCrashReportClient::GetCrashDirectory() const
 
 void FCrashReportClient::FinalizeDiagnoseReportWorker( FText ReportText )
 {
+	DiagnosticText = ReportText;
 	FormattedDiagnosticText = FCrashReportUtil::FormatDiagnosticText( ReportText, GetCrashDescription().MachineId, GetCrashDescription().EpicAccountId, GetCrashDescription().UserName );
 
 	auto DiagnosticsFilePath = ErrorReport.GetReportDirectory() / FCrashReportClientConfig::Get().GetDiagnosticsFilename();
