Fixed a typo

diff --git a/Engine/Source/Programs/CrashReporter/CrashReportClient/Private/GenericErrorReport.cpp b/Engine/Source/Programs/CrashReporter/CrashReportClient/Private/GenericErrorReport.cpp
--- a/Engine/Source/Programs/CrashReporter/CrashReportClient/Private/GenericErrorReport.cpp
+++ b/Engine/Source/Programs/CrashReporter/CrashReportClient/Private/GenericErrorReport.cpp
@@ -115,17 +115,17 @@ bool FGenericErrorReport::SetUserComment(const FText& UserComment, bool bAllowTo
 	extern FCrashDescription& GetCrashDescription();
 
 	// @see FCrashDescription::UpdateIDs
-	const FString EpicMachindAndUserNameIDs = FString::Printf( TEXT( "!MachineId:%s!EpicAccountId:%s!Name:%s" ), *GetCrashDescription().MachineId, *GetCrashDescription().EpicAccountId, *GetCrashDescription().UserName );
+	const FString EpicMachineAndUserNameIDs = FString::Printf( TEXT( "!MachineId:%s!EpicAccountId:%s!Name:%s" ), *GetCrashDescription().MachineId, *GetCrashDescription().EpicAccountId, *GetCrashDescription().UserName );
 
 	// Add or update a user ID.
 	FXmlNode* Parameter4Node = DynamicSignaturesNode->FindChildNode(TEXT("Parameter4"));
 	if( Parameter4Node )
 	{
-		Parameter4Node->SetContent(EpicMachindAndUserNameIDs);
+		Parameter4Node->SetContent(EpicMachineAndUserNameIDs);
 	}
 	else
 	{
-		DynamicSignaturesNode->AppendChildNode(TEXT("Parameter4"), EpicMachindAndUserNameIDs);
+		DynamicSignaturesNode->AppendChildNode(TEXT("Parameter4"), EpicMachineAndUserNameIDs);
 	}
 
 	// Add or update bAllowToBeContacted
