CrashReportClient - Replace user name to anonymous for metadata files

diff --git a/Engine/Source/Programs/CrashReporter/CrashReportClient/Private/GenericErrorReport.cpp b/Engine/Source/Programs/CrashReporter/CrashReportClient/Private/GenericErrorReport.cpp
--- a/Engine/Source/Programs/CrashReporter/CrashReportClient/Private/GenericErrorReport.cpp
+++ b/Engine/Source/Programs/CrashReporter/CrashReportClient/Private/GenericErrorReport.cpp
@@ -43,6 +43,10 @@ FGenericErrorReport::FGenericErrorReport(const FString& Directory)
 
 bool FGenericErrorReport::SetUserComment(const FText& UserComment, bool bAllowToBeContacted)
 {
+	const FString UserName1 = FPlatformProcess::UserName( false );
+	const FString UserName2 = FPlatformProcess::UserName( true );
+	const TCHAR* Anonymous = TEXT( "Anonymous" );
+
 	// Load the file and remove all PII if bAllowToBeContacted is set to false.
 	const bool bRemovePersonalData = !bAllowToBeContacted;
 	if( bRemovePersonalData )
@@ -63,6 +67,10 @@ bool FGenericErrorReport::SetUserComment(const FText& UserComment, bool bAllowTo
 
 			for( auto& It : Lines )
 			{
+				// Replace user name in assert message, command line etc.
+				It = It.Replace( *UserName1, Anonymous );
+				It = It.Replace( *UserName2, Anonymous );
+
 				if (It.Contains( TEXT( "<UserName>" ) ))
 				{
 					It = TEXT( "<UserName></UserName>" );
@@ -99,6 +107,32 @@ bool FGenericErrorReport::SetUserComment(const FText& UserComment, bool bAllowTo
 		return false;
 	}
 
+	if (bRemovePersonalData)
+	{
+		FXmlNode* ProblemNode = XmlFile.GetRootNode()->FindChildNode( TEXT( "ProblemSignatures" ) );
+		if (ProblemNode)
+		{
+			FXmlNode* Parameter8Node = ProblemNode->FindChildNode( TEXT( "Parameter8" ) );
+			if (Parameter8Node)
+			{
+				// Replace user name in assert message, command line etc.
+				FString Content = Parameter8Node->GetContent();
+				Content = Content.Replace( *UserName1, Anonymous );
+				Content = Content.Replace( *UserName2, Anonymous );
+				Parameter8Node->SetContent( Content );
+			}
+			FXmlNode* Parameter9Node = ProblemNode->FindChildNode( TEXT( "Parameter9" ) );
+			if (Parameter9Node)
+			{
+				// Replace user name in assert message, command line etc.
+				FString Content = Parameter9Node->GetContent();
+				Content = Content.Replace( *UserName1, Anonymous );
+				Content = Content.Replace( *UserName2, Anonymous );
+				Parameter9Node->SetContent( Content );
+			}
+		}
+	}
+
 	// Add or update the user comment.
 	FXmlNode* Parameter3Node = DynamicSignaturesNode->FindChildNode(TEXT("Parameter3"));
 	if( Parameter3Node )
