CrashReportClient - Remove commandline if not allowed to contact

diff --git a/Engine/Source/Programs/CrashReporter/CrashReportClient/Private/GenericErrorReport.cpp b/Engine/Source/Programs/CrashReporter/CrashReportClient/Private/GenericErrorReport.cpp
--- a/Engine/Source/Programs/CrashReporter/CrashReportClient/Private/GenericErrorReport.cpp
+++ b/Engine/Source/Programs/CrashReporter/CrashReportClient/Private/GenericErrorReport.cpp
@@ -119,6 +119,17 @@ bool FGenericErrorReport::SetUserComment(const FText& UserComment, bool bAllowTo
 				FString Content = Parameter8Node->GetContent();
 				Content = Content.Replace( *UserName1, Anonymous );
 				Content = Content.Replace( *UserName2, Anonymous );
+
+				// Remove the command line. Command line is between first and second !
+				TArray<FString> ParsedParameters8;
+				Content.ParseIntoArray( ParsedParameters8, TEXT( "!" ), false );
+				if (ParsedParameters8.Num() > 1)
+				{
+					ParsedParameters8[1] = TEXT( "CommandLineRemoved" );
+				}
+
+				Content = FString::Join( ParsedParameters8, TEXT( "!" ) );
+
 				Parameter8Node->SetContent( Content );
 			}
 			FXmlNode* Parameter9Node = ProblemNode->FindChildNode( TEXT( "Parameter9" ) );
