CrashReportLinq - A few more scripts

diff --git a/Engine/Source/Programs/CrashReporter/CrashReportWebSite/Queries/FixBadJiraInBugg.linq b/Engine/Source/Programs/CrashReporter/CrashReportWebSite/Queries/FixBadJiraInBugg.linq
new file mode 100644
--- /dev/null
+++ b/Engine/Source/Programs/CrashReporter/CrashReportWebSite/Queries/FixBadJiraInBugg.linq
@@ -0,0 +1,18 @@
+<Query Kind="Program">
+  <Connection>
+    <ID>e40cf789-b18a-4b65-bae6-a0f5440c1ce2</ID>
+    <Persist>true</Persist>
+    <Server>db-09</Server>
+    <Database>CrashReport</Database>
+    <IsProduction>true</IsProduction>
+  </Connection>
+</Query>
+
+void Main()
+{
+	var Bugg = Buggs.Where (b => b.Id==64141).First ();
+	Bugg.TTPID="";
+	Buggs.Context.SubmitChanges();
+}
+
+// Define other methods and classes here
