Upgrade RegisterPII to .NET 4.5 as CIS tries to build it.

diff --git a/Engine/Source/Programs/CrashReporter/RegisterPII/App.config b/Engine/Source/Programs/CrashReporter/RegisterPII/App.config
--- a/Engine/Source/Programs/CrashReporter/RegisterPII/App.config
+++ b/Engine/Source/Programs/CrashReporter/RegisterPII/App.config
@@ -1,12 +1,12 @@
-﻿<?xml version="1.0"?>
+<?xml version="1.0"?>
 <configuration>
     <configSections>
         <sectionGroup name="applicationSettings" type="System.Configuration.ApplicationSettingsGroup, System, Version=4.0.0.0, Culture=neutral, PublicKeyToken=b77a5c561934e089">
-            <section name="Tools.CrashReporter.RegisterPII.Properties.Settings" type="System.Configuration.ClientSettingsSection, System, Version=4.0.0.0, Culture=neutral, PublicKeyToken=b77a5c561934e089" requirePermission="false" />
+            <section name="Tools.CrashReporter.RegisterPII.Properties.Settings" type="System.Configuration.ClientSettingsSection, System, Version=4.0.0.0, Culture=neutral, PublicKeyToken=b77a5c561934e089" requirePermission="false"/>
         </sectionGroup>
     </configSections>
     <startup> 
-        <supportedRuntime version="v4.0" sku=".NETFramework,Version=v4.0,Profile=Client"/>
+        <supportedRuntime version="v4.0" sku=".NETFramework,Version=v4.5"/>
     </startup>
     <applicationSettings>
         <Tools.CrashReporter.RegisterPII.Properties.Settings>
diff --git a/Engine/Source/Programs/CrashReporter/RegisterPII/RegisterPII.csproj b/Engine/Source/Programs/CrashReporter/RegisterPII/RegisterPII.csproj
--- a/Engine/Source/Programs/CrashReporter/RegisterPII/RegisterPII.csproj
+++ b/Engine/Source/Programs/CrashReporter/RegisterPII/RegisterPII.csproj
@@ -9,7 +9,7 @@
     <AppDesignerFolder>Properties</AppDesignerFolder>
     <RootNamespace>Tools.CrashReporter.RegisterPII</RootNamespace>
     <AssemblyName>RegisterPII</AssemblyName>
-    <TargetFrameworkVersion>v4.0</TargetFrameworkVersion>
+    <TargetFrameworkVersion>v4.5</TargetFrameworkVersion>
     <FileAlignment>512</FileAlignment>
     <TargetFrameworkProfile>Client</TargetFrameworkProfile>
     <PublishUrl>publish\</PublishUrl>
