Change RegisterPII to properly be .Net 4.5

diff --git a/Engine/Source/Programs/CrashReporter/RegisterPII/RegisterPII.csproj b/Engine/Source/Programs/CrashReporter/RegisterPII/RegisterPII.csproj
--- a/Engine/Source/Programs/CrashReporter/RegisterPII/RegisterPII.csproj
+++ b/Engine/Source/Programs/CrashReporter/RegisterPII/RegisterPII.csproj
@@ -11,7 +11,8 @@
     <AssemblyName>RegisterPII</AssemblyName>
     <TargetFrameworkVersion>v4.5</TargetFrameworkVersion>
     <FileAlignment>512</FileAlignment>
-    <TargetFrameworkProfile>Client</TargetFrameworkProfile>
+    <TargetFrameworkProfile>
+    </TargetFrameworkProfile>
     <PublishUrl>publish\</PublishUrl>
     <Install>true</Install>
     <InstallFrom>Disk</InstallFrom>
