CrashReporter-Frontend - Fixed issue with adding ensures as crashes

diff --git a/Engine/Source/Programs/CrashReporter/CrashReportWebSite/Models/CrashRepository.cs b/Engine/Source/Programs/CrashReporter/CrashReportWebSite/Models/CrashRepository.cs
--- a/Engine/Source/Programs/CrashReporter/CrashReportWebSite/Models/CrashRepository.cs
+++ b/Engine/Source/Programs/CrashReporter/CrashReportWebSite/Models/CrashRepository.cs
@@ -629,6 +629,10 @@ namespace Tools.CrashReporter.CrashReportWebSite.Models
 					NewCrash.CrashType = 2;
 				}
 				else if (NewCrash.RawCallStack.Contains( "FDebug::Ensure" ))
+				{
+					NewCrash.CrashType = 3;			
+				}
+				else if (NewCrash.RawCallStack.Contains( "FDebug::OptionallyLogFormattedEnsureMessageReturningFalse" ))
 				{
 					NewCrash.CrashType = 3;
 				}
