CrashReporter-Frontend - Fixed missing callstack in ensures

diff --git a/Engine/Source/Programs/CrashReporter/CrashReportWebSite/Models/CrashDataFormater.cs b/Engine/Source/Programs/CrashReporter/CrashReportWebSite/Models/CrashDataFormater.cs
--- a/Engine/Source/Programs/CrashReporter/CrashReportWebSite/Models/CrashDataFormater.cs
+++ b/Engine/Source/Programs/CrashReporter/CrashReportWebSite/Models/CrashDataFormater.cs
@@ -257,7 +257,7 @@ namespace Tools.CrashReporter.CrashReportWebSite.Models
 				break;
 			case 3:
 				bSkipping = true;
-				LineToSkipUpto = "FDebug::Ensure";
+				LineToSkipUpto = "FDebug::";
 				break;
 			}
 
