UE-19274/Crashes not being recorded in crash reporter for 4.9 on Windows 10 (Check common appdata for crashes, if local appdata is empty)

diff --git a/Engine/Source/Programs/CrashReporter/CrashReportClient/Private/Windows/WindowsErrorReport.cpp b/Engine/Source/Programs/CrashReporter/CrashReportClient/Private/Windows/WindowsErrorReport.cpp
--- a/Engine/Source/Programs/CrashReporter/CrashReportClient/Private/Windows/WindowsErrorReport.cpp
+++ b/Engine/Source/Programs/CrashReporter/CrashReportClient/Private/Windows/WindowsErrorReport.cpp
@@ -129,6 +129,13 @@ FString FWindowsErrorReport::FindMostRecentErrorReport()
 	SHGetFolderPath(0, CSIDL_LOCAL_APPDATA, NULL, 0, LocalAppDataPath);
 	PlatformFile.IterateDirectory( *(FString(LocalAppDataPath) / TEXT("Microsoft/Windows/WER/ReportQueue")), ReportFinder);
 
+	if (ReportDirectory.Len() == 0)
+	{
+		TCHAR LocalAppDataPath[MAX_PATH];
+		SHGetFolderPath( 0, CSIDL_COMMON_APPDATA, NULL, 0, LocalAppDataPath );
+		PlatformFile.IterateDirectory( *(FString( LocalAppDataPath ) / TEXT( "Microsoft/Windows/WER/ReportQueue" )), ReportFinder );
+	}
+
 	return ReportDirectory;
 }
 
