CrashReporter-Client - Use directory creation time for finding the most recent error report as WER may update timestamp for all reports thus making the data inaccurate

diff --git a/Engine/Source/Programs/CrashReporter/CrashReportClient/Private/Windows/WindowsErrorReport.cpp b/Engine/Source/Programs/CrashReporter/CrashReportClient/Private/Windows/WindowsErrorReport.cpp
--- a/Engine/Source/Programs/CrashReporter/CrashReportClient/Private/Windows/WindowsErrorReport.cpp
+++ b/Engine/Source/Programs/CrashReporter/CrashReportClient/Private/Windows/WindowsErrorReport.cpp
@@ -109,17 +109,17 @@ FString FWindowsErrorReport::FindMostRecentErrorReport()
 {
 	auto& PlatformFile = FPlatformFileManager::Get().GetPlatformFile();
 
-	auto DirectoryModifiedTime = FDateTime::MinValue();
+	auto DirectoryCreationTime = FDateTime::MinValue();
 	FString ReportDirectory;
 	auto ReportFinder = MakeDirectoryVisitor([&](const TCHAR* FilenameOrDirectory, bool bIsDirectory) 
 	{
 		if (bIsDirectory)
 		{
-			auto TimeStamp = PlatformFile.GetTimeStamp(FilenameOrDirectory);
-			if (TimeStamp > DirectoryModifiedTime && FCString::Strstr( FilenameOrDirectory, TEXT("UE4-") ) )
+			const FFileStatData StatData = PlatformFile.GetStatData( FilenameOrDirectory );
+			if (StatData.CreationTime > DirectoryCreationTime && FCString::Strstr( FilenameOrDirectory, TEXT("UE4-") ) )
 			{
 				ReportDirectory = FilenameOrDirectory;
-				DirectoryModifiedTime = TimeStamp;
+				DirectoryCreationTime = StatData.CreationTime;
 			}
 		}
 		return true;
