CrashReporter-Frontend - Fixed '(400) Bad Request.' when displaying reports with invalid jiras

diff --git a/Engine/Source/Programs/CrashReporter/CrashReportWebSite/Controllers/ReportsController.cs b/Engine/Source/Programs/CrashReporter/CrashReportWebSite/Controllers/ReportsController.cs
--- a/Engine/Source/Programs/CrashReporter/CrashReportWebSite/Controllers/ReportsController.cs
+++ b/Engine/Source/Programs/CrashReporter/CrashReportWebSite/Controllers/ReportsController.cs
@@ -217,22 +217,46 @@ namespace Tools.CrashReporter.CrashReportWebSite.Controllers
 				{
 					var BuggsCopy = new List<Bugg>( Buggs );
 
+					HashSet<string> InvalidJiras = new HashSet<string>();
+
+					// Simple verification of JIRA
+					foreach (var Value in FoundJiras)
+					{
+						if( Value.Length < 3 || !Value.Contains('-') )
+						{
+							InvalidJiras.Add(Value);
+						}
+					}
+
+					foreach (var InvalidJira in InvalidJiras)
+					{
+						FoundJiras.Remove( InvalidJira );
+					}
+
 					// Grab the data form JIRA.
 					string JiraSearchQuery = string.Join( " OR ", FoundJiras );
 
 					using( FAutoScopedLogTimer JiraResultsTimer = new FAutoScopedLogTimer( "JiraResults" ) )
 					{
-						var JiraResults = JC.SearchJiraTickets(
+						var JiraResults = new Dictionary<string, Dictionary<string, object>>();
+						try
+						{
+							JiraResults = JC.SearchJiraTickets(
 							JiraSearchQuery,
 							new string[] 
-						{ 
-							"key",				// string
-							"summary",			// string
-							"components",		// System.Collections.ArrayList, Dictionary<string,object>, name
-							"resolution",		// System.Collections.Generic.Dictionary`2[System.String,System.Object], name
-							"fixVersions",		// System.Collections.ArrayList, Dictionary<string,object>, name
-							"customfield_11200" // string
-						} );
+							{ 
+								"key",				// string
+								"summary",			// string
+								"components",		// System.Collections.ArrayList, Dictionary<string,object>, name
+								"resolution",		// System.Collections.Generic.Dictionary`2[System.String,System.Object], name
+								"fixVersions",		// System.Collections.ArrayList, Dictionary<string,object>, name
+								"customfield_11200" // string
+							} );
+						}
+						catch (System.Exception)
+						{
+
+						}
 
 
 						// Jira Key, Summary, Components, Resolution, Fix version, Fix changelist
