CrashReporter-Frontend - UECORE-286/CopyToJira in CrashReporter not catching all user descriptions

diff --git a/Engine/Source/Programs/CrashReporter/CrashReportWebSite/Models/CrashReport.cs b/Engine/Source/Programs/CrashReporter/CrashReportWebSite/Models/CrashReport.cs
--- a/Engine/Source/Programs/CrashReporter/CrashReportWebSite/Models/CrashReport.cs
+++ b/Engine/Source/Programs/CrashReporter/CrashReportWebSite/Models/CrashReport.cs
@@ -156,7 +156,7 @@ namespace Tools.CrashReporter.CrashReportWebSite.Models
 			this.AffectedPlatforms = new SortedSet<string>();
 			this.CrashesInTimeFrameAll = CrashesForBugg.Count;
 			this.CrashesInTimeFrameGroup = CrashesForBugg.Count;
-			var HashSetDecsriptions = new HashSet<string>();
+			var HashSetDescsriptions = new HashSet<string>();
 
 			HashSet<string> MachineIds = new HashSet<string>();
 			int FirstCLAffected = int.MaxValue;
@@ -167,12 +167,10 @@ namespace Tools.CrashReporter.CrashReportWebSite.Models
 				if (Crash.MachineId != null && Crash.MachineId.Length == 32)
 				{
 					MachineIds.Add( Crash.MachineId );
-
-					// Sent in the unattended mode 29 char
-					if( Crash.Description.Length > 32 )
+					if (Crash.Description.Length > 4)
 					{
-						HashSetDecsriptions.Add( Crash.Description );
-					}
+						HashSetDescsriptions.Add( Crash.Description );
+					}										
 				}
 
 				// @TODO Ignore bad build versions.
@@ -206,7 +204,7 @@ namespace Tools.CrashReporter.CrashReportWebSite.Models
 			}
 
 			// CopyToJira 
-			foreach( var Line in HashSetDecsriptions )
+			foreach( var Line in HashSetDescsriptions )
 			{
 				string ListItem = "- " + HttpUtility.HtmlEncode( Line );
 				ToJiraDescriptions.Add( ListItem );
