CrashReporter-Frontend - UECORE-223/CRASH REPORTER: List all related crashes in CrashGroups (added filtering by BuggId and BuiltFromCL)

diff --git a/Engine/Source/Programs/CrashReporter/CrashReportWebSite/Models/CrashReport.cs b/Engine/Source/Programs/CrashReporter/CrashReportWebSite/Models/CrashReport.cs
--- a/Engine/Source/Programs/CrashReporter/CrashReportWebSite/Models/CrashReport.cs
+++ b/Engine/Source/Programs/CrashReporter/CrashReportWebSite/Models/CrashReport.cs
@@ -156,7 +156,7 @@ namespace Tools.CrashReporter.CrashReportWebSite.Models
 			this.AffectedPlatforms = new SortedSet<string>();
 			this.CrashesInTimeFrameAll = CrashesForBugg.Count;
 			this.CrashesInTimeFrameGroup = CrashesForBugg.Count;
-			var HashSetDescsriptions = new HashSet<string>();
+			var HashSetDescriptions = new HashSet<string>();
 
 			HashSet<string> MachineIds = new HashSet<string>();
 			int FirstCLAffected = int.MaxValue;
@@ -169,7 +169,7 @@ namespace Tools.CrashReporter.CrashReportWebSite.Models
 					MachineIds.Add( Crash.MachineId );
 					if (Crash.Description.Length > 4)
 					{
-						HashSetDescsriptions.Add( Crash.Description );
+						HashSetDescriptions.Add( Crash.Description );
 					}										
 				}
 
@@ -204,7 +204,7 @@ namespace Tools.CrashReporter.CrashReportWebSite.Models
 			}
 
 			// CopyToJira 
-			foreach( var Line in HashSetDescsriptions )
+			foreach( var Line in HashSetDescriptions )
 			{
 				string ListItem = "- " + HttpUtility.HtmlEncode( Line );
 				ToJiraDescriptions.Add( ListItem );
diff --git a/Engine/Source/Programs/CrashReporter/CrashReportWebSite/Models/CrashRepository.cs b/Engine/Source/Programs/CrashReporter/CrashReportWebSite/Models/CrashRepository.cs
--- a/Engine/Source/Programs/CrashReporter/CrashReportWebSite/Models/CrashRepository.cs
+++ b/Engine/Source/Programs/CrashReporter/CrashReportWebSite/Models/CrashRepository.cs
@@ -17,7 +17,18 @@ using System.Web.Mvc;
 
 namespace Tools.CrashReporter.CrashReportWebSite.Models
 {
-	
+	class CrashComparer : IEqualityComparer<Crash>
+	{
+		public bool Equals( Crash C1, Crash C2 )
+		{
+			return C1.Id == C2.Id;
+		}
+
+		public int GetHashCode( Crash C )
+		{
+			return C.Id;
+		}
+	}
 
 	/// <summary>
 	/// A model to talk to the database.
@@ -291,6 +302,26 @@ namespace Tools.CrashReporter.CrashReportWebSite.Models
 				// Filter by data and get as enumerable.
 				Results = FilterByDate(ResultsAll, FormData.DateFrom, FormData.DateTo);
 
+				// Filter by BuggId
+				if (!string.IsNullOrEmpty( FormData.BuggId ))
+				{
+					int BuggId = 0;
+					bool bValid = int.TryParse( FormData.BuggId, out BuggId );
+
+					if (bValid)
+					{
+						BuggRepository Buggs = new BuggRepository();
+						Bugg NewBugg = Buggs.GetBugg( BuggId );
+
+						if( NewBugg != null )
+						{
+							List<Crash> Crashes = NewBugg.GetCrashes();
+							var NewResult = Results.Intersect( Crashes, new CrashComparer() );
+							Results = NewResult;
+						}
+					}
+				}				
+
 				// Get UserGroup ResultCounts
 				Dictionary<string, int> GroupCounts = GetCountsByGroupFromCrashes( Results );
 
@@ -357,6 +388,8 @@ namespace Tools.CrashReporter.CrashReportWebSite.Models
 					UsernameQuery = FormData.UsernameQuery,
 					EpicIdOrMachineQuery = FormData.EpicIdOrMachineQuery,
 					MessageQuery = FormData.MessageQuery,
+					BuiltFromCL = FormData.BuiltFromCL,
+					BuggId = FormData.BuggId,
 					JiraQuery = FormData.JiraQuery,
 					DateFrom = (long)( FormData.DateFrom - CrashesViewModel.Epoch ).TotalMilliseconds,
 					DateTo = (long)( FormData.DateTo - CrashesViewModel.Epoch ).TotalMilliseconds,
@@ -805,6 +838,7 @@ namespace Tools.CrashReporter.CrashReportWebSite.Models
 				}
 			}
 
+			// Filter by MessageQuery
 			if (!string.IsNullOrEmpty( FormData.MessageQuery ))
 			{
 				Results =
@@ -815,6 +849,23 @@ namespace Tools.CrashReporter.CrashReportWebSite.Models
 						);
 			}
 
+			// Filter by BuiltFromCL
+			if (!string.IsNullOrEmpty( FormData.BuiltFromCL ))
+			{
+				int BuiltFromCL = 0;
+				bool bValid = int.TryParse( FormData.BuiltFromCL, out BuiltFromCL );
+
+				if (bValid)
+				{
+					Results =
+						(
+							from CrashDetail in Results
+							where CrashDetail.BuiltFromCL.Equals( BuiltFromCL )
+							select CrashDetail
+						);
+				}
+			}
+
 			return Results;
 		}
 	}
diff --git a/Engine/Source/Programs/CrashReporter/CrashReportWebSite/Models/CrashesViewModel.cs b/Engine/Source/Programs/CrashReporter/CrashReportWebSite/Models/CrashesViewModel.cs
--- a/Engine/Source/Programs/CrashReporter/CrashReportWebSite/Models/CrashesViewModel.cs
+++ b/Engine/Source/Programs/CrashReporter/CrashReportWebSite/Models/CrashesViewModel.cs
@@ -50,6 +50,12 @@ namespace Tools.CrashReporter.CrashReportWebSite.Models
 		/// <summary>Message/Summary or Description as query for filtering.</summary>
 		public string MessageQuery { get; set; }
 
+		/// <summary>Bugg Id as query for filtering.</summary>
+		public string BuggId { get; set; }
+
+		/// <summary>BuiltFromCL as query for filtering.</summary>
+		public string BuiltFromCL { get; set; }
+
 		/// <summary>The date of the earliest crash to display.</summary>
 		public long DateFrom { get; set; }
 
diff --git a/Engine/Source/Programs/CrashReporter/CrashReportWebSite/Models/FormHelper.cs b/Engine/Source/Programs/CrashReporter/CrashReportWebSite/Models/FormHelper.cs
--- a/Engine/Source/Programs/CrashReporter/CrashReportWebSite/Models/FormHelper.cs
+++ b/Engine/Source/Programs/CrashReporter/CrashReportWebSite/Models/FormHelper.cs
@@ -29,6 +29,12 @@ namespace Tools.CrashReporter.CrashReportWebSite.Models
 		/// <summary>Message/Summary or Description as query for filtering.</summary>
 		public string MessageQuery = "";
 
+		/// <summary>BuiltFromCL as query for filtering.</summary>
+		public string BuiltFromCL = "";
+
+		/// <summary>BuggId as query for filtering.</summary>
+		public string BuggId = "";
+
 		/// <summary>The page to display from the list.</summary>
 		public int Page = 1;
 
@@ -171,6 +177,8 @@ namespace Tools.CrashReporter.CrashReportWebSite.Models
 			GetFormParameter( Request, Form, "EpicIdOrMachineQuery", EpicIdOrMachineQuery, out EpicIdOrMachineQuery );
 			GetFormParameter( Request, Form, "JiraQuery", JiraQuery, out JiraQuery );
 			GetFormParameter( Request, Form, "MessageQuery", MessageQuery, out MessageQuery );
+			GetFormParameter( Request, Form, "BuggId", BuggId, out BuggId );
+			GetFormParameter( Request, Form, "BuiltFromCL", BuiltFromCL, out BuiltFromCL );
 			GetFormParameter( Request, Form, "SortTerm", DefaultSortTerm, out SortTerm );
 			GetFormParameter( Request, Form, "CrashType", CrashType, out CrashType );
 			GetFormParameter( Request, Form, "UserGroup", UserGroup, out UserGroup );
diff --git a/Engine/Source/Programs/CrashReporter/CrashReportWebSite/Views/Crashes/Index.aspx b/Engine/Source/Programs/CrashReporter/CrashReportWebSite/Views/Crashes/Index.aspx
--- a/Engine/Source/Programs/CrashReporter/CrashReportWebSite/Views/Crashes/Index.aspx
+++ b/Engine/Source/Programs/CrashReporter/CrashReportWebSite/Views/Crashes/Index.aspx
@@ -107,7 +107,7 @@
 			<div id="Container">
 				<table>
 					<tr>
-						<td rowspan="3">
+						<td rowspan="4">
 							<div id="CrashTypeList">
 								<p class="SearchTextTitle" title="Press to hide">Show crashes of type</p>
 								<ul>
@@ -202,6 +202,23 @@
 							<%=Html.DropDownListFor( m=>m.VersionName, Model.VersionNames )%>
 						</td>
 					</tr>
+
+					<tr>
+						<td>
+							<p class="SearchTextTitle">Bugg Id</p>
+						</td>
+						<td>
+							<input name="BuggId" type="text" value="<%=Model.BuggId%>" title="" />
+						</td>
+
+						<td>
+							<p class="SearchTextTitle">Built From CL</p>
+						</td>
+						<td>
+							<input name="BuiltFromCL" type="text" value="<%=Model.BuiltFromCL %>" title="" />
+						</td>
+					</tr>
+
 					<tr>
 						<td colspan="7" >
 							<input type="submit" value="Search" class='SearchButton' />
@@ -248,6 +265,8 @@
 						UsernameQuery = Model.UsernameQuery,
 						EpicIdOrMachineQuery = Model.EpicIdOrMachineQuery,
 						MessageQuery = Model.MessageQuery,
+						BuggId = Model.BuggId,
+						BuiltFromCL = Model.BuiltFromCL,
 						JiraQuery = Model.JiraQuery,
 						DateFrom = Model.DateFrom,
 						DateTo = Model.DateTo,
@@ -271,6 +290,8 @@
 						UsernameQuery = Model.UsernameQuery,
 						EpicIdOrMachineQuery = Model.EpicIdOrMachineQuery,
 						MessageQuery = Model.MessageQuery,
+						BuggId = Model.BuggId,
+						BuiltFromCL = Model.BuiltFromCL,
 						JiraQuery = Model.JiraQuery,
 						DateFrom = Model.DateFrom,
 						DateTo = Model.DateTo,
@@ -300,6 +321,8 @@
 			UsernameQuery = Model.UsernameQuery,
 			EpicIdOrMachineQuery = Model.EpicIdOrMachineQuery,
 			MessageQuery = Model.MessageQuery,
+			BuggId = Model.BuggId,
+			BuiltFromCL = Model.BuiltFromCL,
 			JiraQuery = Model.JiraQuery,
 			DateFrom = Model.DateFrom,
 			DateTo = Model.DateTo,
diff --git a/Engine/Source/Programs/CrashReporter/CrashReportWebSite/Views/Helpers/UrlHelperExtension.cs b/Engine/Source/Programs/CrashReporter/CrashReportWebSite/Views/Helpers/UrlHelperExtension.cs
--- a/Engine/Source/Programs/CrashReporter/CrashReportWebSite/Views/Helpers/UrlHelperExtension.cs
+++ b/Engine/Source/Programs/CrashReporter/CrashReportWebSite/Views/Helpers/UrlHelperExtension.cs
@@ -37,6 +37,8 @@ namespace Tools.CrashReporter.CrashReportWebSite.Views.Helpers
 														UsernameQuery = Model.UsernameQuery,
 														EpicIdOrMachineQuery = Model.EpicIdOrMachineQuery,
 														MessageQuery = Model.MessageQuery,
+														BuggId = Model.BuggId,
+														BuiltFromCL = Model.BuiltFromCL,
 														JiraQuery = Model.JiraQuery,
 														DateFrom = Model.DateFrom,
 														DateTo = Model.DateTo,
@@ -91,6 +93,8 @@ namespace Tools.CrashReporter.CrashReportWebSite.Views.Helpers
 														UsernameQuery = Model.UsernameQuery,
 														EpicIdOrMachineQuery = Model.EpicIdOrMachineQuery,
 														MessageQuery = Model.MessageQuery,
+														BuggId = Model.BuggId,
+														BuiltFromCL = Model.BuiltFromCL,
 														JiraQuery = Model.JiraQuery,
 														DateFrom = Model.DateFrom,
 														DateTo = Model.DateTo,
@@ -128,6 +132,8 @@ namespace Tools.CrashReporter.CrashReportWebSite.Views.Helpers
 														UsernameQuery = Model.UsernameQuery,
 														EpicIdOrMachineQuery = Model.EpicIdOrMachineQuery,
 														MessageQuery = Model.MessageQuery,
+														BuggId = Model.BuggId,
+														BuiltFromCL = Model.BuiltFromCL,
 														JiraQuery = Model.JiraQuery,
 														DateFrom = Model.DateFrom,
 														DateTo = Model.DateTo,
diff --git a/Engine/Source/Programs/CrashReporter/CrashReportWebSite/Views/Home/Index.aspx b/Engine/Source/Programs/CrashReporter/CrashReportWebSite/Views/Home/Index.aspx
--- a/Engine/Source/Programs/CrashReporter/CrashReportWebSite/Views/Home/Index.aspx
+++ b/Engine/Source/Programs/CrashReporter/CrashReportWebSite/Views/Home/Index.aspx
@@ -99,7 +99,7 @@
 			<div id="Container">
 				<table>
 					<tr>
-						<td rowspan="3">
+						<td rowspan="4">
 							<div id="CrashTypeList">
 								<p class="SearchTextTitle" title="Press to hide">Show crashes of type</p>
 								<ul>
@@ -194,6 +194,23 @@
 							<%=Html.DropDownListFor( m=>m.VersionName, Model.VersionNames )%>
 						</td>
 					</tr>
+
+					<tr>
+						<td>
+							<p class="SearchTextTitle">Bugg Id</p>
+						</td>
+						<td>
+							<input name="BuggId" type="text" value="<%=Model.BuggId%>" title="" />
+						</td>
+
+						<td>
+							<p class="SearchTextTitle">Built From CL</p>
+						</td>
+						<td>
+							<input name="BuiltFromCL" type="text" value="<%=Model.BuiltFromCL %>" title="" />
+						</td>
+					</tr>
+
 					<tr>
 						<td colspan="7" >
 							<input type="submit" value="Search" class='SearchButton' />
