CrashReporter-Frontend - Added initial support for the crash context, added support for 'OR-4990/Automatically copy full dumps to Network Share'

diff --git a/Engine/Source/Programs/CrashReporter/CrashReportWebSite/Models/CrashReport.cs b/Engine/Source/Programs/CrashReporter/CrashReportWebSite/Models/CrashReport.cs
--- a/Engine/Source/Programs/CrashReporter/CrashReportWebSite/Models/CrashReport.cs
+++ b/Engine/Source/Programs/CrashReporter/CrashReportWebSite/Models/CrashReport.cs
@@ -9,6 +9,8 @@ using System.Diagnostics;
 using Tools.DotNETCommon;
 using System.Web;
 
+using Tools.CrashReporter.CrashReportCommon;
+
 namespace Tools.CrashReporter.CrashReportWebSite.Models
 {
 	/// <summary> Function call. </summary>
@@ -507,6 +509,29 @@ namespace Tools.CrashReporter.CrashReportWebSite.Models
 	/// </summary>
 	public partial class Crash
 	{
+		/// <summary>Hard coded site path.</summary>
+		const string SitePath = @"\\devweb-02\Sites";
+
+		/// <summary>Crash context for this crash.</summary>
+		FGenericCrashContext CrashContext = null;
+
+		/// <summary>If available, will read CrashContext.runtime-xml.</summary>
+		public void ReadCrashContextIfAvailable()
+		{
+			try
+			{
+				bool bHasCrashContext = HasCrashContextFile();
+				if (bHasCrashContext)
+				{
+					CrashContext = FGenericCrashContext.FromFile( SitePath + GetCrashContextUrl() );
+				}
+			}
+			catch (Exception Ex)
+			{
+				Debug.WriteLine( "Exception in ReadCrashContextIfAvailable: " + Ex.ToString() );
+			}
+		}
+
 		/// <summary> Helper method, display this Bugg as a human readable string. Debugging purpose. </summary>
 		public override string ToString()
 		{
@@ -541,7 +566,10 @@ namespace Tools.CrashReporter.CrashReportWebSite.Models
 		/// </summary>
 		public string GetMiniDumpUrl()
 		{
-			return Properties.Settings.Default.CrashReporterFiles + Id + "_MiniDump.dmp";
+			bool bUseFullMinidumpPath = CrashContext != null && CrashContext.PrimaryCrashProperties.FullCrashDumpLocation.Length > 0;
+			var WebPath = Properties.Settings.Default.CrashReporterFiles + Id + "_MiniDump.dmp";
+			var Path = bUseFullMinidumpPath ? CrashContext.PrimaryCrashProperties.GetFullCrashDumpLocation() : WebPath;
+			return Path;
 		}
 
 		/// <summary>
@@ -552,6 +580,14 @@ namespace Tools.CrashReporter.CrashReportWebSite.Models
 			return Properties.Settings.Default.CrashReporterFiles + Id + "_Diagnostics.txt";
 		}
 
+		/// <summary>
+		/// Return the Url of the crash context file.
+		/// </summary>
+		public string GetCrashContextUrl()
+		{
+			return Properties.Settings.Default.CrashReporterFiles + Id + "_CrashContext.runtime-xml";
+		}
+
 		/// <summary>
 		/// Return the Url of the Windows Error Report meta data file.
 		/// </summary>
@@ -756,9 +792,6 @@ namespace Tools.CrashReporter.CrashReportWebSite.Models
 			}
 		}
 
-		/// <summary>Hard coded site path.</summary>
-		const string SitePath = @"\\devweb-02\Sites";
-
 		/// <summary>Return true, if there is a diagnostics file associated with the crash</summary>
 		public bool HasDiagnosticsFile()
 		{
@@ -766,10 +799,19 @@ namespace Tools.CrashReporter.CrashReportWebSite.Models
 			return System.IO.File.Exists( Path );
 		}
 
+
+		/// <summary>Return true, if there is a crash context file associated with the crash</summary>
+		public bool HasCrashContextFile()
+		{
+			var Path = SitePath + GetCrashContextUrl();
+			return System.IO.File.Exists( Path );
+		}
+
 		/// <summary>Return true, if there is a minidump file associated with the crash</summary>
 		public bool HasMiniDumpFile()
 		{
-			var Path = SitePath + GetMiniDumpUrl();
+			bool bUseFullMinidumpPath = CrashContext != null && CrashContext.PrimaryCrashProperties.FullCrashDumpLocation.Length > 0;
+			var Path = bUseFullMinidumpPath ? CrashContext.PrimaryCrashProperties.GetFullCrashDumpLocation() : SitePath + GetMiniDumpUrl();
 			return System.IO.File.Exists( Path );
 		}
 
diff --git a/Engine/Source/Programs/CrashReporter/CrashReportWebSite/Models/CrashRepository.cs b/Engine/Source/Programs/CrashReporter/CrashReportWebSite/Models/CrashRepository.cs
--- a/Engine/Source/Programs/CrashReporter/CrashReportWebSite/Models/CrashRepository.cs
+++ b/Engine/Source/Programs/CrashReporter/CrashReportWebSite/Models/CrashRepository.cs
@@ -179,7 +179,9 @@ namespace Tools.CrashReporter.CrashReportWebSite.Models
 						select CrashDetail
 					);
 
-					return Crashes.FirstOrDefault();
+					Crash Result = Crashes.FirstOrDefault();
+					Result.ReadCrashContextIfAvailable();
+					return Result;
 				}
 				catch( Exception Ex )
 				{
diff --git a/Engine/Source/Programs/CrashReporter/CrashReportWebSite/Views/Crashes/Show.aspx b/Engine/Source/Programs/CrashReporter/CrashReportWebSite/Views/Crashes/Show.aspx
--- a/Engine/Source/Programs/CrashReporter/CrashReportWebSite/Views/Crashes/Show.aspx
+++ b/Engine/Source/Programs/CrashReporter/CrashReportWebSite/Views/Crashes/Show.aspx
@@ -139,6 +139,14 @@
 						<%
 					} %>
 				</dd>
+			<dt>&nbsp;</dt>
+				<dd>
+					<%if (Model.Crash.HasCrashContextFile() ) 
+					{ %>
+						<a id="A1" style="text-decoration:none;" href='<%=Model.Crash.GetCrashContextUrl() %>'><img src="../../Content/Images/Icons/log.png" style="height:20px;border:none;" />&nbsp;CrashContext</a>
+						<%
+					} %>
+				</dd>
 			<dt>&nbsp;</dt>
 				<dd>
 					<%if (Model.Crash.HasDiagnosticsFile() ) 
