CrashReporter-Frontend - Added initial support for the crash context, added support for 'OR-4990/Automatically copy full dumps to Network Share' (misc fixes and tweaks)

diff --git a/Engine/Source/Programs/CrashReporter/CrashReportWebSite/Models/CrashReport.cs b/Engine/Source/Programs/CrashReporter/CrashReportWebSite/Models/CrashReport.cs
--- a/Engine/Source/Programs/CrashReporter/CrashReportWebSite/Models/CrashReport.cs
+++ b/Engine/Source/Programs/CrashReporter/CrashReportWebSite/Models/CrashReport.cs
@@ -566,12 +566,28 @@ namespace Tools.CrashReporter.CrashReportWebSite.Models
 		/// </summary>
 		public string GetMiniDumpUrl()
 		{
-			bool bUseFullMinidumpPath = CrashContext != null && CrashContext.PrimaryCrashProperties.FullCrashDumpLocation.Length > 0;
 			var WebPath = Properties.Settings.Default.CrashReporterFiles + Id + "_MiniDump.dmp";
-			var Path = bUseFullMinidumpPath ? CrashContext.PrimaryCrashProperties.GetFullCrashDumpLocation() : WebPath;
+			var Path = UseFullMinidumpPath() ? CrashContext.PrimaryCrashProperties.FullCrashDumpLocation : WebPath;
 			return Path;
 		}
 
+		/// <summary>
+		/// Tooltip for the dump.
+		/// </summary>
+		public string GetDumpTitle()
+		{
+			var Title = UseFullMinidumpPath() ? "Copy the link and open in the explorer" : "";
+			return Title;
+		}
+
+		/// <summary>
+		/// Return the short name of the dump.
+		/// </summary>
+		public string GetDumpName()
+		{
+			return UseFullMinidumpPath() ? "Fulldump" : "Minidump";
+		}
+
 		/// <summary>
 		/// Return the Url of the diagnostics file.
 		/// </summary>
@@ -810,11 +826,16 @@ namespace Tools.CrashReporter.CrashReportWebSite.Models
 		/// <summary>Return true, if there is a minidump file associated with the crash</summary>
 		public bool HasMiniDumpFile()
 		{
-			bool bUseFullMinidumpPath = CrashContext != null && CrashContext.PrimaryCrashProperties.FullCrashDumpLocation.Length > 0;
-			var Path = bUseFullMinidumpPath ? CrashContext.PrimaryCrashProperties.GetFullCrashDumpLocation() : SitePath + GetMiniDumpUrl();
+			var Path = UseFullMinidumpPath() ? CrashContext.PrimaryCrashProperties.GetFullCrashDumpLocation() : SitePath + GetMiniDumpUrl();
 			return System.IO.File.Exists( Path );
 		}
 
+		bool UseFullMinidumpPath()
+		{
+			bool bUseFullMinidumpPath = CrashContext != null && !string.IsNullOrEmpty( CrashContext.PrimaryCrashProperties.FullCrashDumpLocation );
+			return bUseFullMinidumpPath;
+		}
+
 		/// <summary>Return true, if there is a log file associated with the crash</summary>
 		public bool HasLogFile()
 		{
diff --git a/Engine/Source/Programs/CrashReporter/CrashReportWebSite/Views/Crashes/Show.aspx b/Engine/Source/Programs/CrashReporter/CrashReportWebSite/Views/Crashes/Show.aspx
--- a/Engine/Source/Programs/CrashReporter/CrashReportWebSite/Views/Crashes/Show.aspx
+++ b/Engine/Source/Programs/CrashReporter/CrashReportWebSite/Views/Crashes/Show.aspx
@@ -127,7 +127,7 @@
 				<dd>	
 					<%if( Model.Crash.HasMiniDumpFile() ) 
 					{ %>
-						<a style='text-decoration:none;' href='<%=Model.Crash.GetMiniDumpUrl() %>'><img src="../../Content/Images/Icons/miniDump.png" style="height:20px;border:none;" />&nbsp;MiniDump</a>
+						<a style='text-decoration:none;' href='<%=Model.Crash.GetMiniDumpUrl() %>' title="<%=Model.Crash.GetDumpTitle()%>"><img src="../../Content/Images/Icons/miniDump.png" style="height:20px;border:none;" />&nbsp;<%=Model.Crash.GetDumpName() %></a>
 						<%
 					} %>
 				</dd>
