CrashReporter-Frontend - Workaround for testing existence of external full minidumps

diff --git a/Engine/Source/Programs/CrashReporter/CrashReportWebSite/Models/CrashReport.cs b/Engine/Source/Programs/CrashReporter/CrashReportWebSite/Models/CrashReport.cs
--- a/Engine/Source/Programs/CrashReporter/CrashReportWebSite/Models/CrashReport.cs
+++ b/Engine/Source/Programs/CrashReporter/CrashReportWebSite/Models/CrashReport.cs
@@ -515,6 +515,8 @@ namespace Tools.CrashReporter.CrashReportWebSite.Models
 		/// <summary>Crash context for this crash.</summary>
 		FGenericCrashContext CrashContext = null;
 
+		bool bUseFullMinidumpPath = false;
+
 		/// <summary>If available, will read CrashContext.runtime-xml.</summary>
 		public void ReadCrashContextIfAvailable()
 		{
@@ -524,6 +526,12 @@ namespace Tools.CrashReporter.CrashReportWebSite.Models
 				if (bHasCrashContext)
 				{
 					CrashContext = FGenericCrashContext.FromFile( SitePath + GetCrashContextUrl() );
+					bool bTest = CrashContext != null && !string.IsNullOrEmpty( CrashContext.PrimaryCrashProperties.FullCrashDumpLocation );
+					if(bTest)
+					{
+						bUseFullMinidumpPath = true;// System.IO.Directory.Exists( CrashContext.PrimaryCrashProperties.FullCrashDumpLocation ); Doesn't work, probably due to some permissions.
+						FLogger.Global.WriteEvent( "ReadCrashContextIfAvailable " + CrashContext.PrimaryCrashProperties.FullCrashDumpLocation + " is " + bUseFullMinidumpPath );
+					}
 				}
 			}
 			catch (Exception Ex)
@@ -826,13 +834,13 @@ namespace Tools.CrashReporter.CrashReportWebSite.Models
 		/// <summary>Return true, if there is a minidump file associated with the crash</summary>
 		public bool HasMiniDumpFile()
 		{
-			var Path = UseFullMinidumpPath() ? CrashContext.PrimaryCrashProperties.GetFullCrashDumpLocation() : SitePath + GetMiniDumpUrl();
-			return System.IO.File.Exists( Path );
+			var Path = SitePath + GetMiniDumpUrl();
+			return UseFullMinidumpPath() ? true : System.IO.File.Exists( Path );
 		}
 
-		bool UseFullMinidumpPath()
+		/// <summary>Whether to use the full minidump.</summary>
+		public bool UseFullMinidumpPath()
 		{
-			bool bUseFullMinidumpPath = CrashContext != null && !string.IsNullOrEmpty( CrashContext.PrimaryCrashProperties.FullCrashDumpLocation );
 			return bUseFullMinidumpPath;
 		}
 
