Crash report client improvements - Using crash context when available (OR-4990/Automatically copy full dumps to Network Share)

diff --git a/Engine/Source/Programs/CrashReporter/CrashReportClient/Private/CrashDescription.cpp b/Engine/Source/Programs/CrashReporter/CrashReportClient/Private/CrashDescription.cpp
--- a/Engine/Source/Programs/CrashReporter/CrashReportClient/Private/CrashDescription.cpp
+++ b/Engine/Source/Programs/CrashReporter/CrashReportClient/Private/CrashDescription.cpp
@@ -12,267 +12,392 @@
 #include "Runtime/Analytics/Analytics/Public/Interfaces/IAnalyticsProvider.h"
 #include "EngineBuildSettings.h"
 
-FCrashDescription::FCrashDescription() :
-	CrashDescriptionVersion( ECrashDescVersions::VER_1_NewCrashFormat ),
-	BuiltFromCL( -1 ),
-	TimeOfCrash( FDateTime::UtcNow() ),
-	bHasMiniDumpFile( false ),
-	bHasLogFile( false ),
-	bHasVideoFile( false ),
-	bHasCompleteData( false )
+// #YRX_Crash: 2015-07-23 Move crashes from [REDACTED_USER_HOME]\AppData\Local\Microsoft\Windows\WER\ReportQueue to [REDACTED_USER_HOME]\AppData\Local\CrashReportClient\Saved
+
+/*-----------------------------------------------------------------------------
+	FCrashProperty
+-----------------------------------------------------------------------------*/
+
+FCrashProperty::FCrashProperty( const FString& InMainCategory, const FString& InSecondCategory, FPrimaryCrashProperties* InOwner ) 
+: Owner( InOwner )
+, CachedValue( TEXT("") )
+, MainCategory( InMainCategory )
+, SecondCategory( InSecondCategory )
+, bSet( false )
 {
-	UpdateIDs();
+
 }
 
-/** Unescapes a specified XML string, naive implementation. */
-FString UnescapeXMLString( const FString& Text )
+FCrashProperty& FCrashProperty::operator=(const FString& NewValue)
 {
-	
-	return Text
-		.Replace( TEXT( "&amp;" ), TEXT( "&" ) )
-		.Replace( TEXT( "&quot;" ), TEXT( "\"" ) )
-		.Replace( TEXT( "&lt;" ), TEXT( "<" ) )
-		.Replace( TEXT( "&gt;" ), TEXT( ">" ) );
+	bSet = true;
+	CachedValue = NewValue;
+	Owner->SetCrashProperty( MainCategory, SecondCategory, CachedValue );
+	return *this;
 }
 
-FCrashDescription::FCrashDescription( FString WERXMLFilepath ) :
-	// WER XML files are forced to be in the first version
-	CrashDescriptionVersion( ECrashDescVersions::VER_1_NewCrashFormat ),
-	BuiltFromCL( -1 ),
-	TimeOfCrash( FDateTime::UtcNow() ),
-	bHasMiniDumpFile( false ),
-	bHasLogFile( false ),
-	bHasVideoFile( false ),
-	bHasCompleteData( false )
+FCrashProperty& FCrashProperty::operator=(const TArray<FString>& NewValue)
 {
-	UpdateIDs();
+	bSet = true;
+	CachedValue = Owner->EncodeArrayStringAsXMLString( NewValue );
+	Owner->SetCrashProperty( MainCategory, SecondCategory, CachedValue );
+	return *this;
+}
+
+FCrashProperty& FCrashProperty::operator=(const bool NewValue)
+{
+	bSet = true;
+	CachedValue = NewValue ? TEXT( "1" ) : TEXT( "0" );
+	Owner->SetCrashProperty( MainCategory, SecondCategory, CachedValue );
+	return *this;
+}
+
+FCrashProperty& FCrashProperty::operator=(const int64 NewValue)
+{
+	bSet = true;
+	CachedValue = LexicalConversion::ToString( NewValue );
+	Owner->SetCrashProperty( MainCategory, SecondCategory, CachedValue );
+	return *this;
+}
+
+const FString& FCrashProperty::AsString() const
+{
+	if (!bSet)
+	{
+		Owner->GetCrashProperty( CachedValue, MainCategory, SecondCategory );
+		bSet = true;
+	}
+	return CachedValue;
+}
+
+
+bool FCrashProperty::AsBool() const
+{
+	return AsString().ToBool();
+}
+
+int64 FCrashProperty::AsInt64() const
+{
+	int64 Value = 0;
+	TTypeFromString<int64>::FromString( Value, *AsString() );
+	return Value;
+}
+
+/*-----------------------------------------------------------------------------
+	FPrimaryCrashProperties
+-----------------------------------------------------------------------------*/
+
+FPrimaryCrashProperties* FPrimaryCrashProperties::Singleton = nullptr;
+
+FPrimaryCrashProperties::FPrimaryCrashProperties()
+	// At this moment only these properties can be changed by the crash report client.
+	: PlatformFullName( FGenericCrashContext::RuntimePropertiesTag, TEXT( "PlatformFullName" ), this )
+	, CommandLine( FGenericCrashContext::RuntimePropertiesTag, TEXT( "CommandLine" ), this )
+	, UserName( FGenericCrashContext::RuntimePropertiesTag, TEXT("UserName"), this )
+	, MachineId( FGenericCrashContext::RuntimePropertiesTag, TEXT( "MachineId" ), this )
+	, EpicAccountId( FGenericCrashContext::RuntimePropertiesTag, TEXT( "EpicAccountId" ), this )
+	// Multiline properties
+	, CallStack( FGenericCrashContext::RuntimePropertiesTag, TEXT( "CallStack" ), this )
+	, SourceContext( FGenericCrashContext::RuntimePropertiesTag, TEXT( "SourceContext" ), this )
+	, Modules( FGenericCrashContext::RuntimePropertiesTag, TEXT( "Modules" ), this )
+	, UserDescription( FGenericCrashContext::RuntimePropertiesTag, TEXT( "UserDescription" ), this )
+	, ErrorMessage( FGenericCrashContext::RuntimePropertiesTag, TEXT( "ErrorMessage" ), this )
+	, FullCrashDumpLocation( FGenericCrashContext::RuntimePropertiesTag, TEXT( "FullCrashDumpLocation" ), this )
+	, TimeOfCrash( FGenericCrashContext::RuntimePropertiesTag, TEXT( "TimeOfCrash" ), this )
+	, bAllowToBeContacted( FGenericCrashContext::RuntimePropertiesTag, TEXT( "bAllowToBeContacted" ), this )
+{
+	CrashVersion = ECrashDescVersions::VER_1_NewCrashFormat;
+	CrashDumpMode = ECrashDumpMode::Default;
+	bHasMiniDumpFile = false;
+	bHasLogFile = false;
+	bHasPrimaryData = false;
+}
 
-	// This is for the current system that uses files from Windows Error Reporting.
-	// Will be replaced with the unified version soon.
-	FXmlFile XmlFile( WERXMLFilepath );
+void FPrimaryCrashProperties::Shutdown()
+{
+	delete Get();
+}
 
-	ReportName = FPaths::GetPath( WERXMLFilepath );
-	FPaths::NormalizeDirectoryName( ReportName );
+void FPrimaryCrashProperties::UpdateIDs()
+{
+	const bool bAddPersonalData = FCrashReportClientConfig::Get().GetAllowToBeContacted();
+	bAllowToBeContacted = bAddPersonalData;
+	if (bAddPersonalData)
+	{
+		// The Epic ID can be looked up from this ID.
+		EpicAccountId = FPlatformMisc::GetEpicAccountId();
+	}
+	else
+	{
+		EpicAccountId = FString();
+	}
+
+	// Add real user name only if log files were allowed since the user name is in the log file and the user consented to sending this information.
+	const bool bSendUserName = FCrashReportClientConfig::Get().GetSendLogFile();
+	if (bSendUserName)
+	{
+		// Remove periods from user names to match AutoReporter user names
+		// The name prefix is read by CrashRepository.AddNewCrash in the website code
+		UserName = FString( FPlatformProcess::UserName() ).Replace( TEXT( "." ), TEXT( "" ) );
+	}
+	else
+	{
+		UserName = FString();
+	}
+
+	MachineId = FPlatformMisc::GetMachineId().ToString( EGuidFormats::Digits );
+}
+
+void FPrimaryCrashProperties::ReadXML( const FString& CrashContextFilepath  )
+{
+	XmlFilepath = CrashContextFilepath;
+	XmlFile = new FXmlFile( XmlFilepath );
+	TimeOfCrash = FDateTime::UtcNow().GetTicks();
+	UpdateIDs();
+}
+
+void FPrimaryCrashProperties::SetCrashGUID( const FString& Filepath )
+{
+	FString CrashDirectory = FPaths::GetPath( Filepath );
+	FPaths::NormalizeDirectoryName( CrashDirectory );
 	// Grab the last component...
-	ReportName = FPaths::GetCleanFilename( ReportName );
+	CrashGUID = FPaths::GetCleanFilename( CrashDirectory );
+}
 
-	const bool bIsValid = XmlFile.IsValid();
-	if( bIsValid )
+FString FPrimaryCrashProperties::EncodeArrayStringAsXMLString( const TArray<FString>& ArrayString ) const
+{
+	const FString Encoded = FString::Join( ArrayString, TEXT("\n") );
+	return Encoded;
+}
+
+void FPrimaryCrashProperties::SendAnalytics()
+{
+	// Connect the crash report client analytics provider.
+	FCrashReportAnalytics::Initialize();
+
+	IAnalyticsProvider& Analytics = FCrashReportAnalytics::GetProvider();
+
+	TArray<FAnalyticsEventAttribute> CrashAttributes;
+
+	CrashAttributes.Add( FAnalyticsEventAttribute( TEXT( "bHasPrimaryData" ), bHasPrimaryData ) );
+	CrashAttributes.Add( FAnalyticsEventAttribute( TEXT( "CrashVersion" ), (int32)CrashVersion ) );
+	CrashAttributes.Add( FAnalyticsEventAttribute( TEXT( "CrashGUID" ), CrashGUID ) );
+
+	//	AppID = GameName
+	CrashAttributes.Add( FAnalyticsEventAttribute( TEXT( "GameName" ), GameName ) );
+
+	//	AppVersion = EngineVersion
+	CrashAttributes.Add( FAnalyticsEventAttribute( TEXT( "EngineVersion" ), EngineVersion.ToString() ) );
+
+	// @see UpdateIDs()
+	CrashAttributes.Add( FAnalyticsEventAttribute( TEXT( "MachineID" ), MachineId.AsString() ) );
+	CrashAttributes.Add( FAnalyticsEventAttribute( TEXT( "UserName" ), UserName.AsString() ) );
+	CrashAttributes.Add( FAnalyticsEventAttribute( TEXT( "EpicAccountId" ), EpicAccountId.AsString() ) );
+
+	CrashAttributes.Add( FAnalyticsEventAttribute( TEXT( "Platform" ), PlatformFullName.AsString() ) );
+	CrashAttributes.Add( FAnalyticsEventAttribute( TEXT( "TimeOfCrash" ), TimeOfCrash.AsString() ) );
+	CrashAttributes.Add( FAnalyticsEventAttribute( TEXT( "EngineMode" ), EngineMode ) );
+	CrashAttributes.Add( FAnalyticsEventAttribute( TEXT( "AppDefaultLocale" ), AppDefaultLocale ) );
+
+	Analytics.RecordEvent( TEXT( "CrashReportClient.ReportCrash" ), CrashAttributes );
+
+	// Shutdown analytics.
+	FCrashReportAnalytics::Shutdown();
+}
+
+void FPrimaryCrashProperties::Save()
+{
+	XmlFile->Save( XmlFilepath );
+}
+
+/*-----------------------------------------------------------------------------
+	FCrashContextReader
+-----------------------------------------------------------------------------*/
+
+FCrashContext::FCrashContext( const FString& CrashContextFilepath )
+{
+	ReadXML( CrashContextFilepath );
+
+	const bool bIsValid = XmlFile->IsValid();
+	if (bIsValid)
 	{
-		const FXmlNode* OSVersionInformationNode = XmlFile.GetRootNode()->FindChildNode( TEXT( "OSVersionInformation" ) );
-		const FXmlNode* ProblemSignaturesNode = XmlFile.GetRootNode()->FindChildNode( TEXT( "ProblemSignatures" ) );
-		const FXmlNode* DynamicSignaturesNode = XmlFile.GetRootNode()->FindChildNode( TEXT( "DynamicSignatures" ) );
+		// Setup properties required for the analytics.
+		GetCrashProperty( CrashVersion, FGenericCrashContext::RuntimePropertiesTag, TEXT( "CrashVersion" ) );
+		GetCrashProperty( CrashGUID, FGenericCrashContext::RuntimePropertiesTag, TEXT( "CrashGUID" ) );
+		GetCrashProperty( CrashDumpMode, FGenericCrashContext::RuntimePropertiesTag, TEXT( "CrashDumpMode" ) );
+		GetCrashProperty( GameName, FGenericCrashContext::RuntimePropertiesTag, TEXT( "GameName" ) );
+		GetCrashProperty( EngineVersion, FGenericCrashContext::RuntimePropertiesTag, TEXT( "EngineVersion" ) );
+
+		GetCrashProperty( BaseDir, FGenericCrashContext::RuntimePropertiesTag, TEXT( "BaseDir" ) );
+		FString Misc_OSVersionMajor;
+		GetCrashProperty( Misc_OSVersionMajor, FGenericCrashContext::RuntimePropertiesTag, TEXT( "Misc.OSVersionMajor" ) );
+		FString Misc_OSVersionMinor;
+		GetCrashProperty( Misc_OSVersionMinor, FGenericCrashContext::RuntimePropertiesTag, TEXT( "Misc.OSVersionMinor" ) );
+
+		bool Misc_Is64bitOperatingSystem = false;
+		GetCrashProperty( Misc_Is64bitOperatingSystem, FGenericCrashContext::RuntimePropertiesTag, TEXT( "Misc.Is64bitOperatingSystem" ) );
+
+		// Extract the Platform component.
+		TArray<FString> SubDirs;
+		BaseDir.ParseIntoArray( SubDirs, TEXT( "/" ), true );
+		const int SubDirsNum = SubDirs.Num();
+		const FString PlatformName = SubDirsNum > 0 ? SubDirs[SubDirsNum - 1] : TEXT( "" );
+		if (Misc_OSVersionMajor.Len() > 0)
+		{
+			PlatformFullName = FString::Printf( TEXT( "%s [%s %s %s]" ), *PlatformName, *Misc_OSVersionMajor, *Misc_OSVersionMinor, Misc_Is64bitOperatingSystem ? TEXT( "64b" ) : TEXT( "32b" ) );
+		}
+		else
+		{
+			PlatformFullName = PlatformName;
+		}
 
-		FString Product;
-		int EngineVersionComponents = 0;
+		GetCrashProperty( EngineMode, FGenericCrashContext::RuntimePropertiesTag, TEXT( "EngineMode" ) );
+		GetCrashProperty( AppDefaultLocale, FGenericCrashContext::RuntimePropertiesTag, TEXT( "AppDefaultLocale" ) );
 
-		if( OSVersionInformationNode )
+		if (CrashDumpMode == ECrashDumpMode::FullDump)
 		{
-			const FXmlNode* ProductNode = OSVersionInformationNode->FindChildNode( TEXT( "Product" ) );
-			if( ProductNode )
+			// Set the full dump crash location when we have a full dump.
+			const FString LocationForBranch = FCrashReportClientConfig::Get().GetFullCrashDumpLocationForBranch( EngineVersion.GetBranch() );
+			if (!LocationForBranch.IsEmpty())
 			{
-				Product = ProductNode->GetContent();
+				FullCrashDumpLocation = LocationForBranch / CrashGUID + TEXT("_") + EngineVersion.ToString();
 			}
 		}
 
-		if( ProblemSignaturesNode )
+		bHasPrimaryData = true;
+	}
+}
+
+/*-----------------------------------------------------------------------------
+	FCrashDescription
+-----------------------------------------------------------------------------*/
+
+FCrashWERContext::FCrashWERContext( const FString& WERXMLFilepath )
+	: FPrimaryCrashProperties()
+{
+	ReadXML( WERXMLFilepath );
+	CrashGUID = FPaths::GetCleanFilename( WERXMLFilepath );
+
+	const bool bIsValid = XmlFile->IsValid();
+	if (bIsValid)
+	{
+		FString BuildVersion;
+		FString BranchName;
+		uint32 BuiltFromCL = 0;
+		int EngineVersionComponents = 0;
+
+		GetCrashProperty( GameName, TEXT( "ProblemSignatures" ), TEXT( "Parameter0" ) );
+
+		GetCrashProperty( BuildVersion, TEXT( "ProblemSignatures" ), TEXT( "Parameter1" ) );
+		if (!BuildVersion.IsEmpty())
 		{
-			const FXmlNode* GameNameNode = ProblemSignaturesNode->FindChildNode( TEXT( "Parameter0" ) );
-			if( GameNameNode )
-			{
-				GameName = GameNameNode->GetContent();
-			}
+			EngineVersionComponents++;
+		}
+
+		FString Parameter8Value;
+		GetCrashProperty( Parameter8Value, TEXT( "ProblemSignatures" ), TEXT( "Parameter8" ) );
+		if (!Parameter8Value.IsEmpty())
+		{
+			TArray<FString> ParsedParameters8;
+			Parameter8Value.ParseIntoArray( ParsedParameters8, TEXT( "!" ), false );
 
-			const FXmlNode* BuildVersionNode = ProblemSignaturesNode->FindChildNode( TEXT( "Parameter1" ) );
-			if( BuildVersionNode )
+			if (ParsedParameters8.Num() > 1)
 			{
-				BuildVersion = BuildVersionNode->GetContent();
-				EngineVersionComponents++;
+				CommandLine = FGenericCrashContext::UnescapeXMLString( ParsedParameters8[1] );
+				CrashDumpMode = CommandLine.AsString().Contains( TEXT( "-fullcrashdump" ) ) ? ECrashDumpMode::FullDump : ECrashDumpMode::Default;
 			}
 
-			const FXmlNode* Parameter8Node = ProblemSignaturesNode->FindChildNode( TEXT( "Parameter8" ) );
-			if( Parameter8Node )
+			if (ParsedParameters8.Num() > 2)
 			{
-				const FString Parameter8Value = Parameter8Node->GetContent();
+				ErrorMessage = ParsedParameters8[2];
+			}
+		}
 
-				TArray<FString> ParsedParameters8;
-				Parameter8Value.ParseIntoArray( ParsedParameters8, TEXT( "!" ), false );
+		FString Parameter9Value;
+		GetCrashProperty( Parameter9Value, TEXT( "ProblemSignatures" ), TEXT( "Parameter9" ) );
+		if (!Parameter9Value.IsEmpty())
+		{
+			TArray<FString> ParsedParameters9;
+			Parameter9Value.ParseIntoArray( ParsedParameters9, TEXT( "!" ), false );
 
-				if( ParsedParameters8.Num() > 1 )
-				{
-					CommandLine = UnescapeXMLString( ParsedParameters8[1] );
-				}
+			if (ParsedParameters9.Num() > 0)
+			{
+				BranchName = ParsedParameters9[0].Replace( TEXT( "+" ), TEXT( "/" ) );
 
-				if( ParsedParameters8.Num() > 2 )
+				const FString DepotRoot = TEXT( "//depot/" );
+				if (BranchName.StartsWith( DepotRoot ))
 				{
-					ErrorMessage.Add( ParsedParameters8[2] );
+					BranchName = BranchName.Mid( DepotRoot.Len() );
 				}
+				EngineVersionComponents++;
 			}
 
-			const FXmlNode* Parameter9Node = ProblemSignaturesNode->FindChildNode( TEXT( "Parameter9" ) );
-			if( Parameter9Node )
+			if (ParsedParameters9.Num() > 1)
 			{
-				const FString Parameter9Value = Parameter9Node->GetContent();
-
-				TArray<FString> ParsedParameters9;
-				Parameter9Value.ParseIntoArray( ParsedParameters9, TEXT( "!" ), false );
+				const FString BaseDirectory = ParsedParameters9[1];
 
-				if( ParsedParameters9.Num() > 0 )
-				{
-					BranchName = ParsedParameters9[0].Replace( TEXT( "+" ), TEXT( "/" ) );
-
-					const FString DepotRoot = TEXT( "//depot/" );
-					if( BranchName.StartsWith( DepotRoot ) )
-					{
-						BranchName = BranchName.Mid( DepotRoot.Len() );
-					}
-					EngineVersionComponents++;
-				}
+				TArray<FString> SubDirs;
+				BaseDirectory.ParseIntoArray( SubDirs, TEXT( "/" ), true );
+				const int SubDirsNum = SubDirs.Num();
+				const FString PlatformName = SubDirsNum > 0 ? SubDirs[SubDirsNum - 1] : TEXT( "" );
 
-				if( ParsedParameters9.Num() > 1 )
+				FString Product;
+				GetCrashProperty( Product, TEXT( "OSVersionInformation" ), TEXT( "Product" ) );
+				if (Product.Len() > 0)
 				{
-					const FString BaseDirectory = ParsedParameters9[1];
-
-					TArray<FString> SubDirs;
-					BaseDirectory.ParseIntoArray( SubDirs, TEXT( "/" ), true );
-					const int SubDirsNum = SubDirs.Num();
-					const FString PlatformName = SubDirsNum > 0 ? SubDirs[SubDirsNum - 1] : TEXT("");
-					if( Product.Len() > 0 )
-					{
-						Platform = FString::Printf( TEXT( "%s [%s]" ), *PlatformName, *Product );
-					}
-					else
-					{
-						Platform = PlatformName;
-					}
+					PlatformFullName = FString::Printf( TEXT( "%s [%s]" ), *PlatformName, *Product );
 				}
-
-				if( ParsedParameters9.Num() > 2 )
+				else
 				{
-					EngineMode = ParsedParameters9[2];
+					PlatformFullName = PlatformName;
 				}
+			}
 
-				if( ParsedParameters9.Num() > 3 )
-				{
-					TTypeFromString<uint32>::FromString( BuiltFromCL, *ParsedParameters9[3] );
-					EngineVersionComponents++;
-				}
+			if (ParsedParameters9.Num() > 2)
+			{
+				EngineMode = ParsedParameters9[2];
 			}
-		}
 
-		if( DynamicSignaturesNode )
-		{
-			const FXmlNode* LanguageNode = DynamicSignaturesNode->FindChildNode( TEXT( "Parameter2" ) );
-			if( LanguageNode )
+			if (ParsedParameters9.Num() > 3)
 			{
-				LanguageLCID = LanguageNode->GetContent();
+				TTypeFromString<uint32>::FromString( BuiltFromCL, *ParsedParameters9[3] );
+				EngineVersionComponents++;
 			}
 		}
 
 		// We have all three components of the engine version, so initialize it.
-		if( EngineVersionComponents == 3 )
+		if (EngineVersionComponents == 3)
 		{
-			InitializeEngineVersion();
+			InitializeEngineVersion( BuildVersion, BranchName, BuiltFromCL );
 		}
 
-		bHasCompleteData = true;
-	}
-	else
-	{
-		bHasCompleteData = false;
+		bHasPrimaryData = true;
 	}
 }
 
-void FCrashDescription::InitializeEngineVersion()
+void FCrashWERContext::InitializeEngineVersion( const FString& BuildVersion, const FString& BranchName, uint32 BuiltFromCL )
 {
 	uint16 Major = 0;
 	uint16 Minor = 0;
 	uint16 Patch = 0;
-
+ 
 	TArray<FString> ParsedBuildVersion;
 	BuildVersion.ParseIntoArray( ParsedBuildVersion, TEXT( "." ), false );
-
-	if( ParsedBuildVersion.Num() >= 3 )
+ 
+	if (ParsedBuildVersion.Num() >= 3)
 	{
 		TTypeFromString<uint16>::FromString( Patch, *ParsedBuildVersion[2] );
 	}
-
-	if( ParsedBuildVersion.Num() >= 2 )
+ 
+	if (ParsedBuildVersion.Num() >= 2)
 	{
 		TTypeFromString<uint16>::FromString( Minor, *ParsedBuildVersion[1] );
 	}
-
-	if( ParsedBuildVersion.Num() >= 1 )
+ 
+	if (ParsedBuildVersion.Num() >= 1)
 	{
 		TTypeFromString<uint16>::FromString( Major, *ParsedBuildVersion[0] );
 	}
-
+ 
 	EngineVersion = FEngineVersion( Major, Minor, Patch, BuiltFromCL, BranchName );
 }
-
-void FCrashDescription::UpdateIDs()
-{
-	const bool bAddPersonalData = FCrashReportClientConfig::Get().GetAllowToBeContacted();
-	if (bAddPersonalData)
-	{
-		// The Epic ID can be looked up from this ID.
-		EpicAccountId = FPlatformMisc::GetEpicAccountId();
-	}
-	else
-	{
-		EpicAccountId.Empty();
-	}
-
-	// Add real user name only if log files were allowed since the user name is in the log file and the user consented to sending this information.
-	const bool bSendUserName = FCrashReportClientConfig::Get().GetSendLogFile();
-	if (bSendUserName)
-	{
-		// Remove periods from user names to match AutoReporter user names
-		// The name prefix is read by CrashRepository.AddNewCrash in the website code
-		UserName = FString( FPlatformProcess::UserName() ).Replace( TEXT( "." ), TEXT( "" ) );
-	}
-	else
-	{
-		UserName.Empty();
-	}
-
-	MachineId = FPlatformMisc::GetMachineId().ToString( EGuidFormats::Digits );
-}
-
-void FCrashDescription::SendAnalytics()
-{
-		// Connect the crash report client analytics provider.
-		FCrashReportAnalytics::Initialize();
-
-		IAnalyticsProvider& Analytics = FCrashReportAnalytics::GetProvider();
-
-		TArray<FAnalyticsEventAttribute> CrashAttributes;
-
-		CrashAttributes.Add( FAnalyticsEventAttribute( TEXT( "bHasCompleteData" ), bHasCompleteData ) );
-		CrashAttributes.Add( FAnalyticsEventAttribute( TEXT( "CrashDescriptionVersion" ), (int32)CrashDescriptionVersion ) );
-		CrashAttributes.Add( FAnalyticsEventAttribute( TEXT( "ReportName" ), ReportName ) );
-
-		//	AppID = GameName
-		CrashAttributes.Add( FAnalyticsEventAttribute( TEXT( "GameName" ), GameName ) );
-
-		//	AppVersion = EngineVersion
-		CrashAttributes.Add( FAnalyticsEventAttribute( TEXT( "EngineVersion" ), EngineVersion.ToString() ) );
-		CrashAttributes.Add( FAnalyticsEventAttribute( TEXT( "BuildVersion" ), BuildVersion ) );
-		CrashAttributes.Add( FAnalyticsEventAttribute( TEXT( "BuiltFromCL" ), BuiltFromCL ) );
-		CrashAttributes.Add( FAnalyticsEventAttribute( TEXT( "BranchName" ), BranchName ) );
-
-		// @see UpdateIDs()
-		CrashAttributes.Add( FAnalyticsEventAttribute( TEXT( "MachineID" ), MachineId ) );
-		CrashAttributes.Add( FAnalyticsEventAttribute( TEXT( "UserName" ), UserName ) );
-		CrashAttributes.Add( FAnalyticsEventAttribute( TEXT( "EpicAccountId" ), EpicAccountId ) );
-
-		CrashAttributes.Add( FAnalyticsEventAttribute( TEXT( "Platform" ), Platform ) );
-		CrashAttributes.Add( FAnalyticsEventAttribute( TEXT( "TimeOfCrash" ), TimeOfCrash.GetTicks() ) );
-		CrashAttributes.Add( FAnalyticsEventAttribute( TEXT( "EngineMode" ), EngineMode ) );
-		CrashAttributes.Add( FAnalyticsEventAttribute( TEXT( "LanguageLCID" ), LanguageLCID ) );
-
-		Analytics.RecordEvent( TEXT( "CrashReportClient.ReportCrash" ), CrashAttributes );
-
-		// Shutdown analytics.
-		FCrashReportAnalytics::Shutdown();
-	}
diff --git a/Engine/Source/Programs/CrashReporter/CrashReportClient/Private/CrashDescription.h b/Engine/Source/Programs/CrashReporter/CrashReportClient/Private/CrashDescription.h
--- a/Engine/Source/Programs/CrashReporter/CrashReportClient/Private/CrashDescription.h
+++ b/Engine/Source/Programs/CrashReporter/CrashReportClient/Private/CrashDescription.h
@@ -2,25 +2,151 @@
 
 #pragma once
 #include "GenericPlatformCrashContext.h"
-
-/**
- *	Describes a unified crash, should be used by all platforms.
- *	Still lacks some of the properties, they will be added later.
- *	Must contain the same properties as ...\CrashReportServer\CrashReportCommon\CrashDescription.cs.
- *	Contains all usable information about the crash. 
- *	
- *	At this moment only used to sent the analytics.
+#include "UnrealString.h"
+
+struct FPrimaryCrashProperties;
+
+/** PrimaryCrashProperties. Extracted from: FGenericCrashContext::SerializeContentToBuffer */
+/*
+	"CrashVersion"
+	"ProcessId"
+	"CrashGUID"
+	"IsInternalBuild"
+	"IsPerforceBuild"
+	"IsSourceDistribution"
+	"SecondsSinceStart"
+	"GameName"
+	"ExecutableName"
+	"BuildConfiguration"
+	"PlatformName"
+	"PlatformNameIni"
+	"PlatformFullName"
+	"EngineMode"
+	"EngineVersion"
+	"CommandLine"
+	"LanguageLCID"
+	"AppDefaultLocale"
+	"IsUE4Release"
+	"UserName"
+	"BaseDir"
+	"RootDir"
+	"MachineId"
+	"EpicAccountId"
+	"CallStack"
+	"SourceContext"
+	"UserDescription"
+	"ErrorMessage"
+	"CrashDumpMode"
+	"Misc.NumberOfCores"
+	"Misc.NumberOfCoresIncludingHyperthreads"
+	"Misc.Is64bitOperatingSystem"
+	"Misc.CPUVendor"
+	"Misc.CPUBrand"
+	"Misc.PrimaryGPUBrand"
+	"Misc.OSVersionMajor"
+	"Misc.OSVersionMinor"
+	"Misc.AppDiskTotalNumberOfBytes"
+	"Misc.AppDiskNumberOfFreeBytes"
+	"MemoryStats.TotalPhysical"
+	"MemoryStats.TotalVirtual"
+	"MemoryStats.PageSize"
+	"MemoryStats.TotalPhysicalGB"
+	"MemoryStats.AvailablePhysical"
+	"MemoryStats.AvailableVirtual"
+	"MemoryStats.UsedPhysical"
+	"MemoryStats.PeakUsedPhysical"
+	"MemoryStats.UsedVirtual"
+	"MemoryStats.PeakUsedVirtual"
+	"MemoryStats.bIsOOM"
+	"MemoryStats.OOMAllocationSize"
+	"MemoryStats.OOMAllocationAlignment"
+	"TimeofCrash"
+	"bAllowToBeContacted"
  */
 
-struct FCrashDescription
+namespace LexicalConversion
 {
-	/** 
-	 * Version of this crash description. 
-	 */
-	ECrashDescVersions CrashDescriptionVersion;
+	inline void FromString( ECrashDescVersions& OutValue, const TCHAR* Buffer )
+	{
+		OutValue = (ECrashDescVersions)FCString::Atoi( Buffer );
+	}
+
+	inline void FromString( ECrashDumpMode& OutValue, const TCHAR* Buffer )
+	{
+		OutValue = (ECrashDumpMode)FCString::Atoi( Buffer );
+	}
+
+	inline void FromString( FEngineVersion& OutValue, const TCHAR* Buffer )
+	{
+		FEngineVersion::Parse( Buffer, OutValue );
+	}
+
+	inline void FromString( FString& OutValue, const TCHAR* Buffer )
+	{
+		OutValue = Buffer;
+	}
+}
+
+/** Simple crash property. Only for string values. */
+struct FCrashProperty
+{
+	friend struct FPrimaryCrashProperties;
+
+protected:
+	/** Initialization constructor. */
+	FCrashProperty( const FString& InMainCategory, const FString& InSecondCategory, FPrimaryCrashProperties* InOwner );
+
+public:
+	/** Assignment operator for string. */
+	FCrashProperty& operator=(const FString& NewValue);
+
+	/** Assignment operator for arrays. */
+	FCrashProperty& operator=(const TArray<FString>& NewValue);
+
+	/** Assignment operator for bool. */
+	FCrashProperty& operator=(const bool NewValue);
+
+	/** Assignment operator for int64. */
+	FCrashProperty& operator=(const int64 NewValue);
+
+	/** Getter for string, default. */
+	const FString& AsString() const;
+
+	/** Getter for bool. */
+	bool AsBool() const;
 
-	/** An unique report name that this crash belongs to. */
-	FString ReportName;
+	/** Getter for int64. */
+	int64 AsInt64() const;
+
+protected:
+	/** Owner of the property. */
+	FPrimaryCrashProperties* Owner;
+
+	/** Cached value of the property. */
+	mutable FString CachedValue;
+
+	/** Main category in the crash context. */
+	FString MainCategory;
+
+	/** Second category in the crash context. */
+	FString SecondCategory;
+
+	mutable bool bSet;
+};
+
+/** Primary crash properties required by the crash report system. */
+struct FPrimaryCrashProperties
+{
+	friend struct FCrashProperty;
+
+	/** Version. */
+	ECrashDescVersions CrashVersion;
+
+	/** Crash dump mode. */
+	ECrashDumpMode CrashDumpMode;
+
+	/** An unique report name that this crash belongs to. Folder name. */
+	FString CrashGUID;
 
 	/**
 	 * The name of the game that crashed. (AppID)
@@ -42,9 +168,9 @@ struct FCrashDescription
 	 * The platform that crashed e.g. Win64.
 	 * @PlatformName	varchar(64)	
 	 * 
-	 * FPlatformProperties::PlatformName()
+	 * Last path of the directory
 	 */
-	FString Platform;
+	FCrashProperty PlatformFullName;
 
 	/** 
 	 * Encoded engine version. (AppVersion)
@@ -57,21 +183,14 @@ struct FCrashDescription
 	 */
 	FEngineVersion EngineVersion;
 
-	/**
-	 * The three component version of the app e.g. 4.4.1
-	 * @BuildVersion	varchar(128)
-	 * 
-	 * ENGINE_MAJOR_VERSION.ENGINE_MINOR_VERSION.ENGINE_PATCH_VERSION
-	 */
-	FString BuildVersion;
-
 	/**
 	 * Built from changelist.
 	 * @ChangeListVersion	varchar(64)
 	 * 
 	 * BUILT_FROM_CHANGELIST
 	 */
-	uint32 BuiltFromCL;
+	//EngineVersion.GetChangelist()
+	//uint32 BuiltFromCL;
 
 	/**
 	 * The name of the branch this game was built out of.
@@ -79,7 +198,7 @@ struct FCrashDescription
 	 * 
 	 * BRANCH_NAME
 	 */
-	FString BranchName;
+	//EngineVersion.GetBranch();
 
 	/**
 	 * The command line of the application that crashed.
@@ -87,7 +206,7 @@ struct FCrashDescription
 	 * 
 	 * FCommandLine::Get() 
 	 */
-	FString CommandLine;
+	FCrashProperty CommandLine;
 
 	/**
 	 * The base directory where the app was running.
@@ -101,9 +220,9 @@ struct FCrashDescription
 	 * The language ID the application that crashed.
 	 * @LanguageExt varchar(64)
 	 * 
-	 * FInternationalization::Get().GetCurrentCulture()->GetLCID()
+	 * FPlatformMisc::GetDefaultLocale()
 	 */
-	FString LanguageLCID;
+	FString AppDefaultLocale;
 
 	/** 
 	 * The name of the user that caused this crash.
@@ -111,7 +230,7 @@ struct FCrashDescription
 	 * 
 	 * FString( FPlatformProcess::UserName() ).Replace( TEXT( "." ), TEXT( "" ) )
 	 */
-	FString UserName;
+	FCrashProperty UserName;
 
 	/**
 	 * The unique ID used to identify the machine the crash occurred on.
@@ -119,7 +238,7 @@ struct FCrashDescription
 	 * 
 	 * FPlatformMisc::GetMachineId().ToString( EGuidFormats::Digits )
 	 */
-	FString MachineId;
+	FCrashProperty MachineId;
 
 	/** 
 	 * The Epic account ID for the user who last used the Launcher.
@@ -127,31 +246,34 @@ struct FCrashDescription
 	 * 
 	 * FPlatformMisc::GetEpicAccountId()
 	 */
-	FString EpicAccountId;
+	FCrashProperty EpicAccountId;
 
 	/**
 	 * An array of FStrings representing the callstack of the crash.
 	 * @RawCallStack	varchar(MAX)
 	 * 
-	 * @TODO
 	 */
-	TArray<FString> CallStack;
+	FCrashProperty CallStack;
 
 	/**
 	 * An array of FStrings showing the source code around the crash.
 	 * @SourceContext varchar(max)
 	 * 
-	 * @TODO
 	 */
-	TArray<FString> SourceContext;
+	FCrashProperty SourceContext;
+
+	/**
+	* An array of module's name used by the game that crashed.
+	*
+	*/
+	FCrashProperty Modules;
 
 	/**
 	 * An array of FStrings representing the user description of the crash.
 	 * @Description	varchar(512)
 	 * 
-	 * @TODO
 	 */
-	TArray<FString> UserDescription;
+	FCrashProperty UserDescription;
 
 	/**
 	 * The error message, can be assertion message, ensure message or message from the fatal error.
@@ -159,31 +281,10 @@ struct FCrashDescription
 	 * 
 	 * GErrorMessage
 	 */
-	TArray<FString> ErrorMessage;
+	FCrashProperty ErrorMessage;
 
-	/**
-	 *	Generic platform memory stats.
-	 *	@ADD-TO-THE-DATABASE
-	 *	
-	 *	FPlatformMemoryStats::GetStats()
-	 */
-	FGenericPlatformMemoryStats MemoryStats;
-
-	/**
-	 *	Whether this crash was caused by the OOM.
-	 *	@ADD-TO-THE-DATABASE
-	 *	
-	 *	@TODO
-	 */
-	bool bIsOOM;
-
-	/**
-	 * The misc and platform specific payload data.
-	 * @ADD-TO-THE-DATABASE	varchar(max)
-	 * 
-	 * @TODO
-	 */
-	TArray<FString> PayloadData;
+	/** Location of full crash dump. Displayed in the crash report frontend. */
+	FCrashProperty FullCrashDumpLocation;
 
 	/**
 	 * The UTC time the crash occurred.
@@ -191,7 +292,14 @@ struct FCrashDescription
 	 * 
 	 * FDateTime::UtcNow().GetTicks()
 	 */
-	FDateTime TimeOfCrash;
+	FCrashProperty TimeOfCrash;
+
+	/**
+	 *	Whether the user allowed us to be contacted.
+	 *	If true the following properties are retrieved from the system: UserName (for non-launcher build) and EpicAccountID.
+	 *	Otherwise they will be empty.
+	 */
+	FCrashProperty bAllowToBeContacted;
 
 	/**
 	 * Whether this crash has a minidump file.
@@ -205,29 +313,134 @@ struct FCrashDescription
 	 */
 	bool bHasLogFile;
 
-	/**
-	 * Whether this crash has a video file.
-	 * @HasVideoFile bit 
-	 */
-	bool bHasVideoFile;
+	/** Whether this crash contains primary usable data. */
+	bool bHasPrimaryData;
+
+protected:
+	/** Default constructor. */
+	FPrimaryCrashProperties();
+
 
-	/** Whether this crash contains complete data, otherwise only IDs are available. */
-	bool bHasCompleteData;
+	/** Destructor. */
+	~FPrimaryCrashProperties()
+	{
+		delete XmlFile;
+	}
+
+public:
+	/** Sets new instance as the global. */
+	static void Set( FPrimaryCrashProperties* NewInstance )
+	{
+		Singleton = NewInstance;
+	}
 
 	/**
-	 * A default constructor.
+	 * @return global instance of the primary crash properties for the currently processed/displayed crash
 	 */
-	FCrashDescription();
+	static FPrimaryCrashProperties* Get()
+	{
+		return Singleton;
+	}
 
-	/** Initializes instance based on specified WER XML string. */
-	explicit FCrashDescription( FString WERXMLFilepath );
+	/** Shutdowns the global instance. */
+	static void Shutdown();
 
-	/** Initializes engine version from the separate components. */
-	void InitializeEngineVersion();
+	/** Whether this crash contains callstack, error message and source context thus it means that crash is complete. */
+	bool HasProcessedData() const
+	{
+		return CallStack.AsString().Len() > 0 && ErrorMessage.AsString().Len() > 0;
+	}
 
 	/** Updates following properties: UserName, MachineID and EpicAccountID. */
 	void UpdateIDs();
 
 	/** Sends this crash for analytics. */
 	void SendAnalytics();
+
+	/** Saves the data. */
+	void Save();
+
+protected:
+	/** Reads previously set XML file. */
+	void ReadXML( const FString& CrashContextFilepath );
+
+	/** Sets the CrasgGUID based on the report's path. */
+	void SetCrashGUID( const FString& Filepath );
+
+	/** Gets a crash property from the XML file. */
+	template <typename Type>
+	void GetCrashProperty( Type& out_ReadValue, const FString& MainCategory, const FString& SecondCategory ) const
+	{
+		const FXmlNode* MainNode = XmlFile->GetRootNode()->FindChildNode( MainCategory );
+		if (MainNode)
+		{
+			const FXmlNode* CategoryNode = MainNode->FindChildNode( SecondCategory );
+			if (CategoryNode)
+			{
+				LexicalConversion::FromString( out_ReadValue, *FGenericCrashContext::UnescapeXMLString( CategoryNode->GetContent() ) );
+			}
+		}
+	}
+
+	/** Sets a crash property to a new value. */
+	template <typename Type>
+	void SetCrashProperty( const FString& MainCategory, const FString& SecondCategory, const Type& Value )
+	{
+		SetCrashProperty( MainCategory, SecondCategory, *TTypeToString<Type>::ToString( Value ) );
+	}
+
+	/** Sets a crash property to a new value. */
+	void SetCrashProperty( const FString& MainCategory, const FString& SecondCategory, const FString& NewValue )
+	{
+		FXmlNode* MainNode = XmlFile->GetRootNode()->FindChildNode( MainCategory );
+		if (MainNode)
+		{
+			FXmlNode* CategoryNode = MainNode->FindChildNode( SecondCategory );
+			const FString EscapedValue = FGenericCrashContext::EscapeXMLString( NewValue );
+			if (CategoryNode)
+			{
+				CategoryNode->SetContent( EscapedValue );
+			}
+			else
+			{
+				MainNode->AppendChildNode( SecondCategory, EscapedValue );
+			}
+		}
+	}
+
+	/** Encodes multi line property to be saved as single line. */
+	FString EncodeArrayStringAsXMLString( const TArray<FString>& ArrayString ) const;
+
+	/** Reader for the xml file. */
+	FXmlFile* XmlFile;
+
+	/** Cached filepath. */
+	FString XmlFilepath;
+
+	/** Global instance. */
+	static FPrimaryCrashProperties* Singleton;
+};
+
+/**
+ *	Describes a unified crash, should be used by all platforms.
+ *	Based on FGenericCrashContext, reads all saved properties, accessed by looking into read XML. 
+ *	Still lacks some of the properties, they will be added later.
+ *	Must contain the same properties as ...\CrashReportServer\CrashReportCommon\CrashDescription.cs.
+ *	Contains all usable information about the crash. 
+ *	
+ */
+struct FCrashContext : public FPrimaryCrashProperties
+{
+	/** Initializes instance based on specified Crash Context filepath. */
+	explicit FCrashContext( const FString& CrashContextFilepath );
+};
+
+/** Crash context based on the Window Error Reporting WER files, obsolete, only for backward compatibility. */
+struct FCrashWERContext : public FPrimaryCrashProperties
+{
+	/** Initializes instance based on specified WER XML filepath. */
+	explicit FCrashWERContext( const FString& WERXMLFilepath );
+
+	/** Initializes engine version from the separate components. */
+	void InitializeEngineVersion( const FString& BuildVersion, const FString& BranchName, uint32 BuiltFromCL );
 };
\ No newline at end of file
diff --git a/Engine/Source/Programs/CrashReporter/CrashReportClient/Private/CrashReportClient.cpp b/Engine/Source/Programs/CrashReporter/CrashReportClient/Private/CrashReportClient.cpp
--- a/Engine/Source/Programs/CrashReporter/CrashReportClient/Private/CrashReportClient.cpp
+++ b/Engine/Source/Programs/CrashReporter/CrashReportClient/Private/CrashReportClient.cpp
@@ -12,11 +12,16 @@
 
 #define LOCTEXT_NAMESPACE "CrashReportClient"
 
-FCrashDescription& GetCrashDescription()
+struct FCrashReportUtil
 {
-	static FCrashDescription Singleton;
-	return Singleton;
-}
+	/** Formats processed diagnostic text by adding additional information about machine and user. */
+	static FText FormatDiagnosticText( const FText& DiagnosticText )
+	{
+		const FString MachineId = FPrimaryCrashProperties::Get()->MachineId.AsString();
+		const FString EpicAccountId = FPrimaryCrashProperties::Get()->EpicAccountId.AsString();
+		return FText::Format( LOCTEXT( "CrashReportClientCallstackPattern", "MachineId:{0}\nEpicAccountId:{1}\n\n{2}" ), FText::FromString( MachineId ), FText::FromString( EpicAccountId ), DiagnosticText );
+	}
+};
 
 #if !CRASH_REPORT_UNATTENDED_ONLY
 
@@ -29,15 +34,31 @@ FCrashReportClient::FCrashReportClient(const FPlatformErrorReport& InErrorReport
 	, bShouldWindowBeHidden(false)
 	, bSendData(false)
 {
-
-	if (!ErrorReport.TryReadDiagnosticsFile(DiagnosticText) && !FParse::Param(FCommandLine::Get(), TEXT("no-local-diagnosis")))
+	bool bUsePrimaryData = false;
+	if (FPrimaryCrashProperties::Get()->HasProcessedData())
 	{
-		DiagnoseReportTask = new FAsyncTask<FDiagnoseReportWorker>( this );
-		DiagnoseReportTask->StartBackgroundTask();
+		bUsePrimaryData = true;
 	}
-	else if( !DiagnosticText.IsEmpty() )
+	else
 	{
-		FormattedDiagnosticText = FCrashReportUtil::FormatDiagnosticText( DiagnosticText, GetCrashDescription().MachineId, GetCrashDescription().EpicAccountId, GetCrashDescription().UserName );
+		if (!ErrorReport.TryReadDiagnosticsFile() && !FParse::Param( FCommandLine::Get(), TEXT( "no-local-diagnosis" ) ))
+		{
+			DiagnoseReportTask = new FAsyncTask<FDiagnoseReportWorker>( this );
+			DiagnoseReportTask->StartBackgroundTask();
+		}
+		else
+		{
+			bUsePrimaryData = true;
+		}
+	}
+
+	if (bUsePrimaryData)
+	{
+		const FString CallstackString = FPrimaryCrashProperties::Get()->CallStack.AsString();
+		const FString ReportString = FString::Printf( TEXT( "%s\n\n%s" ), *FPrimaryCrashProperties::Get()->ErrorMessage.AsString(), *CallstackString );
+		DiagnosticText = FText::FromString( ReportString );
+
+		FormattedDiagnosticText = FCrashReportUtil::FormatDiagnosticText( FText::FromString( ReportString ) );
 	}
 }
 
@@ -62,7 +83,7 @@ FReply FCrashReportClient::Submit()
 	if (ErrorReport.HasFilesToUpload())
 	{
 		// Send analytics.
-		GetCrashDescription().SendAnalytics();
+		FPrimaryCrashProperties::Get()->SendAnalytics();
 	}
 
 	bSendData = true;
@@ -76,8 +97,7 @@ FReply FCrashReportClient::SubmitAndRestart()
 	Submit();
 
 	const FString CrashedAppPath = ErrorReport.FindCrashedAppPath();
-	const FCrashDescription& CrashDescription = GetCrashDescription();
-	const FString CommandLineArguments = CrashDescription.CommandLine;
+	const FString CommandLineArguments = FPrimaryCrashProperties::Get()->CommandLine.AsString();
 
 	FPlatformProcess::CreateProc(*CrashedAppPath, *CommandLineArguments, true, false, false, NULL, 0, NULL, NULL);
 
@@ -131,10 +151,13 @@ void FCrashReportClient::AllowToBeContacted_OnCheckStateChanged( ECheckBoxState
 	FCrashReportClientConfig::Get().SetAllowToBeContacted( NewRadioState == ECheckBoxState::Checked );
 
 	// Refresh PII based on the bAllowToBeContacted flag.
-	GetCrashDescription().UpdateIDs();
+	FPrimaryCrashProperties::Get()->UpdateIDs();
+
+	// Save updated properties.
+	FPrimaryCrashProperties::Get()->Save();
 
 	// Update diagnostics text.
-	FormattedDiagnosticText = FCrashReportUtil::FormatDiagnosticText( DiagnosticText, GetCrashDescription().MachineId, GetCrashDescription().EpicAccountId, GetCrashDescription().UserName );
+	FormattedDiagnosticText = FCrashReportUtil::FormatDiagnosticText( DiagnosticText );
 }
 
 void FCrashReportClient::SendLogFile_OnCheckStateChanged( ECheckBoxState NewRadioState )
@@ -150,7 +173,7 @@ void FCrashReportClient::StartTicker()
 void FCrashReportClient::StoreCommentAndUpload()
 {
 	// Write user's comment
-	ErrorReport.SetUserComment( UserComment, FCrashReportClientConfig::Get().GetAllowToBeContacted() );
+	ErrorReport.SetUserComment( UserComment );
 	StartTicker();
 }
 
@@ -184,23 +207,28 @@ bool FCrashReportClient::Tick(float UnusedDeltaTime)
 	return false;
 }
 
-FString FCrashReportClient::GetCrashedAppName() const
-{
-	return GetCrashDescription().GameName;
-}
-
 FString FCrashReportClient::GetCrashDirectory() const
 {
 	return ErrorReport.GetReportDirectory();
 }
 
-void FCrashReportClient::FinalizeDiagnoseReportWorker( FText ReportText )
+void FCrashReportClient::FinalizeDiagnoseReportWorker()
 {
-	DiagnosticText = ReportText;
-	FormattedDiagnosticText = FCrashReportUtil::FormatDiagnosticText( ReportText, GetCrashDescription().MachineId, GetCrashDescription().EpicAccountId, GetCrashDescription().UserName );
+	// Update properties for the crash.
+	ErrorReport.SetPrimaryCrashProperties( *FPrimaryCrashProperties::Get() );
+
+	FString CallstackString = FPrimaryCrashProperties::Get()->CallStack.AsString();
+	if (CallstackString.IsEmpty())
+	{
+		DiagnosticText = LOCTEXT( "NoDebuggingSymbols", "You do not have any debugging symbols required to display the callstack for this crash." );
+	}
+	else
+	{
+		const FString ReportString = FString::Printf( TEXT( "%s\n\n%s" ), *FPrimaryCrashProperties::Get()->ErrorMessage.AsString(), *CallstackString );
+		DiagnosticText = FText::FromString( ReportString );
+	}
 
-	auto DiagnosticsFilePath = ErrorReport.GetReportDirectory() / FCrashReportClientConfig::Get().GetDiagnosticsFilename();
-	Uploader.LocalDiagnosisComplete(FPaths::FileExists(DiagnosticsFilePath) ? DiagnosticsFilePath : TEXT(""));
+	FormattedDiagnosticText = FCrashReportUtil::FormatDiagnosticText( DiagnosticText );
 }
 
 
@@ -215,20 +243,16 @@ FDiagnoseReportWorker::FDiagnoseReportWorker( FCrashReportClient* InCrashReportC
 
 void FDiagnoseReportWorker::DoWork()
 {
-	const FText ReportText = CrashReportClient->ErrorReport.DiagnoseReport();
+	CrashReportClient->ErrorReport.DiagnoseReport();
+
 	// Inform the game thread that we are done.
 	FSimpleDelegateGraphTask::CreateAndDispatchWhenReady
 	(
-		FSimpleDelegateGraphTask::FDelegate::CreateRaw( CrashReportClient, &FCrashReportClient::FinalizeDiagnoseReportWorker, ReportText ),
+		FSimpleDelegateGraphTask::FDelegate::CreateRaw( CrashReportClient, &FCrashReportClient::FinalizeDiagnoseReportWorker ),
 		TStatId(), nullptr, ENamedThreads::GameThread
 	);
 }
 
 #endif // !CRASH_REPORT_UNATTENDED_ONLY
 
-FText FCrashReportUtil::FormatDiagnosticText( const FText& DiagnosticText, const FString MachineId, const FString EpicAccountId, const FString UserNameNoDot )
-{
-	return FText::Format( LOCTEXT( "CrashReportClientCallstackPattern", "MachineId:{0}\nEpicAccountId:{1}\n\n{2}" ), FText::FromString( MachineId ), FText::FromString( EpicAccountId ), DiagnosticText );
-}
-
 #undef LOCTEXT_NAMESPACE
diff --git a/Engine/Source/Programs/CrashReporter/CrashReportClient/Private/CrashReportClient.h b/Engine/Source/Programs/CrashReporter/CrashReportClient/Private/CrashReportClient.h
--- a/Engine/Source/Programs/CrashReporter/CrashReportClient/Private/CrashReportClient.h
+++ b/Engine/Source/Programs/CrashReporter/CrashReportClient/Private/CrashReportClient.h
@@ -85,13 +85,6 @@ public:
 	 */
 	FText GetDiagnosticText() const;
 
-	/**
-	 * Access to name (which usually includes the build config) of the app that crashed
-	 * @return Name of executable
-	 */
-	FString GetCrashedAppName() const;
-
-
 	/**
 	 * @return the full path of the crash directory.
 	 */
@@ -144,7 +137,7 @@ private:
 	void StartTicker();
 
 	/** Enqueued from the diagnose report worker thread to be executed on the game thread. */
-	void FinalizeDiagnoseReportWorker( FText ReportText );
+	void FinalizeDiagnoseReportWorker();
 
 	/**
 	 * @return true if we are still processing a callstack
diff --git a/Engine/Source/Programs/CrashReporter/CrashReportClient/Private/CrashReportClientApp.cpp b/Engine/Source/Programs/CrashReporter/CrashReportClient/Private/CrashReportClientApp.cpp
--- a/Engine/Source/Programs/CrashReporter/CrashReportClient/Private/CrashReportClientApp.cpp
+++ b/Engine/Source/Programs/CrashReporter/CrashReportClient/Private/CrashReportClientApp.cpp
@@ -73,17 +73,26 @@ FPlatformErrorReport LoadErrorReport()
 	}
 
 	FPlatformErrorReport ErrorReport(ReportDirectoryAbsolutePath);
-	
-	FString XMLWerFilename;
-	ErrorReport.FindFirstReportFileWithExtension( XMLWerFilename, TEXT( ".xml" ) );
 
-	extern FCrashDescription& GetCrashDescription();
-	GetCrashDescription() = FCrashDescription( ReportDirectoryAbsolutePath / XMLWerFilename );	
+	FString Filename;
+	// CrashContext.runtime-xml has the precedence over the WER
+	if (ErrorReport.FindFirstReportFileWithExtension( Filename, *FGenericCrashContext::CrashContextExtension ))
+	{
+		FPrimaryCrashProperties::Set( new FCrashContext( ReportDirectoryAbsolutePath / Filename ) );
+	}
+	else if (ErrorReport.FindFirstReportFileWithExtension( Filename, TEXT( ".xml" ) ))
+	{
+		FPrimaryCrashProperties::Set( new FCrashWERContext( ReportDirectoryAbsolutePath / Filename ) );
+	}
+	else
+	{
+		return FPlatformErrorReport();
+	}
 
 #if CRASH_REPORT_UNATTENDED_ONLY
 	return ErrorReport;
 #else
-	if( !GameNameFromCmd.IsEmpty() && GameNameFromCmd != GetCrashDescription().GameName )
+	if (!GameNameFromCmd.IsEmpty() && GameNameFromCmd != FPrimaryCrashProperties::Get()->GameName)
 	{
 		// Don't display or upload anything if it's not the report we expected
 		ErrorReport = FPlatformErrorReport();
@@ -92,54 +101,6 @@ FPlatformErrorReport LoadErrorReport()
 #endif
 }
 
-FCrashReportClientConfig::FCrashReportClientConfig()
-	: DiagnosticsFilename( TEXT( "Diagnostics.txt" ))
-{
-	const bool bUnattended =
-#if CRASH_REPORT_UNATTENDED_ONLY
-		true;
-#else
-		FApp::IsUnattended();
-#endif // CRASH_REPORT_UNATTENDED_ONLY
-
-	if( !GConfig->GetString( TEXT( "CrashReportClient" ), TEXT( "CrashReportReceiverIP" ), CrashReportReceiverIP, GEngineIni ) )
-	{
-		// Use the default value.
-		CrashReportReceiverIP = TEXT( "http://crashreporter.epicgames.com:57005" );
-	}
-
-	if ( !GConfig->GetBool( TEXT( "CrashReportClient" ), TEXT( "bAllowToBeContacted" ), bAllowToBeContacted, GEngineIni ) )
-	{
-		// Default to true when unattended when config is missing. This is mostly for dedicated servers that do not have config files for CRC.
-		if (bUnattended)
-		{
-			bAllowToBeContacted = true;
-		}
-	}
-
-	if ( !GConfig->GetBool( TEXT( "CrashReportClient" ), TEXT( "bSendLogFile" ), bSendLogFile, GEngineIni ) )
-	{
-		// Default to true when unattended when config is missing. This is mostly for dedicated servers that do not have config files for CRC.
-		if (bUnattended)
-		{
-			bSendLogFile = true;
-		}
-	}
-
-	UE_LOG( CrashReportClientLog, Log, TEXT( "CrashReportReceiverIP: %s" ), *CrashReportReceiverIP );
-}
-
-void FCrashReportClientConfig::SetAllowToBeContacted( bool bNewValue )
-{
-	bAllowToBeContacted = bNewValue;
-	GConfig->SetBool( TEXT( "CrashReportClient" ), TEXT( "bAllowToBeContacted" ), bAllowToBeContacted, GEngineIni );
-}
-
-void FCrashReportClientConfig::SetSendLogFile( bool bNewValue )
-{
-	bSendLogFile = bNewValue;
-	GConfig->SetBool( TEXT( "CrashReportClient" ), TEXT( "bSendLogFile" ), bSendLogFile, GEngineIni );
-}
 
 void RunCrashReportClient(const TCHAR* CommandLine)
 {
@@ -174,7 +135,7 @@ void RunCrashReportClient(const TCHAR* CommandLine)
 	if (bUnattended)
 	{
 		// In the unattended mode we don't send any PII.
-		ErrorReport.SetUserComment( NSLOCTEXT( "CrashReportClient", "UnattendedMode", "Sent in the unattended mode" ), false );
+		ErrorReport.SetUserComment( NSLOCTEXT( "CrashReportClient", "UnattendedMode", "Sent in the unattended mode" ) );
 		FCrashReportClientUnattended CrashReportClient( ErrorReport );
 
 		// loop until the app is ready to quit
@@ -236,18 +197,11 @@ void RunCrashReportClient(const TCHAR* CommandLine)
 #endif // !CRASH_REPORT_UNATTENDED_ONLY
 	}
 
+	FPrimaryCrashProperties::Shutdown();
 	FPlatformErrorReport::ShutDown();
 
 	FEngineLoop::AppPreExit();
 	FTaskGraphInterface::Shutdown();
 
 	FEngineLoop::AppExit();
-}
-
-void CrashReportClientCheck(bool bCondition, const TCHAR* Location)
-{
-	if (!bCondition)
-	{
-		UE_LOG(CrashReportClientLog, Warning, TEXT("CHECK FAILED at %s"), Location);
-	}
 }
\ No newline at end of file
diff --git a/Engine/Source/Programs/CrashReporter/CrashReportClient/Private/CrashReportClientConfig.cpp b/Engine/Source/Programs/CrashReporter/CrashReportClient/Private/CrashReportClientConfig.cpp
new file mode 100644
--- /dev/null
+++ b/Engine/Source/Programs/CrashReporter/CrashReportClient/Private/CrashReportClientConfig.cpp
@@ -0,0 +1,101 @@
+// Copyright 1998-2015 Epic Games, Inc. All Rights Reserved.
+
+#include "CrashReportClientApp.h"
+
+FCrashReportClientConfig::FCrashReportClientConfig()
+	: DiagnosticsFilename( TEXT( "Diagnostics.txt" ) )
+	, SectionName( TEXT( "CrashReportClient" ) )
+{
+	const bool bUnattended =
+#if CRASH_REPORT_UNATTENDED_ONLY
+		true;
+#else
+		FApp::IsUnattended();
+#endif // CRASH_REPORT_UNATTENDED_ONLY
+
+	if (!GConfig->GetString( *SectionName, TEXT( "CrashReportReceiverIP" ), CrashReportReceiverIP, GEngineIni ))
+	{
+		// Use the default value.
+		CrashReportReceiverIP = TEXT( "http://crashreporter.epicgames.com:57005" );
+	}
+
+	if (!GConfig->GetBool( TEXT( "CrashReportClient" ), TEXT( "bAllowToBeContacted" ), bAllowToBeContacted, GEngineIni ))
+	{
+		// Default to true when unattended when config is missing. This is mostly for dedicated servers that do not have config files for CRC.
+		if (bUnattended)
+		{
+			bAllowToBeContacted = true;
+		}
+	}
+
+	if (!GConfig->GetBool( TEXT( "CrashReportClient" ), TEXT( "bSendLogFile" ), bSendLogFile, GEngineIni ))
+	{
+		// Default to true when unattended when config is missing. This is mostly for dedicated servers that do not have config files for CRC.
+		if (bUnattended)
+		{
+			bSendLogFile = true;
+		}
+	}
+	
+	ReadFullCrashDumpConfigurations();
+
+	UE_LOG( CrashReportClientLog, Log, TEXT( "CrashReportReceiverIP: %s" ), *CrashReportReceiverIP );
+}
+
+void FCrashReportClientConfig::SetAllowToBeContacted( bool bNewValue )
+{
+	bAllowToBeContacted = bNewValue;
+	GConfig->SetBool( *SectionName, TEXT( "bAllowToBeContacted" ), bAllowToBeContacted, GEngineIni );
+}
+
+void FCrashReportClientConfig::SetSendLogFile( bool bNewValue )
+{
+	bSendLogFile = bNewValue;
+	GConfig->SetBool( *SectionName, TEXT( "bSendLogFile" ), bSendLogFile, GEngineIni );
+}
+
+const FString FCrashReportClientConfig::GetFullCrashDumpLocationForBranch( const FString& BranchName ) const
+{
+	for (const auto& It : FullCrashDumpConfigurations)
+	{
+		const bool bExactMatch = It.bExactMatch;
+		if (bExactMatch && BranchName == It.BranchName)
+		{
+			return It.Location;
+		}
+		else if( BranchName.Contains( It.BranchName ) )
+		{
+			return It.Location;
+		}
+	}
+
+	return TEXT( "" );
+}
+
+FString FCrashReportClientConfig::GetKey( const FString& KeyName )
+{
+	FString Result;
+	if (!GConfig->GetString( *SectionName, *KeyName, Result, GEngineIni ))
+	{
+		return TEXT( "" );
+	}
+	return Result;
+}
+
+void FCrashReportClientConfig::ReadFullCrashDumpConfigurations()
+{
+	for (int32 NumEntries = 0;; ++NumEntries)
+	{
+		FString Branch = GetKey( FString::Printf( TEXT( "FullCrashDump_%i_Branch" ), NumEntries ) );
+		if (Branch.IsEmpty())
+		{
+			break;
+		}
+
+		const FString NetworkLocation = GetKey( FString::Printf( TEXT( "FullCrashDump_%i_Location" ), NumEntries ) );
+		const bool bExactMatch = !Branch.EndsWith( TEXT( "*" ) );
+		Branch.ReplaceInline( TEXT( "*" ), TEXT( "" ) );
+
+		FullCrashDumpConfigurations.Add( FFullCrashDumpEntry( Branch, NetworkLocation, bExactMatch ) );
+	}
+}
diff --git a/Engine/Source/Programs/CrashReporter/CrashReportClient/Private/CrashReportClientConfig.h b/Engine/Source/Programs/CrashReporter/CrashReportClient/Private/CrashReportClientConfig.h
new file mode 100644
--- /dev/null
+++ b/Engine/Source/Programs/CrashReporter/CrashReportClient/Private/CrashReportClientConfig.h
@@ -0,0 +1,107 @@
+// Copyright 1998-2015 Epic Games, Inc. All Rights Reserved.
+
+#pragma once
+
+#include "Core.h"
+
+/**
+*  Holds FullCrashDump properties from the config.
+*
+*	FullCrashDump_0_Branch=UE4
+*	FullCrashDump_0_Location=\\epicgames.net\root\Builds\UE4
+*	FullCrashDump_1_Branch=...
+*	...
+*/
+struct FFullCrashDumpEntry
+{
+	/** Initialization constructor. */
+	FFullCrashDumpEntry( const FString& InBranchName, const FString& InLocation, const bool bInExactMatch )
+		: BranchName( InBranchName )
+		, Location( InLocation )
+		, bExactMatch( bInExactMatch )
+	{}
+
+
+	/** Partial branch name. */
+	const FString BranchName;
+
+	/** Location where the full crash dump will be copied. Usually a network share. */
+	const FString Location;
+
+	/**
+	*	Branch=UE4 means exact match
+	*	Branch=UE4* means contain match
+	*/
+	const bool bExactMatch;
+};
+
+/** Holds basic configuration for the crash report client. */
+struct FCrashReportClientConfig
+{
+	/** Accesses the singleton. */
+	static FCrashReportClientConfig& Get()
+	{
+		static FCrashReportClientConfig Instance;
+		return Instance;
+	}
+
+	/** Initialization constructor. */
+	FCrashReportClientConfig();
+
+	const FString& GetReceiverAddress() const
+	{
+		return CrashReportReceiverIP;
+	}
+
+	const FString& GetDiagnosticsFilename() const
+	{
+		return DiagnosticsFilename;
+	}
+
+	const bool& GetAllowToBeContacted() const
+	{
+		return bAllowToBeContacted;
+	}
+
+	const bool& GetSendLogFile() const
+	{
+		return bSendLogFile;
+	}
+
+	void SetAllowToBeContacted( bool bNewValue );
+	void SetSendLogFile( bool bNewValue );
+
+	/**
+	 * @return location for full crash dump for the specified branch.
+	 */
+	const FString GetFullCrashDumpLocationForBranch( const FString& BranchName ) const;
+
+protected:
+	/** Returns empty string if couldn't read */
+	FString GetKey( const FString& KeyName );
+
+	/** Reads FFullCrashDumpEntry config entries. */
+	void ReadFullCrashDumpConfigurations();
+
+	/** IP address of crash report receiver. */
+	FString CrashReportReceiverIP;
+
+	/** Filename to use when saving diagnostics report, if generated locally. */
+	FString DiagnosticsFilename;
+
+	/** Section for crash report client configuration. */
+	FString SectionName;
+
+	/** Configuration used for copying full dump crashes. */
+	TArray<FFullCrashDumpEntry> FullCrashDumpConfigurations;
+
+	/**
+	*	Whether the user allowed us to be contacted.
+	*	If true the following properties are retrieved from the system: UserName (for non-launcher build) and EpicAccountID.
+	*	Otherwise they will be empty.
+	*/
+	bool bAllowToBeContacted;
+
+	/** Whether the user allowed us to send the log file. */
+	bool bSendLogFile;
+};
diff --git a/Engine/Source/Programs/CrashReporter/CrashReportClient/Private/CrashReportClientUnattended.cpp b/Engine/Source/Programs/CrashReporter/CrashReportClient/Private/CrashReportClientUnattended.cpp
--- a/Engine/Source/Programs/CrashReporter/CrashReportClient/Private/CrashReportClientUnattended.cpp
+++ b/Engine/Source/Programs/CrashReporter/CrashReportClient/Private/CrashReportClientUnattended.cpp
@@ -3,14 +3,19 @@
 #include "CrashReportClientApp.h"
 #include "CrashReportClientUnattended.h"
 #include "CrashReportUtil.h"
+#include "CrashDescription.h"
 
-FCrashReportClientUnattended::FCrashReportClientUnattended(const FPlatformErrorReport& ErrorReport)
+FCrashReportClientUnattended::FCrashReportClientUnattended(FPlatformErrorReport& ErrorReport)
 	: Uploader( FCrashReportClientConfig::Get().GetReceiverAddress() )
 {
+	ErrorReport.TryReadDiagnosticsFile();
+
 	// Process the report synchronously
 	ErrorReport.DiagnoseReport();
-	auto DiagnosticsFilePath = ErrorReport.GetReportDirectory() / FCrashReportClientConfig::Get().GetDiagnosticsFilename();
-	Uploader.LocalDiagnosisComplete( FPaths::FileExists( DiagnosticsFilePath ) ? DiagnosticsFilePath : TEXT( "" ) );
+
+	// Update properties for the crash.
+	ErrorReport.SetPrimaryCrashProperties( *FPrimaryCrashProperties::Get() );
+
 	Uploader.BeginUpload( ErrorReport );
 	StartTicker();
 }
diff --git a/Engine/Source/Programs/CrashReporter/CrashReportClient/Private/CrashReportClientUnattended.h b/Engine/Source/Programs/CrashReporter/CrashReportClient/Private/CrashReportClientUnattended.h
--- a/Engine/Source/Programs/CrashReporter/CrashReportClient/Private/CrashReportClientUnattended.h
+++ b/Engine/Source/Programs/CrashReporter/CrashReportClient/Private/CrashReportClientUnattended.h
@@ -14,7 +14,7 @@ public:
 	 * Set up uploader object
 	 * @param ErrorReport Error report to upload
 	 */
-	explicit FCrashReportClientUnattended(const FPlatformErrorReport& ErrorReport);
+	explicit FCrashReportClientUnattended(FPlatformErrorReport& ErrorReport);
 
 private:
 	/**
diff --git a/Engine/Source/Programs/CrashReporter/CrashReportClient/Private/CrashReportUtil.h b/Engine/Source/Programs/CrashReporter/CrashReportClient/Private/CrashReportUtil.h
--- a/Engine/Source/Programs/CrashReporter/CrashReportClient/Private/CrashReportUtil.h
+++ b/Engine/Source/Programs/CrashReporter/CrashReportClient/Private/CrashReportUtil.h
@@ -4,60 +4,6 @@
 
 #include "Core.h"
 
-/** Holds basic configuration for the crash report client. */
-struct FCrashReportClientConfig
-{
-	/** Accesses the singleton. */
-	static FCrashReportClientConfig& Get()
-	{
-		static FCrashReportClientConfig Instance;
-		return Instance;
-	}
-
-	/** Initialization constructor. */
-	FCrashReportClientConfig();
-
-	const FString& GetReceiverAddress() const
-	{
-		return CrashReportReceiverIP;
-	}
-
-	const FString& GetDiagnosticsFilename() const
-	{
-		return DiagnosticsFilename;
-	}
-
-	const bool& GetAllowToBeContacted() const
-	{
-		return bAllowToBeContacted;
-	}
-
-	const bool& GetSendLogFile() const
-	{
-		return bSendLogFile;
-	}
-
-	void SetAllowToBeContacted( bool bNewValue );
-	void SetSendLogFile( bool bNewValue );
-
-protected:
-	/** IP address of crash report receiver. */
-	FString CrashReportReceiverIP;
-
-	/** Filename to use when saving diagnostics report, if generated locally. */
-	FString DiagnosticsFilename;
-
-	/**
-	 *	Whether the user allowed us to be contacted. 
-	 *	If true the following properties are retrieved from the system: UserName (for non-launcher build) and EpicAccountID.
-	 *	Otherwise they will be empty.
-	 */
-	bool bAllowToBeContacted;
-
-	/** Whether the user allowed us to send the log file. */
-	bool bSendLogFile;
-};
-
 
 /**
  * Helper class for MakeDirectoryVisitor
@@ -99,42 +45,4 @@ template <class Functor>
 FunctorDirectoryVisitor<Functor> MakeDirectoryVisitor(Functor&& FunctorInstance)
 {
 	return FunctorDirectoryVisitor<Functor>(MoveTemp(FunctorInstance));
-}
-
-struct FCrashReportUtil
-{
-	/**
-	 * Create a multi line string to display from an exception and callstack
-	 * @param Exception Exception description string
-	 * @param Assertion Assertion description string
-	 * @param Callstack List of callstack entry strings
-	 * @return Multiline text
-	 */
-	static inline FText FormatReportDescription( const FString& Exception, const FString& Assertion, const TArray<FString>& Callstack )
-	{
-		FString Diagnostic = Exception + "\n\n";
-
-		if( !Assertion.IsEmpty() )
-		{
-			TArray<FString> MultilineAssertion;
-			Assertion.ParseIntoArray( MultilineAssertion, TEXT( "#" ), true );
-
-			for( const auto& AssertionLine : MultilineAssertion )
-			{
-				Diagnostic += AssertionLine;
-				Diagnostic += "\n";
-			}
-
-			Diagnostic += "\n";
-		}
-
-		for( const auto& Line : Callstack )
-		{
-			Diagnostic += Line + "\n";
-		}
-		return FText::FromString( Diagnostic );
-	}
-
-	/** Formats processed diagnostic text by adding additional information about machine and user. */
-	static FText FormatDiagnosticText( const FText& DiagnosticText, const FString MachineId, const FString EpicAccountId, const FString UserNameNoDot );
-};
+}
\ No newline at end of file
diff --git a/Engine/Source/Programs/CrashReporter/CrashReportClient/Private/CrashUpload.cpp b/Engine/Source/Programs/CrashReporter/CrashReportClient/Private/CrashUpload.cpp
--- a/Engine/Source/Programs/CrashReporter/CrashReportClient/Private/CrashUpload.cpp
+++ b/Engine/Source/Programs/CrashReporter/CrashReportClient/Private/CrashUpload.cpp
@@ -8,6 +8,8 @@
 #include "PendingReports.h"
 #include "XmlFile.h"
 #include "CrashReportUtil.h"
+#include "GenericPlatformCrashContext.h"
+#include "CrashDescription.h"
 
 #define LOCTEXT_NAMESPACE "CrashReportClient"
 
@@ -70,21 +72,6 @@ const FText& FCrashUpload::GetStatusText() const
 	return UploadStateText;
 }
 
-void FCrashUpload::LocalDiagnosisComplete(const FString& DiagnosticsFile)
-{
-	if (State >= EUploadState::FirstCompletedState)
-	{
-		// Must be a failure/canceled state, or the report was a rejected by the server
-		return;
-	}
-
-	const bool SendDiagnosticsFile = !DiagnosticsFile.IsEmpty();
-	if (SendDiagnosticsFile)
-	{
-		PendingFiles.Push(DiagnosticsFile);
-	}
-}
-
 bool FCrashUpload::IsFinished() const
 {
 	return State >= EUploadState::FirstCompletedState;
@@ -174,12 +161,36 @@ void FCrashUpload::CompressAndSendData()
 
 	int32 CurrentFileIndex = 0;
 
+	const FString FullCrashDumpLocation = FPrimaryCrashProperties::Get()->FullCrashDumpLocation.AsString();
+
 	// Loop to keep trying files until a send succeeds or we run out of files
 	while (PendingFiles.Num() != 0)
 	{
 		const FString PathOfFileToUpload = PendingFiles.Pop();
+		const FString Filename = FPaths::GetCleanFilename( PathOfFileToUpload );
+
+		const bool bValidFullDumpForCopy = Filename == FGenericCrashContext::UE4MinidumpName && 
+			FPrimaryCrashProperties::Get()->CrashDumpMode == ECrashDumpMode::FullDump && 
+			FPrimaryCrashProperties::Get()->CrashVersion >= ECrashDescVersions::VER_3_CrashContext &&
+			!FullCrashDumpLocation.IsEmpty();
+		if (bValidFullDumpForCopy)
+		{
+			const FString DestinationPath = FullCrashDumpLocation / FGenericCrashContext::UE4MinidumpName;
+			const bool bCreated = IFileManager::Get().MakeDirectory( *FullCrashDumpLocation, true );
+			if (!bCreated)
+			{
+				UE_LOG( CrashReportClientLog, Error, TEXT( "Couldn't create directory for full crash dump %s" ), *DestinationPath );
+			}
+			else
+			{
+				UE_LOG( CrashReportClientLog, Warning, TEXT( "Copying full crash minidump to %s" ), *DestinationPath );
+				IFileManager::Get().Copy( *DestinationPath, *PathOfFileToUpload );
+			}
+						
+			continue;
+		}
 		
-		if (FPlatformFileManager::Get().GetPlatformFile().FileSize(*PathOfFileToUpload) > MaxFileSizeToUpload)
+		if (IFileManager::Get().FileSize( *PathOfFileToUpload ) > MaxFileSizeToUpload)
 		{
 			UE_LOG(CrashReportClientLog, Warning, TEXT("Skipping large crash report file"));
 			continue;
@@ -194,18 +205,36 @@ void FCrashUpload::CompressAndSendData()
 		const bool bSkipLogFile = !FCrashReportClientConfig::Get().GetSendLogFile() && PathOfFileToUpload.EndsWith( TEXT( ".log" ) );
 		if (bSkipLogFile)
 		{
-			UE_LOG( CrashReportClientLog, Warning, TEXT( "Skipping the log file" ) );
+			UE_LOG( CrashReportClientLog, Warning, TEXT( "Skipping the %s" ), *Filename );
 			continue;
 		}
 
-		UE_LOG(CrashReportClientLog, Log, TEXT("CompressAndSendData compressing %d bytes ('%s')"), PostData.Num(), *PathOfFileToUpload);
-		FString Filename = FPaths::GetCleanFilename(PathOfFileToUpload);
-		if (Filename == "diagnostics.txt")
+		// Skip old WERInternalMetadata.
+		const bool bSkipXMLFile = PathOfFileToUpload.EndsWith( TEXT( ".xml" ) );
+		if (bSkipXMLFile)
 		{
-			// Ensure diagnostics file is capitalized for server
-			Filename[0] = 'D';
+			UE_LOG( CrashReportClientLog, Warning, TEXT( "Skipping the %s" ), *Filename );
+			continue;
 		}
 
+		// Skip old Report.wer file.
+		const bool bSkipWERFile = PathOfFileToUpload.Contains( TEXT( "Report.wer" ) );
+		if (bSkipWERFile)
+		{
+			UE_LOG( CrashReportClientLog, Warning, TEXT( "Skipping the %s" ), *Filename );
+			continue;
+		}
+
+		// Skip old diagnostics.txt file, all data is stored in the CrashContext.runtime-xml
+		const bool bSkipDiagnostics = Filename == FCrashReportClientConfig::Get().GetDiagnosticsFilename();
+		if (bSkipDiagnostics)
+		{
+			UE_LOG( CrashReportClientLog, Warning, TEXT( "Skipping the %s" ), *Filename );
+			continue;
+		}
+
+		UE_LOG(CrashReportClientLog, Log, TEXT("CompressAndSendData compressing %d bytes ('%s')"), PostData.Num(), *PathOfFileToUpload);
+
 		FCompressedCrashFile FileToCompress( CurrentFileIndex, Filename, PostData );
 		CurrentFileIndex++;
 
diff --git a/Engine/Source/Programs/CrashReporter/CrashReportClient/Private/CrashUpload.h b/Engine/Source/Programs/CrashReporter/CrashReportClient/Private/CrashUpload.h
--- a/Engine/Source/Programs/CrashReporter/CrashReportClient/Private/CrashUpload.h
+++ b/Engine/Source/Programs/CrashReporter/CrashReportClient/Private/CrashUpload.h
@@ -33,19 +33,6 @@ public:
 	 */
 	const FText& GetStatusText() const;
 
-	/**
-	 * Notification that diagnosis of the main report has finished
-	 * @param DiagnosticsFile Full path to the diagnostics file to upload on success, or an empty string
-	 */
-	void LocalDiagnosisComplete(const FString& DiagnosticsFile);
-
-	/**
-	 * Inform uploader that no local diagnosis is happening
-	 */
-	void LocalDiagnosisSkipped()
-	{
-		LocalDiagnosisComplete("");
-	}
 
 	/**
 	 * Determine whether the upload has finished (successfully or otherwise)
diff --git a/Engine/Source/Programs/CrashReporter/CrashReportClient/Private/GenericErrorReport.cpp b/Engine/Source/Programs/CrashReporter/CrashReportClient/Private/GenericErrorReport.cpp
--- a/Engine/Source/Programs/CrashReporter/CrashReportClient/Private/GenericErrorReport.cpp
+++ b/Engine/Source/Programs/CrashReporter/CrashReportClient/Private/GenericErrorReport.cpp
@@ -1,6 +1,7 @@
 // Copyright 1998-2015 Epic Games, Inc. All Rights Reserved.
 
 #include "CrashReportClientApp.h"
+#include "CrashDebugHelperModule.h"
 #include "GenericErrorReport.h"
 #include "XmlFile.h"
 #include "CrashReportUtil.h"
@@ -30,6 +31,7 @@ namespace
 
 FGenericErrorReport::FGenericErrorReport(const FString& Directory)
 	: ReportDirectory(Directory)
+	, bValidCallstack(true)
 {
 	auto FilenamesVisitor = MakeDirectoryVisitor([this](const TCHAR* FilenameOrDirectory, bool bIsDirectory) {
 		if (!bIsDirectory)
@@ -41,53 +43,30 @@ FGenericErrorReport::FGenericErrorReport(const FString& Directory)
 	FPlatformFileManager::Get().GetPlatformFile().IterateDirectory(*ReportDirectory, FilenamesVisitor);
 }
 
-bool FGenericErrorReport::SetUserComment(const FText& UserComment, bool bAllowToBeContacted)
+bool FGenericErrorReport::SetUserComment(const FText& UserComment)
 {
+	const bool bAllowToBeContacted = FCrashReportClientConfig::Get().GetAllowToBeContacted();
+
 	const FString UserName1 = FPlatformProcess::UserName( false );
 	const FString UserName2 = FPlatformProcess::UserName( true );
 	const TCHAR* Anonymous = TEXT( "Anonymous" );
 
+	FPrimaryCrashProperties::Get()->UserDescription = UserComment.ToString();
+
 	// Load the file and remove all PII if bAllowToBeContacted is set to false.
 	const bool bRemovePersonalData = !bAllowToBeContacted;
 	if( bRemovePersonalData )
 	{
-		FString CrashContextFilename;
-		const bool bFound = FindFirstReportFileWithExtension( CrashContextFilename, TEXT( ".runtime-xml" ) );
-		if (bFound)
-		{
-			const FString CrashContextPath = GetReportDirectory() / CrashContextFilename;
-
-			// Cannot use FXMLFile because runtime-xml may contain multi line elements.
-			FString CrashContextString;
-			FFileHelper::LoadFileToString( CrashContextString, *CrashContextPath );
-
-			// Remove UserName and EpicAccountId;
-			TArray<FString> Lines;
-			CrashContextString.ParseIntoArrayLines( Lines, true );
-
-			for( auto& It : Lines )
-			{
-				// Replace user name in assert message, command line etc.
-				It = It.Replace( *UserName1, Anonymous );
-				It = It.Replace( *UserName2, Anonymous );
-
-				if (It.Contains( TEXT( "<UserName>" ) ))
-				{
-					It = TEXT( "<UserName></UserName>" );
-				}
-				else if (It.Contains( TEXT( "<EpicAccountId>" ) ))
-				{
-					It = TEXT( "<EpicAccountId></EpicAccountId>" );
-				}
-			}
-
-			const FString FixedCrashContextString = FString::Join( Lines, TEXT( "\n" ) );
-			FFileHelper::SaveStringToFile( FixedCrashContextString, *CrashContextPath );
-		}
-
-		// #YRX_Crash: 2015-07-01 Move to FGenericCrashContext
+		FPrimaryCrashProperties::Get()->UserName = TEXT( "" );
+		FPrimaryCrashProperties::Get()->EpicAccountId = TEXT( "" );
+		// For now remove the command line completely, to hide the potential personal data. Need to revisit it later.
+		FPrimaryCrashProperties::Get()->CommandLine = TEXT( "CommandLineRemoved" );
 	}
 
+	// Save updated properties, including removed all PII if bAllowToBeContacted is set to false.
+	FPrimaryCrashProperties::Get()->Save();
+
+	// Remove it later, in the next iteration.
 	// Find .xml file
 	FString XmlFilename;
 	if (!FindFirstReportFileWithExtension(XmlFilename, TEXT(".xml")))
@@ -154,13 +133,9 @@ bool FGenericErrorReport::SetUserComment(const FText& UserComment, bool bAllowTo
 	{
 		DynamicSignaturesNode->AppendChildNode(TEXT("Parameter3"), UserComment.ToString());
 	}
-	
-
-	// Set global user name ID: will be added to the report
-	extern FCrashDescription& GetCrashDescription();
 
 	// @see FCrashDescription::UpdateIDs
-	const FString EpicMachineAndUserNameIDs = FString::Printf( TEXT( "!MachineId:%s!EpicAccountId:%s!Name:%s" ), *GetCrashDescription().MachineId, *GetCrashDescription().EpicAccountId, *GetCrashDescription().UserName );
+	const FString EpicMachineAndUserNameIDs = FString::Printf( TEXT( "!MachineId:%s!EpicAccountId:%s!Name:%s" ), *FPrimaryCrashProperties::Get()->MachineId.AsString(), *FPrimaryCrashProperties::Get()->EpicAccountId.AsString(), *FPrimaryCrashProperties::Get()->UserName.AsString() );
 
 	// Add or update a user ID.
 	FXmlNode* Parameter4Node = DynamicSignaturesNode->FindChildNode(TEXT("Parameter4"));
@@ -189,6 +164,26 @@ bool FGenericErrorReport::SetUserComment(const FText& UserComment, bool bAllowTo
 	return XmlFile.Save(XmlFilePath);
 }
 
+void FGenericErrorReport::SetPrimaryCrashProperties( FPrimaryCrashProperties& out_PrimaryCrashProperties )
+{
+	FCrashDebugHelperModule& CrashHelperModule = FModuleManager::LoadModuleChecked<FCrashDebugHelperModule>( FName( "CrashDebugHelper" ) );
+	ICrashDebugHelper* Helper = CrashHelperModule.Get();
+	if (Helper && bValidCallstack)
+	{
+		out_PrimaryCrashProperties.CallStack = Helper->CrashInfo.Exception.CallStackString;
+		out_PrimaryCrashProperties.Modules = Helper->CrashInfo.ModuleNames;
+		out_PrimaryCrashProperties.SourceContext = Helper->CrashInfo.SourceContext;
+
+		// If error message is empty, it means general crash like accessing invalid memory ptr.
+		if (out_PrimaryCrashProperties.ErrorMessage.AsString().Len() == 0)
+		{
+			out_PrimaryCrashProperties.ErrorMessage = Helper->CrashInfo.Exception.ExceptionString;
+		}
+
+		out_PrimaryCrashProperties.Save();
+	}
+}
+
 TArray<FString> FGenericErrorReport::GetFilesToUpload() const
 {
 	TArray<FString> FilesToUpload;
@@ -212,7 +207,7 @@ bool FGenericErrorReport::LoadWindowsReportXmlFile( FString& OutString ) const
 	return FFileHelper::LoadFileToString( OutString, *(ReportDirectory / XmlFilename) );
 }
 
-bool FGenericErrorReport::TryReadDiagnosticsFile(FText& OutReportDescription)
+bool FGenericErrorReport::TryReadDiagnosticsFile()
 {
 	FString FileContent;
 	if (!FFileHelper::LoadFileToString(FileContent, *(ReportDirectory / FCrashReportClientConfig::Get().GetDiagnosticsFilename())))
@@ -236,8 +231,6 @@ bool FGenericErrorReport::TryReadDiagnosticsFile(FText& OutReportDescription)
 	{
 		switch (ReportSection)
 		{
-		default:
-			CRASHREPORTCLIENT_CHECK(false);
 
 		case EReportSection::CallStack:
 			if (Line.StartsWith(CallStackEndKey))
@@ -276,11 +269,19 @@ bool FGenericErrorReport::TryReadDiagnosticsFile(FText& OutReportDescription)
 			break;
 		}
 	}
-	OutReportDescription = FCrashReportUtil::FormatReportDescription( Exception, TEXT( "" ), Callstack );
+
+	// Update properties for the crash.
+	FPrimaryCrashProperties::Get()->CallStack = Callstack;
+	// If error message is empty, it means general crash like accessing invalid memory ptr.
+	if (FPrimaryCrashProperties::Get()->ErrorMessage.AsString().Len() == 0)
+	{
+		FPrimaryCrashProperties::Get()->ErrorMessage = Exception;
+	}
+
 	return true;
 }
 
-bool FGenericErrorReport::FindFirstReportFileWithExtension(FString& OutFilename, const TCHAR* Extension) const
+bool FGenericErrorReport::FindFirstReportFileWithExtension( FString& OutFilename, const TCHAR* Extension ) const
 {
 	for (const auto& Filename: ReportFilenames)
 	{
diff --git a/Engine/Source/Programs/CrashReporter/CrashReportClient/Private/GenericErrorReport.h b/Engine/Source/Programs/CrashReporter/CrashReportClient/Private/GenericErrorReport.h
--- a/Engine/Source/Programs/CrashReporter/CrashReportClient/Private/GenericErrorReport.h
+++ b/Engine/Source/Programs/CrashReporter/CrashReportClient/Private/GenericErrorReport.h
@@ -3,6 +3,7 @@
 #pragma once
 
 class FCrashDebugHelperModule;
+struct FPrimaryCrashProperties;
 
 /**
  * Helper that works with Windows Error Reports
@@ -38,11 +39,14 @@ public:
 	}
 
 	/**
-	 * Write the provided comment into the error report, also remove any PII if bAllowToBeContacted is set to false
+	 * Write the provided comment into the error report
 	 * @param UserComment Information provided by the user to add to the report
 	 * @return Whether the comment was successfully written to the report
 	 */
-	bool SetUserComment(const FText& UserComment, bool bAllowToBeContacted);
+	bool SetUserComment(const FText& UserComment);
+
+	/** Sets properties specific to the processed crash. Used to convert the old data into the crash context format. */
+	void SetPrimaryCrashProperties( FPrimaryCrashProperties& out_PrimaryCrashProperties );
 
 	/**
 	 * Provide full paths to all the report files
@@ -60,10 +64,9 @@ public:
 	bool LoadWindowsReportXmlFile( FString& OutString ) const;
 
 	/**
-	 * @param Description (exception and callstack) to fill in if successful
 	 * @return Whether the file was found and successfully read
 	 */
-	bool TryReadDiagnosticsFile(FText& OutReportDescription);
+	bool TryReadDiagnosticsFile();
 
 	/**
 	 * Provide the full path of the error report directory
@@ -111,7 +114,9 @@ protected:
 	/** Full path to the directory the report files are in */
 	FString ReportDirectory;
 
-private:
 	/** List of leaf filenames of all the files in the report folder */
 	TArray<FString> ReportFilenames;
+
+	/** Whether the error report generated an valid callstack. */
+	mutable bool bValidCallstack;
 };
diff --git a/Engine/Source/Programs/CrashReporter/CrashReportClient/Private/SCrashReportClient.cpp b/Engine/Source/Programs/CrashReporter/CrashReportClient/Private/SCrashReportClient.cpp
--- a/Engine/Source/Programs/CrashReporter/CrashReportClient/Private/SCrashReportClient.cpp
+++ b/Engine/Source/Programs/CrashReporter/CrashReportClient/Private/SCrashReportClient.cpp
@@ -8,6 +8,7 @@
 #include "CrashReportClientStyle.h"
 #include "SlateStyle.h"
 #include "SThrobber.h"
+#include "CrashDescription.h"
 
 #define LOCTEXT_NAMESPACE "CrashReportClient"
 
@@ -33,7 +34,7 @@ void SCrashReportClient::Construct(const FArguments& InArgs, TSharedRef<FCrashRe
 {
 	CrashReportClient = Client;
 
-	auto CrashedAppName = CrashReportClient->GetCrashedAppName();
+	auto CrashedAppName = FPrimaryCrashProperties::Get()->GameName;
 
 	// Set the text displaying the name of the crashed app, if available
 	const FText CrashedAppText = CrashedAppName.IsEmpty() ?
diff --git a/Engine/Source/Programs/CrashReporter/CrashReportClient/Private/Windows/WindowsErrorReport.cpp b/Engine/Source/Programs/CrashReporter/CrashReportClient/Private/Windows/WindowsErrorReport.cpp
--- a/Engine/Source/Programs/CrashReporter/CrashReportClient/Private/Windows/WindowsErrorReport.cpp
+++ b/Engine/Source/Programs/CrashReporter/CrashReportClient/Private/Windows/WindowsErrorReport.cpp
@@ -1,11 +1,8 @@
 // Copyright 1998-2015 Epic Games, Inc. All Rights Reserved.
 
 #include "CrashReportClientApp.h"
-
 #include "WindowsErrorReport.h"
-#include "XmlFile.h"
 #include "CrashDebugHelperModule.h"
-#include "../CrashReportUtil.h"
 
 #include "AllowWindowsPlatformTypes.h"
 #include <ShlObj.h>
@@ -74,6 +71,9 @@ FString FWindowsErrorReport::FindCrashedAppPath() const
 
 FText FWindowsErrorReport::DiagnoseReport() const
 {
+	// Mark the callstack as invalid.
+	bValidCallstack = false;
+
 	// Should check if there are local PDBs before doing anything
 	auto CrashDebugHelper = CrashHelperModule ? CrashHelperModule->Get() : nullptr;
 	if (!CrashDebugHelper)
@@ -96,13 +96,9 @@ FText FWindowsErrorReport::DiagnoseReport() const
 		return LOCTEXT("NoDebuggingSymbols", "You do not have any debugging symbols required to display the callstack for this crash.");
 	}
 
-	// There's a callstack, so write it out to save the server trying to do it
-	CrashDebugHelper->CrashInfo.GenerateReport(ReportDirectory / FCrashReportClientConfig::Get().GetDiagnosticsFilename());
-
-	const auto& Exception = CrashDebugHelper->CrashInfo.Exception;
-	const FString Assertion = FWindowsReportParser::Find( ReportDirectory, TEXT( "AssertLog=" ) );
-
-	return FCrashReportUtil::FormatReportDescription( Exception.ExceptionString, Assertion, Exception.CallStackString );
+	// No longer required, only for backward compatibility, mark the callstack as valid.
+	bValidCallstack = true;
+	return FText();
 }
 
 static bool TryGetDirectoryCreationTime(const FString& InDirectoryName, FDateTime& OutCreationTime)
diff --git a/Engine/Source/Programs/CrashReporter/CrashReportClient/Public/CrashReportClientApp.h b/Engine/Source/Programs/CrashReporter/CrashReportClient/Public/CrashReportClientApp.h
--- a/Engine/Source/Programs/CrashReporter/CrashReportClient/Public/CrashReportClientApp.h
+++ b/Engine/Source/Programs/CrashReporter/CrashReportClient/Public/CrashReportClientApp.h
@@ -11,17 +11,9 @@
 	#include "StandaloneRenderer.h"
 #endif // CRASH_REPORT_UNATTENDED_ONLY
 
-DECLARE_LOG_CATEGORY_EXTERN(CrashReportClientLog, Log, All)
-
-// Helper macros
-#define STRINGIZE_IMPL(STR) #STR
-#define STRINGIZE(STR) STRINGIZE_IMPL(STR)
+#include "CrashReportClientConfig.h"
 
-#define FILE_LINE_STRING TEXT(__FILE__) TEXT("(") TEXT(STRINGIZE(__LINE__)) TEXT(")")
-
-/** This writes and error to the log rather than calling 'check': don't want the crash reporter crashing */
-void CrashReportClientCheck(bool bCondition, const TCHAR* Location);
-#define CRASHREPORTCLIENT_CHECK(COND) CrashReportClientCheck(COND, FILE_LINE_STRING)
+DECLARE_LOG_CATEGORY_EXTERN(CrashReportClientLog, Log, All)
 
 /**
  * Run the crash report client app
