CrashReporter-Frontend - Safer method for adding new crashes

diff --git a/Engine/Source/Programs/CrashReporter/CrashReportWebSite/Controllers/CrashesController.cs b/Engine/Source/Programs/CrashReporter/CrashReportWebSite/Controllers/CrashesController.cs
--- a/Engine/Source/Programs/CrashReporter/CrashReportWebSite/Controllers/CrashesController.cs
+++ b/Engine/Source/Programs/CrashReporter/CrashReportWebSite/Controllers/CrashesController.cs
@@ -8,6 +8,7 @@ using System.Web.Mvc;
 using Tools.DotNETCommon.XmlHandler;
 using Tools.CrashReporter.CrashReportCommon;
 using Tools.CrashReporter.CrashReportWebSite.Models;
+using System.Data.SqlClient;
 
 namespace Tools.CrashReporter.CrashReportWebSite.Controllers
 {
@@ -161,25 +162,49 @@ namespace Tools.CrashReporter.CrashReportWebSite.Controllers
 				CrashReporterResult NewCrashResult = new CrashReporterResult();
 				NewCrashResult.ID = -1;
 
-				try
+				for (int Index = 0; Index < 3; Index++)
 				{
-					using( StreamReader Reader = new StreamReader( Request.InputStream, Request.ContentEncoding ) )
+					try
 					{
-						string Result = Reader.ReadToEnd();
-						CrashDescription NewCrash = XmlHandler.FromXmlString<CrashDescription>( Result );
-						NewCrashResult.ID = Crashes.AddNewCrash( NewCrash );
-						NewCrashResult.bSuccess = true;
+						UnsafeAddCrash( NewCrashResult, Crashes );
+						break;
 					}
-				}
-				catch( Exception Ex )
-				{
-					NewCrashResult.Message = Ex.ToString();
-					NewCrashResult.bSuccess = false;
+					catch (SqlException SqlExc)
+					{
+						if (SqlExc.Number == -2)
+						{
+							FLogger.Global.WriteEvent( string.Format( "AddCrash:Timeout, retrying {0} of 3", Index + 1 ) );
+						}
+						else
+						{
+							NewCrashResult.Message = SqlExc.ToString();
+							NewCrashResult.bSuccess = false;
+							break;
+						}
+					}
+					catch (Exception Ex)
+					{
+						NewCrashResult.Message = Ex.ToString();
+						NewCrashResult.bSuccess = false;
+						break;
+					}
+					System.Threading.Thread.Sleep( 5000 * ( Index + 1 ) );
 				}
 
 				string ReturnResult = XmlHandler.ToXmlString<CrashReporterResult>( NewCrashResult );
 				return Content( ReturnResult, "text/xml" );
 			}
 		}
+
+		private void UnsafeAddCrash( CrashReporterResult NewCrashResult, CrashRepository Crashes )
+		{
+			using (StreamReader Reader = new StreamReader( Request.InputStream, Request.ContentEncoding ))
+			{
+				string Result = Reader.ReadToEnd();
+				CrashDescription NewCrash = XmlHandler.FromXmlString<CrashDescription>( Result );
+				NewCrashResult.ID = Crashes.AddNewCrash( NewCrash );
+				NewCrashResult.bSuccess = true;
+			}
+		}
 	}
 }
