CrashReport-Client - Fixes for Mac

diff --git a/Engine/Source/Programs/CrashReporter/CrashReportClient/Private/CrashDescription.cpp b/Engine/Source/Programs/CrashReporter/CrashReportClient/Private/CrashDescription.cpp
--- a/Engine/Source/Programs/CrashReporter/CrashReportClient/Private/CrashDescription.cpp
+++ b/Engine/Source/Programs/CrashReporter/CrashReportClient/Private/CrashDescription.cpp
@@ -276,7 +276,7 @@ FCrashWERContext::FCrashWERContext( const FString& WERXMLFilepath )
 	: FPrimaryCrashProperties()
 {
 	ReadXML( WERXMLFilepath );
-	CrashGUID = FPaths::GetCleanFilename( WERXMLFilepath );
+	CrashGUID = FPaths::GetCleanFilename( FPaths::GetPath( WERXMLFilepath ) );
 
 	const bool bIsValid = XmlFile->IsValid();
 	if (bIsValid)
diff --git a/Engine/Source/Programs/CrashReporter/CrashReportClient/Private/CrashUpload.cpp b/Engine/Source/Programs/CrashReporter/CrashReportClient/Private/CrashUpload.cpp
--- a/Engine/Source/Programs/CrashReporter/CrashReportClient/Private/CrashUpload.cpp
+++ b/Engine/Source/Programs/CrashReporter/CrashReportClient/Private/CrashUpload.cpp
@@ -209,13 +209,15 @@ void FCrashUpload::CompressAndSendData()
 			continue;
 		}
 
+		// Disabled due to issues with Mac not using the crash context.
+		/*
 		// Skip old WERInternalMetadata.
 		const bool bSkipXMLFile = PathOfFileToUpload.EndsWith( TEXT( ".xml" ) );
 		if (bSkipXMLFile)
 		{
 			UE_LOG( CrashReportClientLog, Warning, TEXT( "Skipping the %s" ), *Filename );
 			continue;
-		}
+		}*/
 
 		// Skip old Report.wer file.
 		const bool bSkipWERFile = PathOfFileToUpload.Contains( TEXT( "Report.wer" ) );
@@ -226,12 +228,14 @@ void FCrashUpload::CompressAndSendData()
 		}
 
 		// Skip old diagnostics.txt file, all data is stored in the CrashContext.runtime-xml
+		// Disabled due to issues with Mac not using the crash context.
+		/*
 		const bool bSkipDiagnostics = Filename == FCrashReportClientConfig::Get().GetDiagnosticsFilename();
 		if (bSkipDiagnostics)
 		{
 			UE_LOG( CrashReportClientLog, Warning, TEXT( "Skipping the %s" ), *Filename );
 			continue;
-		}
+		}*/
 
 		UE_LOG(CrashReportClientLog, Log, TEXT("CompressAndSendData compressing %d bytes ('%s')"), PostData.Num(), *PathOfFileToUpload);
 
