CrashReport-Client - Fixes for Mac

diff --git a/Engine/Source/Programs/CrashReporter/CrashReportClient/Private/CrashUpload.cpp b/Engine/Source/Programs/CrashReporter/CrashReportClient/Private/CrashUpload.cpp
--- a/Engine/Source/Programs/CrashReporter/CrashReportClient/Private/CrashUpload.cpp
+++ b/Engine/Source/Programs/CrashReporter/CrashReportClient/Private/CrashUpload.cpp
@@ -53,7 +53,7 @@ bool FCrashUpload::PingTimeout(float DeltaTime)
 void FCrashUpload::BeginUpload(const FPlatformErrorReport& PlatformErrorReport)
 {
 	ErrorReport = PlatformErrorReport;
-	PendingFiles += ErrorReport.GetFilesToUpload();
+	PendingFiles = FPlatformErrorReport( ErrorReport.GetReportDirectory() ).GetFilesToUpload();
 	UE_LOG(CrashReportClientLog, Log, TEXT("Got %d pending files to upload from '%s'"), PendingFiles.Num(), *ErrorReport.GetReportDirectoryLeafName());
 
 	PauseState = EUploadState::Finished;
