Update buttons on CrashReporter window to encourage submission

diff --git a/Engine/Source/Programs/CrashReporter/CrashReportClient/Private/SCrashReportClient.cpp b/Engine/Source/Programs/CrashReporter/CrashReportClient/Private/SCrashReportClient.cpp
--- a/Engine/Source/Programs/CrashReporter/CrashReportClient/Private/SCrashReportClient.cpp
+++ b/Engine/Source/Programs/CrashReporter/CrashReportClient/Private/SCrashReportClient.cpp
@@ -243,17 +243,7 @@ void SCrashReportClient::Construct(const FArguments& InArgs, TSharedRef<FCrashRe
 			.Padding( FMargin(4, 4+16, 4, 4) )
 			[
 				SNew(SHorizontalBox)
-
-				+SHorizontalBox::Slot()
-				.FillWidth(1.0f)
-				.HAlign(HAlign_Left)
-				.VAlign(VAlign_Center)
-				.Padding(0)
-				[			
-					SNew(SSpacer)
-				]
-
-				
+	
 				+ SHorizontalBox::Slot()
 				.HAlign( HAlign_Center )
 				.VAlign( VAlign_Center )
@@ -266,6 +256,15 @@ void SCrashReportClient::Construct(const FArguments& InArgs, TSharedRef<FCrashRe
 					.OnClicked( Client, &FCrashReportClient::CloseWithoutSending )
 				]
 
+				+SHorizontalBox::Slot()
+				.FillWidth(1.0f)
+				.HAlign(HAlign_Left)
+				.VAlign(VAlign_Center)
+				.Padding(0)
+				[			
+					SNew(SSpacer)
+				]
+
 				+SHorizontalBox::Slot()
 				.HAlign(HAlign_Center)
 				.VAlign(VAlign_Center)
@@ -274,8 +273,8 @@ void SCrashReportClient::Construct(const FArguments& InArgs, TSharedRef<FCrashRe
 				[
 					SNew(SButton)
 					.ContentPadding( FMargin(8,2) )
-					.Text(LOCTEXT("SendAndRestartEditor", "Send and Restart"))
-					.OnClicked(Client, &FCrashReportClient::SubmitAndRestart)
+					.Text(LOCTEXT("Send", "Send and Close"))
+					.OnClicked(Client, &FCrashReportClient::Submit)
 				]
 
 				+SHorizontalBox::Slot()
@@ -286,9 +285,9 @@ void SCrashReportClient::Construct(const FArguments& InArgs, TSharedRef<FCrashRe
 				[
 					SNew(SButton)
 					.ContentPadding( FMargin(8,2) )
-					.Text(LOCTEXT("Send", "Send"))
-					.OnClicked(Client, &FCrashReportClient::Submit)
-				]
+					.Text(LOCTEXT("SendAndRestartEditor", "Send and Restart"))
+					.OnClicked(Client, &FCrashReportClient::SubmitAndRestart)
+				]			
 			]
 		]
 	];
