CrashReport-Client - Crash fix when there is no crash in the queue

diff --git a/Engine/Source/Programs/CrashReporter/CrashReportClient/Private/CrashDescription.cpp b/Engine/Source/Programs/CrashReporter/CrashReportClient/Private/CrashDescription.cpp
--- a/Engine/Source/Programs/CrashReporter/CrashReportClient/Private/CrashDescription.cpp
+++ b/Engine/Source/Programs/CrashReporter/CrashReportClient/Private/CrashDescription.cpp
@@ -105,6 +105,7 @@ FPrimaryCrashProperties::FPrimaryCrashProperties()
 	, FullCrashDumpLocation( FGenericCrashContext::RuntimePropertiesTag, TEXT( "FullCrashDumpLocation" ), this )
 	, TimeOfCrash( FGenericCrashContext::RuntimePropertiesTag, TEXT( "TimeOfCrash" ), this )
 	, bAllowToBeContacted( FGenericCrashContext::RuntimePropertiesTag, TEXT( "bAllowToBeContacted" ), this )
+	, XmlFile( nullptr )
 {
 	CrashVersion = ECrashDescVersions::VER_1_NewCrashFormat;
 	CrashDumpMode = ECrashDumpMode::Default;
diff --git a/Engine/Source/Programs/CrashReporter/CrashReportClient/Private/CrashDescription.h b/Engine/Source/Programs/CrashReporter/CrashReportClient/Private/CrashDescription.h
--- a/Engine/Source/Programs/CrashReporter/CrashReportClient/Private/CrashDescription.h
+++ b/Engine/Source/Programs/CrashReporter/CrashReportClient/Private/CrashDescription.h
@@ -320,7 +320,6 @@ protected:
 	/** Default constructor. */
 	FPrimaryCrashProperties();
 
-
 	/** Destructor. */
 	~FPrimaryCrashProperties()
 	{
@@ -342,6 +341,14 @@ public:
 		return Singleton;
 	}
 
+	/**
+	* @return false, if there is no crash
+	*/
+	static bool IsValid()
+	{
+		return Singleton != nullptr;
+	}
+
 	/** Shutdowns the global instance. */
 	static void Shutdown();
 
diff --git a/Engine/Source/Programs/CrashReporter/CrashReportClient/Private/CrashReportClient.cpp b/Engine/Source/Programs/CrashReporter/CrashReportClient/Private/CrashReportClient.cpp
--- a/Engine/Source/Programs/CrashReporter/CrashReportClient/Private/CrashReportClient.cpp
+++ b/Engine/Source/Programs/CrashReporter/CrashReportClient/Private/CrashReportClient.cpp
@@ -34,31 +34,34 @@ FCrashReportClient::FCrashReportClient(const FPlatformErrorReport& InErrorReport
 	, bShouldWindowBeHidden(false)
 	, bSendData(false)
 {
-	bool bUsePrimaryData = false;
-	if (FPrimaryCrashProperties::Get()->HasProcessedData())
+	if (FPrimaryCrashProperties::Get()->IsValid())
 	{
-		bUsePrimaryData = true;
-	}
-	else
-	{
-		if (!ErrorReport.TryReadDiagnosticsFile() && !FParse::Param( FCommandLine::Get(), TEXT( "no-local-diagnosis" ) ))
+		bool bUsePrimaryData = false;
+		if (FPrimaryCrashProperties::Get()->HasProcessedData())
 		{
-			DiagnoseReportTask = new FAsyncTask<FDiagnoseReportWorker>( this );
-			DiagnoseReportTask->StartBackgroundTask();
+			bUsePrimaryData = true;
 		}
 		else
 		{
-			bUsePrimaryData = true;
+			if (!ErrorReport.TryReadDiagnosticsFile() && !FParse::Param( FCommandLine::Get(), TEXT( "no-local-diagnosis" ) ))
+			{
+				DiagnoseReportTask = new FAsyncTask<FDiagnoseReportWorker>( this );
+				DiagnoseReportTask->StartBackgroundTask();
+			}
+			else
+			{
+				bUsePrimaryData = true;
+			}
 		}
-	}
 
-	if (bUsePrimaryData)
-	{
-		const FString CallstackString = FPrimaryCrashProperties::Get()->CallStack.AsString();
-		const FString ReportString = FString::Printf( TEXT( "%s\n\n%s" ), *FPrimaryCrashProperties::Get()->ErrorMessage.AsString(), *CallstackString );
-		DiagnosticText = FText::FromString( ReportString );
+		if (bUsePrimaryData)
+		{
+			const FString CallstackString = FPrimaryCrashProperties::Get()->CallStack.AsString();
+			const FString ReportString = FString::Printf( TEXT( "%s\n\n%s" ), *FPrimaryCrashProperties::Get()->ErrorMessage.AsString(), *CallstackString );
+			DiagnosticText = FText::FromString( ReportString );
 
-		FormattedDiagnosticText = FCrashReportUtil::FormatDiagnosticText( FText::FromString( ReportString ) );
+			FormattedDiagnosticText = FCrashReportUtil::FormatDiagnosticText( FText::FromString( ReportString ) );
+		}
 	}
 }
 
diff --git a/Engine/Source/Programs/CrashReporter/CrashReportClient/Private/SCrashReportClient.cpp b/Engine/Source/Programs/CrashReporter/CrashReportClient/Private/SCrashReportClient.cpp
--- a/Engine/Source/Programs/CrashReporter/CrashReportClient/Private/SCrashReportClient.cpp
+++ b/Engine/Source/Programs/CrashReporter/CrashReportClient/Private/SCrashReportClient.cpp
@@ -34,12 +34,12 @@ void SCrashReportClient::Construct(const FArguments& InArgs, TSharedRef<FCrashRe
 {
 	CrashReportClient = Client;
 
-	auto CrashedAppName = FPrimaryCrashProperties::Get()->GameName;
+	auto CrashedAppName = FPrimaryCrashProperties::Get()->IsValid() ? FPrimaryCrashProperties::Get()->GameName : TEXT("");
 
 	// Set the text displaying the name of the crashed app, if available
 	const FText CrashedAppText = CrashedAppName.IsEmpty() ?
-		LOCTEXT("CrashedAppNotFound", "An Unreal process has crashed") :
-		LOCTEXT("CrashedApp", "The following process has crashed: ");
+		LOCTEXT( "CrashedAppNotFound", "An unknown process has crashed" ) :
+		LOCTEXT( "CrashedAppUnreal", "An Unreal process has crashed: " );
 
 	const FText CrashReportDataText = FText::Format( LOCTEXT(
 		"CrashReportData",
@@ -275,6 +275,7 @@ void SCrashReportClient::Construct(const FArguments& InArgs, TSharedRef<FCrashRe
 					.ContentPadding( FMargin(8,2) )
 					.Text(LOCTEXT("Send", "Send and Close"))
 					.OnClicked(Client, &FCrashReportClient::Submit)
+					.IsEnabled( !CrashedAppName.IsEmpty() )
 				]
 
 				+SHorizontalBox::Slot()
@@ -287,6 +288,7 @@ void SCrashReportClient::Construct(const FArguments& InArgs, TSharedRef<FCrashRe
 					.ContentPadding( FMargin(8,2) )
 					.Text(LOCTEXT("SendAndRestartEditor", "Send and Restart"))
 					.OnClicked(Client, &FCrashReportClient::SubmitAndRestart)
+					.IsEnabled( !CrashedAppName.IsEmpty() )
 				]			
 			]
 		]
