CrashReport-Client - Fixed issue with full crash dump configuration

diff --git a/Engine/Source/Programs/CrashReporter/CrashReportClient/Private/CrashReportClientConfig.cpp b/Engine/Source/Programs/CrashReporter/CrashReportClient/Private/CrashReportClientConfig.cpp
--- a/Engine/Source/Programs/CrashReporter/CrashReportClient/Private/CrashReportClientConfig.cpp
+++ b/Engine/Source/Programs/CrashReporter/CrashReportClient/Private/CrashReportClientConfig.cpp
@@ -63,7 +63,7 @@ const FString FCrashReportClientConfig::GetFullCrashDumpLocationForBranch( const
 		{
 			return It.Location;
 		}
-		else if( BranchName.Contains( It.BranchName ) )
+		else if (!bExactMatch && BranchName.Contains( It.BranchName ))
 		{
 			return It.Location;
 		}
