CrashReporter-Frontend - Fixes for JIRA changes, better support for missing/changed JIRAs

diff --git a/Engine/Source/Programs/CrashReporter/CrashReportWebSite/Controllers/BuggsController.cs b/Engine/Source/Programs/CrashReporter/CrashReportWebSite/Controllers/BuggsController.cs
--- a/Engine/Source/Programs/CrashReporter/CrashReportWebSite/Controllers/BuggsController.cs
+++ b/Engine/Source/Programs/CrashReporter/CrashReportWebSite/Controllers/BuggsController.cs
@@ -163,8 +163,11 @@ namespace Tools.CrashReporter.CrashReportWebSite.Controllers
 					{
 						// Grab the data form JIRA.
 						string JiraSearchQuery = "key = " + NewBugg.Jira;
+						var JiraResults = new Dictionary<string, Dictionary<string, object>>();
 
-						var JiraResults = JC.SearchJiraTickets(
+						try
+						{
+							JiraResults = JC.SearchJiraTickets(
 							JiraSearchQuery,
 							new string[] 
 							{ 
@@ -176,6 +179,15 @@ namespace Tools.CrashReporter.CrashReportWebSite.Controllers
 								"customfield_11200" // string
 							} );
 
+						}
+						catch (System.Exception)
+						{
+							NewBugg.JiraSummary = "JIRA MISMATCH";
+							NewBugg.JiraComponentsText = "JIRA MISMATCH";
+							NewBugg.JiraResolution = "JIRA MISMATCH";
+							NewBugg.JiraFixVersionsText = "JIRA MISMATCH";
+							NewBugg.JiraFixCL = "JIRA MISMATCH";
+						}
 
 						// Jira Key, Summary, Components, Resolution, Fix version, Fix changelist
 						foreach( var Jira in JiraResults )
diff --git a/Engine/Source/Programs/CrashReporter/CrashReportWebSite/Controllers/ReportsController.cs b/Engine/Source/Programs/CrashReporter/CrashReportWebSite/Controllers/ReportsController.cs
--- a/Engine/Source/Programs/CrashReporter/CrashReportWebSite/Controllers/ReportsController.cs
+++ b/Engine/Source/Programs/CrashReporter/CrashReportWebSite/Controllers/ReportsController.cs
@@ -238,6 +238,7 @@ namespace Tools.CrashReporter.CrashReportWebSite.Controllers
 
 					using( FAutoScopedLogTimer JiraResultsTimer = new FAutoScopedLogTimer( "JiraResults" ) )
 					{
+						bool bInvalid = false;
 						var JiraResults = new Dictionary<string, Dictionary<string, object>>();
 						try
 						{
@@ -255,9 +256,39 @@ namespace Tools.CrashReporter.CrashReportWebSite.Controllers
 						}
 						catch (System.Exception)
 						{
-
+							bInvalid = true;
 						}
 
+						// Invalid records have been found, find the broken using the slow path.
+						if( bInvalid )
+						{
+							foreach (var Query in FoundJiras)
+							{
+								try
+								{
+									var TempResult = JC.SearchJiraTickets(
+									Query,
+									new string[] 
+									{ 
+										"key",				// string
+										"summary",			// string
+										"components",		// System.Collections.ArrayList, Dictionary<string,object>, name
+										"resolution",		// System.Collections.Generic.Dictionary`2[System.String,System.Object], name
+										"fixVersions",		// System.Collections.ArrayList, Dictionary<string,object>, name
+										"customfield_11200" // string
+									} );
+
+									foreach(var Temp in TempResult)
+									{
+										JiraResults.Add( Temp.Key, Temp.Value );
+									}
+								}
+								catch (System.Exception)
+								{
+
+								}
+							}
+						}
 
 						// Jira Key, Summary, Components, Resolution, Fix version, Fix changelist
 						foreach( var Jira in JiraResults )
diff --git a/Engine/Source/Programs/CrashReporter/CrashReportWebSite/Models/CrashReport.cs b/Engine/Source/Programs/CrashReporter/CrashReportWebSite/Models/CrashReport.cs
--- a/Engine/Source/Programs/CrashReporter/CrashReportWebSite/Models/CrashReport.cs
+++ b/Engine/Source/Programs/CrashReporter/CrashReportWebSite/Models/CrashReport.cs
@@ -125,7 +125,7 @@ namespace Tools.CrashReporter.CrashReportWebSite.Models
 		string ToJiraSummary = "";
 
 		/// <summary>Branches in JIRA</summary>
-		List<object> ToJiraBranches = null;
+		HashSet<string> ToJiraBranches = null;
 
 		/// <summary>Versions in JIRA</summary>
 		List<object> ToJiraVersions = null;
@@ -268,26 +268,23 @@ namespace Tools.CrashReporter.CrashReportWebSite.Models
 			}
 
 			// ToJiraBranches
-			this.ToJiraBranches = new List<object>();
-			foreach( var Platform in this.BranchesFoundIn )
+			this.ToJiraBranches = new HashSet<string>();
+			foreach( var BranchName in this.BranchesFoundIn )
 			{
-				string CleanedBranch = Platform.Contains( "UE4-Releases" ) ? "UE4-Releases" : Platform;
-				Dictionary<string, object> JiraBranch = null;
-				JC.GetNameToBranchFoundIn().TryGetValue( CleanedBranch, out JiraBranch );
-				if( JiraBranch != null && !this.ToJiraBranches.Contains( JiraBranch ) )
+				if( !string.IsNullOrEmpty(BranchName) )
 				{
-					this.ToJiraBranches.Add( JiraBranch );
+					this.ToJiraBranches.Add( CrashReporterConstants.P4_DEPOT_PREFIX + BranchName );
 				}
 			}
 
 			// ToJiraPlatforms
 			this.ToJiraPlatforms = new List<object>();
-			foreach( var Platform in this.AffectedPlatforms )
+			foreach( var BranchName in this.AffectedPlatforms )
 			{
-				bool bValid = JC.GetNameToPlatform().ContainsKey( Platform );
+				bool bValid = JC.GetNameToPlatform().ContainsKey( BranchName );
 				if( bValid )
 				{
-					this.ToJiraPlatforms.Add( JC.GetNameToPlatform()[Platform] );
+					this.ToJiraPlatforms.Add( JC.GetNameToPlatform()[BranchName] );
 				}
 			}
 		}
@@ -317,13 +314,13 @@ namespace Tools.CrashReporter.CrashReportWebSite.Models
 				Fields.Add( "versions", ToJiraVersions );
 
 				// ToJiraBranches
-				Fields.Add( "customfield_11201", ToJiraBranches );
+				Fields.Add( "customfield_12402", ToJiraBranches.ToList() );						// NewBranchFoundIn
 
 				// ToJiraPlatforms
 				Fields.Add( "customfield_11203", ToJiraPlatforms );
 
 				// Callstack customfield_11807
-				string JiraCallstack = string.Join( "\r\n", ToJiraFunctionCalls );
+				string JiraCallstack = "{noformat}" + string.Join( "\r\n", ToJiraFunctionCalls ) + "{noformat}";
 				Fields.Add( "customfield_11807", JiraCallstack );								// Callstack
 
 				string BuggLink = "http://crashreporter/Buggs/Show/" + Id;
@@ -359,7 +356,11 @@ namespace Tools.CrashReporter.CrashReportWebSite.Models
 			Tooltip += "JiraVersions: " + JiraVersions + NL;
 
 			// "value"
-			string JiraBranches = GetFieldsFrom( ToJiraBranches, "value" );
+			string JiraBranches = "";
+			foreach (var Branch in ToJiraBranches)
+			{
+				JiraBranches += Branch + ", ";
+			}
 			Tooltip += "JiraBranches: " + JiraBranches + NL;
 
 			// "value"
