Copying //UE4/Dev-Core to //UE4/Main

diff --git a/Engine/Source/Programs/CrashReporter/CrashReportCommon/CrashContext.cs b/Engine/Source/Programs/CrashReporter/CrashReportCommon/CrashContext.cs
new file mode 100644
--- /dev/null
+++ b/Engine/Source/Programs/CrashReporter/CrashReportCommon/CrashContext.cs
@@ -0,0 +1,622 @@
+// Copyright 1998-2015 Epic Games, Inc. All Rights Reserved.
+
+using System;
+using System.IO;
+using System.Xml.Serialization;
+using System.Collections;
+using System.Collections.Generic;
+using Tools.DotNETCommon.XmlHandler;
+using System.Xml;
+using System.Linq;
+using System.Security;
+using System.Diagnostics;
+
+namespace Tools.CrashReporter.CrashReportCommon
+{
+	/*
+	 Unused Crashes' fields.
+	 
+			Title				nchar(20)		
+			Selected			bit				
+			Version				int				
+			AutoReporterID		int				
+			Processed			bit				
+			HasDiagnosticsFile	bit				
+			HasNewLogFile		bit				
+			HasMetaData			bit				
+
+	 Fields of Crashes used only by the crash report website. 
+		
+			Status				varchar(64)			- Status of the crash.
+			Pattern				varchar(800)		- Pattern used to group cashes into buggs.
+			FixedChangeList		varchar(64)			- Changelist that fixed this crash.
+			Module				varchar(128)		- Name of the module where the crash occurred.
+			UserNameId			int					- The ID of the user name.
+			Jira				varchar(64)			- The Jira.
+			CrashType			smallint NOT NULL	- The type of the crash, crash(1), assert(2), ensure(3)
+	*/
+
+	/* PrimaryCrashProperties. Extracted from: FGenericCrashContext::SerializeContentToBuffer */
+	/*
+		"CrashVersion"
+		"ProcessId"
+		"CrashGUID"
+		"IsInternalBuild"
+		"IsPerforceBuild"
+		"IsSourceDistribution"
+		"SecondsSinceStart"
+		"GameName"
+		"ExecutableName"
+		"BuildConfiguration"
+		"PlatformName"
+		"PlatformNameIni"
+		"PlatformFullName"
+		"EngineMode"
+		"EngineVersion"
+		"CommandLine"
+		"LanguageLCID"
+		"AppDefaultLocale"
+		"IsUE4Release"
+		"UserName"
+		"BaseDir"
+		"RootDir"
+		"MachineId"
+		"EpicAccountId"
+		"CallStack"
+		"SourceContext"
+		"UserDescription"
+		"ErrorMessage"
+		"CrashDumpMode"
+		"Misc.NumberOfCores"
+		"Misc.NumberOfCoresIncludingHyperthreads"
+		"Misc.Is64bitOperatingSystem"
+		"Misc.CPUVendor"
+		"Misc.CPUBrand"
+		"Misc.PrimaryGPUBrand"
+		"Misc.OSVersionMajor"
+		"Misc.OSVersionMinor"
+		"Misc.AppDiskTotalNumberOfBytes"
+		"Misc.AppDiskNumberOfFreeBytes"
+		"MemoryStats.TotalPhysical"
+		"MemoryStats.TotalVirtual"
+		"MemoryStats.PageSize"
+		"MemoryStats.TotalPhysicalGB"
+		"MemoryStats.AvailablePhysical"
+		"MemoryStats.AvailableVirtual"
+		"MemoryStats.UsedPhysical"
+		"MemoryStats.PeakUsedPhysical"
+		"MemoryStats.UsedVirtual"
+		"MemoryStats.PeakUsedVirtual"
+		"MemoryStats.bIsOOM"
+		"MemoryStats.OOMAllocationSize"
+		"MemoryStats.OOMAllocationAlignment"
+		"TimeofCrash"
+		"bAllowToBeContacted"
+	*/
+
+	/// <summary> Enumerates crash description versions. </summary>
+	public enum ECrashDescVersions
+	{
+		/// <summary> Introduces a new crash description format. </summary>
+		VER_1_NewCrashFormat,
+
+		/// <summary> Added misc properties (CPU,GPU,OS,etc), memory related stats and platform specific properties as generic payload. </summary>
+		VER_2_AddedNewProperties,
+
+		/// <summary> Using crash context when available. </summary>
+		VER_3_CrashContext = 3,
+	}
+
+	/// <summary> 
+	/// GenericCrashContext for managed code. Based on XML file.
+	/// Holds data from the CrashContext.runtime-xml file
+	/// Must match FGenericCrashContext::SerializeContentToBuffer
+	/// 
+	/// 1. Generated by FGenericCrashContext::SerializeContentToBuffer via the native code
+	/// 2. Updated with callstack, source context etc via the crash report client / crash report process
+	/// 3. Read by the crash report process, verified and forward in the same form to the crash report frontend
+	/// 4. Basic properties displayed directly in the frontend (already existing in the database), remaining properties are displayed per crash read from the CrashContext.runtime-xml file.
+	/// 
+	/// <RuntimeCrashDescription></RuntimeCrashDescription>
+	/// <version></version>
+	/// <PlatformSpecificProperties>
+	/// 
+	/// </PlatformSpecificProperties>
+	/// </summary>
+	public class FGenericCrashContext
+	{
+		/// <summary> Name of the crash context file. </summary>
+		public const string CrashContextRuntimeXMLName = "CrashContext.runtime-xml";
+
+		/// <summary>Required properties</summary>
+		[XmlElement( ElementName = "RuntimeProperties" )]
+		public FPrimaryCrashProperties PrimaryCrashProperties = new FPrimaryCrashProperties();
+
+		/// <summary>Optional, platform specific properties</summary>
+		[XmlElement( ElementName = "PlatformProperties" )]
+		public object UntypedXMLNodes;
+
+		/// <summary> Directory of the crash. </summary>
+		public string CrashDirectory;
+
+		//// <summary> Platform specific properties. </summary>
+		//[XmlElement]
+		//public List<KeyValuePair<string, string>> PlatformProperties = new List<KeyValuePair<string, string>>();
+
+		/// <summary>A simple default constructor to allow Xml serialization.</summary>
+		public FGenericCrashContext()
+		{
+		}
+
+		/// <summary> Creates a new instance of the crash report context from the specified file. </summary>
+		public static FGenericCrashContext FromFile( string Filepath )
+		{
+			// The crash report context contains invalid characters, we need to fix them all in order to parse the XML.
+			HashSet<string> TagsToIgnore = new HashSet<string>() { "<FGenericCrashContext>", "<RuntimeProperties>", "<PlatformProperties>", "</FGenericCrashContext>", "</RuntimeProperties>", "</PlatformProperties>" };
+
+			string[] XMLText = File.ReadAllLines( Filepath );
+			string CopyText = string.Join( Environment.NewLine, XMLText );
+
+			string NodeOpen = "";
+			string NodeEnd = "";
+
+			List<string> Multiline = new List<string>();
+			List<int> MultiIndex = new List<int>();
+
+			int OpenEnd = 0;
+			int EndStart = 0;
+			for (int LineIndex = 1; LineIndex < XMLText.Length; LineIndex++)
+			{
+				string XMLLine = XMLText[LineIndex].Replace( "\t", "" );
+				bool bIgnore = TagsToIgnore.Contains( XMLLine );
+				if (bIgnore)
+				{
+					continue;
+				}
+
+				// XML node without data, skip.
+				if (XMLLine.Contains( "/>" ))
+				{
+					continue;
+				}
+
+				if (NodeOpen == "")
+				{
+					// Parse the tag.
+					//<Version>"1"</Version>
+					//<Version />
+					int OpenStart = XMLLine.IndexOf( '<' );
+					OpenEnd = XMLLine.IndexOf( '>', OpenStart );
+					NodeOpen = XMLLine.Substring( OpenStart + 1, OpenEnd - OpenStart - 1 );
+				}
+
+				EndStart = XMLLine.IndexOf( "</" );
+				int EndEnd = 0;
+				if (EndStart != -1)
+				{
+					EndEnd = XMLLine.IndexOf( '>', EndStart );
+					NodeEnd = XMLLine.Substring( EndStart + 2, EndEnd - EndStart - 2 );
+				}
+				else
+				{
+					Multiline.Add( XMLLine.Substring( OpenEnd != 0 ? OpenEnd + 1 : 0 ) );
+					MultiIndex.Add( LineIndex );
+					OpenEnd = 0;
+				}
+
+				// Valid tag, fix invalid characters.
+				if (NodeOpen == NodeEnd)
+				{
+					if (Multiline.Count == 0)
+					{
+						string NodeValue = XMLLine.Substring( OpenEnd + 1, EndStart - OpenEnd - 1 );
+						string EscapedXMLText = EscapeXMLString( UnescapeXMLString( NodeValue ) ); // Required for support old and new type format.
+						XMLText[LineIndex] = string.Format( "<{0}>{1}</{0}>", NodeOpen, EscapedXMLText );
+					}
+					else
+					{
+						Multiline.Add( XMLLine.Substring( 0, EndStart ) );
+						MultiIndex.Add( LineIndex );
+
+						for (int Inner = 0; Inner < Multiline.Count; Inner++)
+						{
+							string EscapedXMLText = EscapeXMLString( UnescapeXMLString( Multiline[Inner] ) );
+
+							if (Inner == 0)
+							{
+								XMLText[MultiIndex[Inner]] = string.Format( "<{0}>{1}", NodeOpen, EscapedXMLText );
+							}
+							else if (Inner == Multiline.Count - 1)
+							{
+								XMLText[MultiIndex[Inner]] = string.Format( "{1}</{0}>", NodeOpen, EscapedXMLText );
+							}
+							else
+							{
+								XMLText[MultiIndex[Inner]] = EscapedXMLText;
+							}
+						}
+
+						Multiline.Clear();
+						MultiIndex.Clear();
+					}
+
+					NodeOpen = NodeEnd = "";
+				}
+			}
+
+			string XMLTextJoined = string.Join( Environment.NewLine, XMLText );
+			FGenericCrashContext Context = XmlHandler.FromXmlString<FGenericCrashContext>( XMLTextJoined );
+
+			if (Context.UntypedXMLNodes != null)
+			{
+				var Typed = ( (IEnumerable)Context.UntypedXMLNodes )
+								.Cast<XmlNode>()
+								.Select( X => new KeyValuePair<string, string>( X.Name, X.InnerText ) )
+								.ToList();
+			}
+
+			return Context;
+		}
+
+		const string NewLineTag = "&nl;";
+
+		/// <summary> Escape XML string. </summary>
+		static public string EscapeXMLString( string Text )
+		{
+			return Text
+				.Replace( "&", "&amp;" )
+				.Replace( "\"", "&quot;" )
+				.Replace( "'", "&apos;" )
+				.Replace( "<", "&lt;" )
+				.Replace( ">", "&gt;" )
+				// Ignore return carriage.
+				.Replace( "\r", "" );
+				//.Replace( "\n", NewLineTag );
+		}
+
+		/// <summary> Unescape XML string. </summary>
+		static public string UnescapeXMLString( string Text )
+		{
+			return Text
+				.Replace( "&amp;", "&" )
+				.Replace( "&quot;", "\"" )
+				.Replace( "&apos;", "'" )
+				.Replace( "&lt;", "<" )
+				.Replace( "&gt;", ">" )
+				.Replace( NewLineTag, "\n" );
+		}
+
+		/// <summary> Whether this crash contains callstack, error message and source context thus it means that crash is complete. </summary>
+		public bool HasProcessedData()
+		{
+			return !string.IsNullOrEmpty(PrimaryCrashProperties.CallStack) &&
+					!string.IsNullOrEmpty(PrimaryCrashProperties.ErrorMessage);
+		}
+
+		/// <summary> Writes this instance of the crash report context to the specified file. </summary>
+		public void ToFile()
+		{
+			XmlHandler.WriteXml<FGenericCrashContext>( this, Path.Combine( CrashDirectory, FGenericCrashContext.CrashContextRuntimeXMLName ) );
+		}
+
+		/// <returns>Unique name of this report, can be used as name of file.</returns>
+		public string GetAsFilename()
+		{
+			string EngineVersion = PrimaryCrashProperties.EngineVersion;
+
+			string Directory = string.Format( "{0}_{1}__{2}", EngineVersion, !string.IsNullOrEmpty( PrimaryCrashProperties.GameName ) ? PrimaryCrashProperties.GameName : "NO_GAMENAME", PrimaryCrashProperties.TimeofCrash.Ticks );
+			return Directory;
+		}
+	}
+
+	/*	
+		"EngineVersion"
+	 */
+
+	/// <summary> Required properties for the generic crash context. </summary>
+	public class FPrimaryCrashProperties
+	{
+		/// <summary></summary>
+		[XmlElement]
+		public int CrashVersion;
+
+		/// <summary></summary>
+		[XmlElement]
+		public string CrashGUID;
+
+		/// <summary></summary>
+		[XmlElement]
+		public int ProcessId;
+
+		/// <summary></summary>
+		[XmlElement]
+		public bool IsInternalBuild;
+
+		/// <summary></summary>
+		[XmlElement]
+		public bool IsPerforceBuild;
+
+		/// <summary></summary>
+		[XmlElement]
+		public bool IsSourceDistribution;
+
+		/// <summary></summary>
+		[XmlElement]
+		public int SecondsSinceStart;
+
+		/// <summary></summary>
+		[XmlElement]
+		public string GameName;
+
+		/// <summary></summary>
+		[XmlElement]
+		public string ExecutableName;
+
+		/// <summary></summary>
+		[XmlElement]
+		public string BuildConfiguration;
+
+		/// <summary></summary>
+		[XmlElement]
+		public string PlatformName;
+
+		/// <summary></summary>
+		[XmlElement]
+		public string PlatformNameIni;
+
+		/// <summary></summary>
+		[XmlElement]
+		public string PlatformFullName;
+
+		/// <summary> 
+		/// Creates a full platform name, if missing. PlatformName [OS.Major Os.Minor Arch]
+		/// Must match 
+		/// ...\Engine\Source\Programs\CrashReporter\CrashReportClient\Private\CrashDescription.cpp FCrashContext.FCrashContext 
+		///</summary>
+		public void SetPlatformFullName()
+		{
+			if (string.IsNullOrEmpty( PlatformFullName ))
+			{
+				// Create a new one.
+				string LocalPlatformName = "";
+				if (BaseDir != "")
+				{
+					string[] FolderComponents = BaseDir.Split( new[] { '/' }, StringSplitOptions.RemoveEmptyEntries );
+					if (FolderComponents.Length != 0)
+					{
+						LocalPlatformName = FolderComponents[FolderComponents.Length - 1];
+					}
+				}
+				else
+				{
+					LocalPlatformName = PlatformNameIni;
+				}
+
+				if (!string.IsNullOrEmpty( Misc_OSVersionMajor ))
+				{
+					PlatformFullName = string.Format( "{0} [{1} {2} {3}]", 
+						LocalPlatformName,
+						Misc_OSVersionMajor,
+						string.IsNullOrEmpty( Misc_OSVersionMinor ) ? "" : Misc_OSVersionMinor,
+						Misc_Is64bitOperatingSystem ? "64b" : "32b" );
+				}
+			}
+		}
+
+		/// <summary></summary>
+		[XmlElement]
+		public string EngineMode;
+
+		/// <summary></summary>
+		[XmlElement]
+		public string EngineVersion;
+
+		/// <summary></summary>
+		[XmlElement]
+		public string CommandLine;
+
+		/// <summary></summary>
+		[XmlElement]
+		public int LanguageLCID;
+
+		/// <summary></summary>
+		[XmlElement]
+		public string AppDefaultLocale;
+
+		/// <summary></summary>
+		[XmlElement]
+		public bool IsUE4Release;
+
+		/// <summary></summary>
+		[XmlElement]
+		public string UserName;
+
+		/// <summary></summary>
+		[XmlElement]
+		public string BaseDir;
+
+		/// <summary></summary>
+		[XmlElement]
+		public string RootDir;
+
+		/// <summary></summary>
+		[XmlElement]
+		public string MachineId;
+
+		/// <summary></summary>
+		[XmlElement]
+		public string EpicAccountId;
+
+		/// <summary></summary>
+		[XmlElement]
+		public string CallStack;
+
+		/// <summary> Callstack as string[] unescaped. </summary>
+		public string[] GetCallstack()
+		{
+			string UnescapedValue = FGenericCrashContext.UnescapeXMLString( CallStack );
+			return UnescapedValue.Split( new char[] { '\n' }, StringSplitOptions.RemoveEmptyEntries );
+		}
+
+		/// <summary></summary>
+		[XmlElement]
+		public string SourceContext;
+
+		/// <summary> SourceContext as string[] unescaped. </summary>
+		public string[] GetSourceContext()
+		{
+			string UnescapedValue = FGenericCrashContext.UnescapeXMLString( SourceContext );
+			return UnescapedValue.Split( new char[] { '\n' }, StringSplitOptions.RemoveEmptyEntries );
+		}
+
+		/// <summary> An array of module's name used by the game that crashed. </summary>
+		[XmlElement]
+		public string Modules;
+
+		/// <summary></summary>
+		[XmlElement]
+		public string UserDescription;
+
+		/// <summary> UserDescription as string[] unescaped. </summary>
+		public string[] GetUserDescription()
+		{
+			string UnescapedValue = FGenericCrashContext.UnescapeXMLString( UserDescription );
+			return UnescapedValue.Split( new char[] { '\n' }, StringSplitOptions.RemoveEmptyEntries );
+		}
+
+		/// <summary></summary>
+		[XmlElement]
+		public string ErrorMessage;
+
+		/// <summary> ErrorMessage as string[] unescaped. </summary>
+		public string[] GetErrorMessage()
+		{
+			string UnescapedValue = FGenericCrashContext.UnescapeXMLString( ErrorMessage );
+			UnescapedValue = UnescapedValue.Substring( 0, Math.Min( 511, UnescapedValue.Length ) ); // Database limitation.
+			return UnescapedValue.Split( new char[] { '\n' }, StringSplitOptions.RemoveEmptyEntries );
+		}
+
+		/// <summary> Location of the full crash dump. Displayed in the crash report frontend, without the filename. </summary>
+		[XmlElement]
+		public string FullCrashDumpLocation;
+
+		/// <summary> Fullname of the full crash dump.</summary>
+		public string GetFullCrashDumpLocation()
+		{
+			string Fullname = Path.Combine( FullCrashDumpLocation, CrashReporterConstants.UE4MinidumpName );
+			return Fullname;
+		}
+
+		/// <summary></summary>
+		[XmlElement]
+		public int CrashDumpMode;
+
+		
+
+		/// <summary></summary>
+		[XmlElement( ElementName = "Misc.NumberOfCores" )]
+		public int Misc_NumberOfCores;
+
+		/// <summary></summary>
+		[XmlElement( ElementName = "Misc.NumberOfCoresIncludingHyperthreads" )]
+		public int Misc_NumberOfCoresIncludingHyperthreads;
+
+		/// <summary></summary>
+		[XmlElement( ElementName = "Misc.Is64bitOperatingSystem" )]
+		public bool Misc_Is64bitOperatingSystem;
+
+		/// <summary></summary>
+		[XmlElement( ElementName = "Misc.CPUVendor" )]
+		public string Misc_CPUVendor;
+
+		/// <summary></summary>
+		[XmlElement( ElementName = "Misc.CPUBrand" )]
+		public string Misc_CPUBrand;
+
+		/// <summary></summary>
+		[XmlElement( ElementName = "Misc.PrimaryGPUBrand" )]
+		public string Misc_PrimaryGPUBrand;
+
+		/// <summary></summary>
+		[XmlElement( ElementName = "Misc.OSVersionMajor" )]
+		public string Misc_OSVersionMajor;
+
+		/// <summary></summary>
+		[XmlElement( ElementName = "Misc.OSVersionMinor" )]
+		public string Misc_OSVersionMinor;
+
+		
+		/// <summary></summary>
+		[XmlElement( ElementName = "Misc.AppDiskTotalNumberOfBytes")]
+		public string Misc_AppDiskTotalNumberOfBytes;
+
+		/// <summary></summary>
+		[XmlElement( ElementName = "Misc.AppDiskNumberOfFreeBytes")]
+		public string Misc_AppDiskNumberOfFreeBytes;
+
+		/// <summary></summary>
+		[XmlElement( ElementName = "MemoryStats.TotalPhysical")]
+		public string MemoryStats_TotalPhysical;
+
+		/// <summary></summary>
+		[XmlElement( ElementName = "MemoryStats.TotalVirtual")]
+		public string MemoryStats_TotalVirtual;
+
+		/// <summary></summary>
+		[XmlElement( ElementName = "MemoryStats.PageSize")]
+		public string MemoryStats_PageSize;
+
+		/// <summary></summary>
+		[XmlElement( ElementName = "MemoryStats.TotalPhysicalGB")]
+		public string MemoryStats_TotalPhysicalGB;
+
+		/// <summary></summary>
+		[XmlElement( ElementName = "MemoryStats.AvailablePhysical")]
+		public string MemoryStats_AvailablePhysical;
+
+		/// <summary></summary>
+		[XmlElement( ElementName = "MemoryStats.AvailableVirtual")]
+		public string MemoryStats_AvailableVirtual;
+
+		/// <summary></summary>
+		[XmlElement( ElementName = "MemoryStats.UsedPhysical")]
+		public string MemoryStats_UsedPhysical;
+
+		/// <summary></summary>
+		[XmlElement( ElementName = "MemoryStats.PeakUsedPhysical")]
+		public string MemoryStats_PeakUsedPhysical;
+
+		/// <summary></summary>
+		[XmlElement( ElementName = "MemoryStats.UsedVirtual")]
+		public string MemoryStats_UsedVirtual;
+
+		/// <summary></summary>
+		[XmlElement( ElementName = "MemoryStats.PeakUsedVirtual")]
+		public string MemoryStats_PeakUsedVirtual;
+
+		/// <summary></summary>
+		[XmlElement( ElementName = "MemoryStats.bIsOOM")]
+		public string MemoryStats_bIsOOM;
+
+		/// <summary></summary>
+		[XmlElement( ElementName = "MemoryStats.OOMAllocationSize")]
+		public string MemoryStats_OOMAllocationSize;
+
+		/// <summary></summary>
+		[XmlElement( ElementName = "MemoryStats.OOMAllocationAlignment")]
+		public string MemoryStats_OOMAllocationAlignment;
+
+
+		/// <summary>The UTC time the crash occurred.</summary>
+		[XmlElement]
+		public DateTime TimeofCrash;
+
+		/// <summary> Whether the user allowed us to be contacted. </summary>
+		[XmlElement]
+		public bool bAllowToBeContacted = false;
+
+		/// <summary>A simple default constructor to allow Xml serialization.</summary>
+		public FPrimaryCrashProperties()
+		{ }
+	}
+}
\ No newline at end of file
diff --git a/Engine/Source/Programs/CrashReporter/CrashReportCommon/CrashReportCommon.csproj b/Engine/Source/Programs/CrashReporter/CrashReportCommon/CrashReportCommon.csproj
new file mode 100644
--- /dev/null
+++ b/Engine/Source/Programs/CrashReporter/CrashReportCommon/CrashReportCommon.csproj
@@ -0,0 +1,78 @@
+﻿<?xml version="1.0" encoding="utf-8"?>
+<Project ToolsVersion="4.0" DefaultTargets="Build" xmlns="http://schemas.microsoft.com/developer/msbuild/2003">
+  <Import Project="$(MSBuildExtensionsPath)\$(MSBuildToolsVersion)\Microsoft.Common.props" Condition="Exists('$(MSBuildExtensionsPath)\$(MSBuildToolsVersion)\Microsoft.Common.props')" />
+  <PropertyGroup>
+    <Configuration Condition=" '$(Configuration)' == '' ">Debug</Configuration>
+    <Platform Condition=" '$(Platform)' == '' ">AnyCPU</Platform>
+    <ProjectGuid>{B83DE8C9-9429-4765-86BD-8EEEA89B138F}</ProjectGuid>
+    <OutputType>Library</OutputType>
+    <AppDesignerFolder>Properties</AppDesignerFolder>
+    <RootNamespace>Tools.CrashReporter.CrashReportCommon</RootNamespace>
+    <AssemblyName>CrashReportCommon</AssemblyName>
+    <TargetFrameworkVersion>v4.5</TargetFrameworkVersion>
+    <FileAlignment>512</FileAlignment>
+    <TargetFrameworkProfile>
+    </TargetFrameworkProfile>
+  </PropertyGroup>
+  <PropertyGroup Condition=" '$(Configuration)|$(Platform)' == 'Debug|AnyCPU' ">
+    <DebugSymbols>true</DebugSymbols>
+    <DebugType>full</DebugType>
+    <Optimize>false</Optimize>
+    <OutputPath>..\..\..\..\Binaries\DotNET\</OutputPath>
+    <DefineConstants>DEBUG;TRACE</DefineConstants>
+    <ErrorReport>prompt</ErrorReport>
+    <WarningLevel>4</WarningLevel>
+    <DocumentationFile>..\..\..\..\Binaries\DotNET\CrashReportCommon.xml</DocumentationFile>
+    <RunCodeAnalysis>true</RunCodeAnalysis>
+  </PropertyGroup>
+  <PropertyGroup Condition="'$(Configuration)|$(Platform)' == 'Development|AnyCPU'">
+    <OutputPath>..\..\..\..\Binaries\DotNET\</OutputPath>
+    <DefineConstants>TRACE</DefineConstants>
+    <DocumentationFile>..\..\..\..\Binaries\DotNET\CrashReportCommon.xml</DocumentationFile>
+    <Optimize>true</Optimize>
+    <DebugType>pdbonly</DebugType>
+    <PlatformTarget>AnyCPU</PlatformTarget>
+    <ErrorReport>prompt</ErrorReport>
+    <CodeAnalysisRuleSet>MinimumRecommendedRules.ruleset</CodeAnalysisRuleSet>
+    <CodeAnalysisIgnoreGeneratedCode>true</CodeAnalysisIgnoreGeneratedCode>
+    <RunCodeAnalysis>true</RunCodeAnalysis>
+  </PropertyGroup>
+  <ItemGroup>
+    <Reference Include="System" />
+    <Reference Include="System.Core" />
+    <Reference Include="System.Web" />
+    <Reference Include="System.Web.Extensions" />
+    <Reference Include="System.Xml.Linq" />
+    <Reference Include="System.Data.DataSetExtensions" />
+    <Reference Include="Microsoft.CSharp" />
+    <Reference Include="System.Data" />
+    <Reference Include="System.Xml" />
+  </ItemGroup>
+  <ItemGroup>
+    <Compile Include="Common.cs" />
+    <Compile Include="CrashContext.cs" />
+    <Compile Include="CrashReporterConstants.cs" />
+    <Compile Include="JiraConnection.cs" />
+    <Compile Include="LogWriter.cs" />
+    <Compile Include="Properties\AssemblyInfo.cs" />
+    <Compile Include="..\..\DotNETCommon\MetaData.cs">
+      <Link>Properties\MetaData.cs</Link>
+    </Compile>
+    <Compile Include="ReportData.cs" />
+    <Compile Include="WerReport.cs" />
+  </ItemGroup>
+  <ItemGroup>
+    <ProjectReference Include="..\..\DotNETCommon\DotNETUtilities\DotNETUtilities.csproj">
+      <Project>{5d7d66e8-8c76-4af9-b3ec-2ef03421d730}</Project>
+      <Name>DotNETUtilities</Name>
+    </ProjectReference>
+  </ItemGroup>
+  <Import Project="$(MSBuildToolsPath)\Microsoft.CSharp.targets" />
+  <!-- To modify your build process, add your task inside one of the targets below and uncomment it. 
+       Other similar extension points exist, see Microsoft.Common.targets.
+  <Target Name="BeforeBuild">
+  </Target>
+  <Target Name="AfterBuild">
+  </Target>
+  -->
+</Project>
\ No newline at end of file
diff --git a/Engine/Source/Programs/CrashReporter/CrashReportCommon/CrashReporterConstants.cs b/Engine/Source/Programs/CrashReporter/CrashReportCommon/CrashReporterConstants.cs
new file mode 100644
--- /dev/null
+++ b/Engine/Source/Programs/CrashReporter/CrashReportCommon/CrashReporterConstants.cs
@@ -0,0 +1,42 @@
+﻿// Copyright 1998-2015 Epic Games, Inc. All Rights Reserved.
+
+using System;
+using System.Xml.Serialization;
+
+namespace Tools.CrashReporter.CrashReportCommon
+{
+	/// <summary>
+	/// Any constant variables that need sharing between client and server.
+	/// </summary>
+	public class CrashReporterConstants
+	{
+		/// <summary>Completely arbitrary maximum file size to upload (basic exploit check). Currently 100MB.</summary>
+		public const long MaxFileSizeToUpload = 100 * 1024 * 1024;
+
+		/// <summary>Size of data chunks expected. Currently 16kB.</summary>
+		public const int StreamChunkSize = 16384;
+
+		/// <summary> Minimum depth of a callstack that is considered to be a valid one. </summary>
+		public const int MinCallstackDepth = 3;
+
+		/// <summary>The name of the generated report. Currently 'Diagnostics.txt'.</summary>
+		public const string DiagnosticsFileName = "Diagnostics.txt";
+
+		/// <summary>The name of the generated video of the editor session. Currently 'CrashVideo.avi'</summary>
+		public const string VideoFileName = "CrashVideo.avi";
+
+		/// <summary> Licensee branch name, we ignore crashes from this branch. </summary>
+		public const string LicenseeBranchName = "UE4-QA";
+
+		/// <summary>String to chop off the start of branch paths</summary>
+		public const string P4_DEPOT_PREFIX = "//depot/";
+
+		/// <summary>UE4 minidump name</summary>
+		public const string UE4MinidumpName = "UE4Minidump.dmp";
+
+		/// <summary>A simple default constructor to allow Xml serialisation.</summary>
+		public CrashReporterConstants()
+		{
+		}
+	}
+}
\ No newline at end of file
diff --git a/Engine/Source/Programs/CrashReporter/CrashReportCommon/JiraConnection.cs b/Engine/Source/Programs/CrashReporter/CrashReportCommon/JiraConnection.cs
new file mode 100644
--- /dev/null
+++ b/Engine/Source/Programs/CrashReporter/CrashReportCommon/JiraConnection.cs
@@ -0,0 +1,387 @@
+// Copyright 1998-2015 Epic Games, Inc. All Rights Reserved.
+
+using System;
+using System.Collections.Generic;
+using System.Diagnostics;
+using System.IO;
+using System.Linq;
+using System.Net;
+using System.Text;
+using System.Threading; 
+using System.Threading.Tasks;
+using System.Web.Script.Serialization;
+
+/// <summary>
+/// 
+/// </summary>
+public class JiraConnection
+{
+	/// <summary> Helper class for inline initialization of the list of the KeyValuePair. </summary>
+	class ListKVP<TKey, TValue> : List<KeyValuePair<TKey, TValue>>
+	{
+		/// <summary> Overloaded Add. </summary>
+		public void Add( TKey Key, TValue Value )
+		{
+			Add( new KeyValuePair<TKey, TValue>( Key, Value ) );
+		}
+	}
+
+	/// <summary>
+	/// Verbs for HTTP requests
+	/// </summary>
+	public enum JiraMethod
+	{
+		/// <summary></summary>
+		POST,
+		/// <summary></summary>
+		GET,
+		/// <summary></summary>
+		DELETE,
+		/// <summary></summary>
+		PUT,
+	}
+
+	/// <summary>
+	/// 
+	/// </summary>
+	Dictionary<string, string> Credentials = new Dictionary<string, string>( StringComparer.InvariantCultureIgnoreCase );
+
+	/// <summary> Whether we can use JIRA integrations. </summary>
+	bool bCanBeUsed = false;
+
+	/// <summary> Whether we have read component and version arrays. </summary>
+	bool bHasComponentAndVersionArrays = false;
+
+	/// <summary> Maps the name to the component. </summary>
+	Dictionary<string, Dictionary<string, object>> NameToComponents = new Dictionary<string, Dictionary<string, object>>();
+	/// <summary> Maps the name to the component. </summary>
+	public Dictionary<string, Dictionary<string, object>> GetNameToComponents()
+	{
+		return NameToComponents;
+	}
+
+	/// <summary> Maps the name to the version. </summary>
+	Dictionary<string, Dictionary<string, object>> NameToVersions = new Dictionary<string, Dictionary<string, object>>();
+	/// <summary> Maps the name to the version. </summary>
+	public Dictionary<string, Dictionary<string, object>> GetNameToVersions()
+	{
+		return NameToVersions;
+	}
+
+	/// <summary> Maps the name to the platform. HARD-CODED </summary>
+	Dictionary<string, Dictionary<string, object>> NameToPlatform = new Dictionary<string, Dictionary<string, object>>();
+	/// <summary> Maps the name to the platform. HARD-CODED </summary>
+	public Dictionary<string, Dictionary<string, object>> GetNameToPlatform()
+	{
+		return NameToPlatform;
+	}
+
+	/// <summary>
+	/// Singleton
+	/// </summary>
+	static JiraConnection Singleton = new JiraConnection();
+	private static readonly object SyncObj = new object();
+
+	/// <summary>
+	/// 
+	/// </summary>
+	JiraConnection()
+	{
+		// Parse the credentials		
+		string CredentialsFileName = Path.Combine( Environment.GetFolderPath( Environment.SpecialFolder.CommonApplicationData ), "ue4credentials" );
+		if( File.Exists( CredentialsFileName ) )
+		{
+			ParseConfigFile( CredentialsFileName, Credentials );
+			bCanBeUsed = true;
+		}
+	}
+
+	/// <summary>
+	/// Retrieves all components from https://jira.ol.epicgames.net/rest/api/2/project/UE/components
+	/// </summary>
+	private void InitializeComponentArray()
+	{
+		/*
+			"self": "https://jira.ol.epicgames.net/rest/api/2/component/..",
+			"id": "14019",
+			"name": "...",
+			"lead": {
+			"self": "https://jira.ol.epicgames.net/rest/api/2/user?username=...,
+			"key": "...",
+			"name": "...",
+			"avatarUrls": {...},
+			"displayName": "...",
+			"active": true
+			},
+			"assigneeType": "PROJECT_DEFAULT",
+			"realAssigneeType": "PROJECT_DEFAULT",
+			"isAssigneeTypeValid": false
+		*/
+		Dictionary<string, object> Request = new Dictionary<string, object>();
+		string[] FieldsToGet = new string[] { "self", "id", "name" };
+		Request.Add( "fields", FieldsToGet );
+		HttpWebResponse Response = JiraRequest( "/project/UE/components", JiraMethod.GET, null, HttpStatusCode.OK );
+
+		// Parse the results
+		List<Dictionary<string, object>> Results = ParseJiraResponseIntoList( Response );
+		foreach( var Component in Results )
+		{
+			string Name = (string)Component["name"];
+			NameToComponents.Add( Name, Component );
+		}
+	}
+
+	/// <summary>
+	/// Retrieves all build versions from https://jira.ol.epicgames.net/rest/api/2/project/UE/versions
+	/// </summary>
+	private void InitializeVersionArray()
+	{
+		/*
+			"self": "https://jira.ol.epicgames.net/rest/api/2/version/...",
+			"id": "11700",
+			"description": "Scheduled engine release",
+			"name": "4.8",
+			"archived": false,
+			"released": false,
+			"startDate": "2015-01-12",
+			"releaseDate": "2015-03-03",
+			"overdue": false,
+			"userStartDate": "12/Jan/15",
+			"userReleaseDate": "03/Mar/15",
+			"projectId": 11205
+		*/
+		Dictionary<string, object> Request = new Dictionary<string, object>();
+		string[] FieldsToGet = new string[] { "self", "id", "name" };
+		Request.Add( "fields", FieldsToGet );
+		HttpWebResponse Response = JiraRequest( "/project/UE/versions", JiraMethod.GET, null, HttpStatusCode.OK );
+
+		// Parse the results
+		List<Dictionary<string, object>> Results = ParseJiraResponseIntoList( Response );
+		foreach( var Version in Results )
+		{
+			string Name = (string)Version["name"];
+			NameToVersions.Add( Name, Version );
+		}
+	}
+
+
+	/// <summary> Initializes hard-coded array of field Platforms. </summary>
+	private void InitializePlatforms()
+	{
+		var JiraNameToPlatform = new ListKVP<string,string>
+		{ 
+			{"All","/customFieldOption/11103"},
+			{"Android","/customFieldOption/11104"},
+			{"IOS","/customFieldOption/11105"},
+			{"Cooker","/customFieldOption/11706"},
+			{"HTML5","/customFieldOption/11106"},
+			{"Linux","/customFieldOption/11107"},
+			{"Mac","/customFieldOption/11108"},
+			{"PS4","/customFieldOption/11111"},
+			{"Windows","/customFieldOption/11112"},
+			{"Win32","/customFieldOption/11112"},
+			{"Win64","/customFieldOption/11112"},
+			{"XBoxOne","/customFieldOption/11113"},
+			{"Oculus","/customFieldOption/11353"},		
+		};
+
+		foreach( var Mapping in JiraNameToPlatform )
+		{
+			HttpWebResponse Response = JiraRequest( Mapping.Value, JiraMethod.GET, null, HttpStatusCode.OK );
+			Dictionary<string, object> JSON = ParseJiraResponse( Response );
+
+			NameToPlatform.Add( Mapping.Key, JSON );
+		}
+	}
+
+	/// <summary>
+	/// 
+	/// </summary>
+	/// <returns></returns>
+	public static JiraConnection Get()
+	{
+		if( !Singleton.bHasComponentAndVersionArrays )
+		{
+			lock( SyncObj )
+			{
+				if( !Singleton.bHasComponentAndVersionArrays )
+				{
+					Singleton.InitializeComponentArray();
+					Singleton.InitializeVersionArray();
+					Singleton.InitializePlatforms();
+					Singleton.bHasComponentAndVersionArrays = true;
+				}
+			}
+		}
+
+		return Singleton;
+	}
+
+	/// <summary>
+	/// 
+	/// </summary>
+	/// <returns></returns>
+	public bool CanBeUsed()
+	{
+		return bCanBeUsed;
+	}
+
+	/// <summary>
+	/// Parses a config file on disk. 
+	/// </summary>
+	/// <param name="FileName">The config file to read.</param>
+	/// <param name="Settings">Dictionary to receive all the key/value pairs read from the config file.</param>
+	void ParseConfigFile( string FileName, Dictionary<string, string> Settings )
+	{
+		foreach( string Line in File.ReadAllLines( FileName ) )
+		{
+			string TrimLine = Line.Trim();
+			if( !TrimLine.StartsWith( ";" ) )
+			{
+				int EqualsIdx = Line.IndexOf( '=' );
+				if( EqualsIdx != -1 )
+				{
+					Settings.Add( Line.Substring( 0, EqualsIdx ).TrimEnd(), Line.Substring( EqualsIdx + 1 ).TrimStart() );
+				}
+			}
+		}
+	}
+
+	/// <summary>
+	/// Creates a Jira ticket with the given fields
+	/// </summary>
+	public string AddJiraTicket( Dictionary<string, object> Fields )
+	{
+		// Create the request object
+		Dictionary<string, object> Request = new Dictionary<string, object> { { "fields", Fields } };
+		HttpWebResponse Response = JiraRequest( "/issue", JiraMethod.POST, Request, HttpStatusCode.Created );
+
+		// Return the new issue id
+		Dictionary<string, object> ResponseObject = ParseJiraResponse( Response );
+		return (string)ResponseObject["key"];
+	}
+
+	/// <summary>
+	/// Updates a Jira ticket with the given fields
+	/// </summary>
+	public void UpdateJiraTicket( string Key, Dictionary<string, object> Fields )
+	{
+		// Create the request object
+		Dictionary<string, object> Request = new Dictionary<string, object> { { "fields", Fields } };
+		HttpWebResponse Response = JiraRequest( String.Format( "/issue/{0}", Key ), JiraMethod.PUT, Request, HttpStatusCode.OK );
+	}
+
+	/// <summary>
+	/// Searches for JIRA tickets
+	/// </summary>
+	/// <param name="SearchQuery">Search query in JQL format</param>
+	/// <param name="FieldsToGet">Fields that will be grabbed from the results</param>
+	public Dictionary<string,Dictionary<string, object>> SearchJiraTickets( string SearchQuery, string[] FieldsToGet )
+	{
+		Dictionary<string, Dictionary<string, object>> Result = new Dictionary<string, Dictionary<string, object>>();
+
+		// Do the query in multiple steps, to avoid timeouts when receiving large amounts of data
+		const int MaxResultsPerQuery = 100;
+		for( int StartAt = 0; ; StartAt += MaxResultsPerQuery )
+		{
+			// Create the query
+			Dictionary<string, object> Request = new Dictionary<string, object>();
+			Request.Add( "jql", SearchQuery );
+			Request.Add( "startAt", StartAt.ToString() );
+			Request.Add( "maxResults", MaxResultsPerQuery.ToString() );
+			Request.Add( "fields", FieldsToGet );
+			HttpWebResponse Response = JiraRequest( "/search", JiraMethod.POST, Request, HttpStatusCode.OK );
+
+			// Parse the results
+			Dictionary<string, object> Results = ParseJiraResponse( Response );
+			System.Collections.ArrayList Issues = (System.Collections.ArrayList)Results["issues"];
+
+			// Process this chunk of results
+			foreach( Dictionary<string, object> Issue in Issues )
+			{
+				string Key = (string)Issue["key"];
+				Dictionary<string, object> Fields = (Dictionary<string, object>)Issue["fields"];
+				Result.Add( Key, Fields );
+			}
+
+			// Quit once we got less than the max results
+			if( Issues.Count < MaxResultsPerQuery )
+			{
+				break;
+			}
+		}
+
+		return Result;
+	}
+
+	/// <summary>
+	/// Issues a Jira web request
+	/// </summary>
+	HttpWebResponse JiraRequest( string RequestUrl, JiraMethod Method, object Input )
+	{
+		// Get the username and password
+		string UserName, Password;
+		if( Credentials.TryGetValue( "jira.username", out UserName ) && Credentials.TryGetValue( "jira.password", out Password ) )
+		{
+			// Send the request
+			HttpWebRequest Request = (HttpWebRequest)WebRequest.Create( "https://jira.ol.epicgames.net/rest/api/2" + RequestUrl );
+			Request.Method = Enum.GetName( typeof( JiraMethod ), Method );
+			Request.Headers.Add( "Authorization", "Basic " + Convert.ToBase64String( UTF8Encoding.UTF8.GetBytes( UserName + ":" + Password ) ) );
+			Request.ContentType = "application/json";
+			if( Input != null )
+			{
+				using( StreamWriter RequestWriter = new StreamWriter( Request.GetRequestStream(), Encoding.ASCII ) )
+				{
+					JavaScriptSerializer Serializer = new JavaScriptSerializer();
+					RequestWriter.Write( Serializer.Serialize( Input ) );
+				}
+			}
+
+			// Check the response was ok
+			return (HttpWebResponse)Request.GetResponse();
+		}
+		else
+		{
+			return null;
+		}
+	}
+
+	/// <summary>
+	/// Issues a Jira web request
+	/// </summary>
+	public HttpWebResponse JiraRequest( string RequestUrl, JiraMethod Method, object Input, HttpStatusCode ExpectedStatus )
+	{
+		HttpWebResponse Response = JiraRequest( RequestUrl, Method, Input );
+		if( Response.StatusCode != ExpectedStatus )
+		{
+			return null;
+		}
+		return Response;
+	}
+
+	/// <summary>
+	/// Parses a Json object from a web response
+	/// </summary>
+	public static Dictionary<string, object> ParseJiraResponse( HttpWebResponse Response )
+	{
+		using( StreamReader ResponseReader = new StreamReader( Response.GetResponseStream() ) )
+		{
+			string ResponseText = ResponseReader.ReadToEnd();
+			JavaScriptSerializer Serializer = new JavaScriptSerializer();
+			return Serializer.Deserialize<Dictionary<string, object>>( ResponseText );
+		}
+	}
+
+	/// <summary>
+	/// Parses a Json object from a web response
+	/// </summary>
+	public static List<Dictionary<string, object>> ParseJiraResponseIntoList( HttpWebResponse Response )
+	{
+		using( StreamReader ResponseReader = new StreamReader( Response.GetResponseStream() ) )
+		{
+			string ResponseText = ResponseReader.ReadToEnd();
+			JavaScriptSerializer Serializer = new JavaScriptSerializer();
+			return Serializer.Deserialize<List<Dictionary<string, object>>>( ResponseText );
+		}
+	}
+}
\ No newline at end of file
diff --git a/Engine/Source/Programs/CrashReporter/CrashReportCommon/LogWriter.cs b/Engine/Source/Programs/CrashReporter/CrashReportCommon/LogWriter.cs
new file mode 100644
--- /dev/null
+++ b/Engine/Source/Programs/CrashReporter/CrashReportCommon/LogWriter.cs
@@ -0,0 +1,129 @@
+// Copyright 1998-2015 Epic Games, Inc. All Rights Reserved.
+
+using System;
+using System.IO;
+using System.Text;
+using System.Diagnostics;
+using System.Threading;
+
+namespace Tools.CrashReporter.CrashReportCommon
+{
+	/// <summary>
+	/// Thread-safe log creation and writing for crash report processes
+	/// </summary>
+	public class LogWriter
+	{
+		/// <summary>
+		/// Constructor: opens the log file
+		/// </summary>
+		/// <param name="AppNameParam">Application name to prepend to log filename</param>
+		/// <param name="LogFolderParam">Folder in which to store log file</param>
+		public LogWriter(string AppNameParam, string LogFolderParam)
+		{
+			// Store the arguments for repeated log file creation
+			AppName = AppNameParam;
+			LogFolder = LogFolderParam;
+
+			// Ensure the log folder exists
+			DirectoryInfo DirInfo = new DirectoryInfo(LogFolder);
+			if(!DirInfo.Exists)
+			{
+				DirInfo.Create();
+			}
+
+			CreateNewLogFile();
+		}
+
+		/// <summary>
+		/// Shutdown: flush and close the log file
+		/// </summary>
+		public void Dispose()
+		{
+			lock( Sync )
+			{
+				LogFile.Flush();
+				LogFile.Close();
+				LogFile = null;
+			}
+		}
+
+		/// <summary>
+		/// Write a line to the log file and to the console
+		/// </summary>
+		/// <param name="Line">Text to write</param>
+		public void Print(string Line)
+		{
+			if (Line != null)
+			{
+				string FullLine = DateTime.Now.ToString( "yyyy/MM/dd HH:mm:ss" ) + " [" + Thread.CurrentThread.ManagedThreadId.ToString( "D3" ) + "] : " + Line;
+				lock( Sync )
+				{
+					LogFile.WriteLine( FullLine );
+					LogFile.Flush();
+				}
+
+				// Relay to the console if it's available
+				Console.WriteLine( FullLine );
+
+				if( System.Diagnostics.Debugger.IsAttached )
+				{
+					Debug.WriteLine( FullLine );
+				}
+			}
+		}
+
+		/// <summary>
+		/// Clean out any old log files in the same directory as this log
+		/// </summary>
+		/// <param name="OlderThanDays">Files older than this many days will be deleted</param>
+		public void CleanOutOldLogs(
+			int OlderThanDays
+			)
+		{
+			DirectoryInfo DirInfo = new DirectoryInfo(LogFolder);
+			foreach(FileInfo LogFileInfo in DirInfo.GetFiles())
+			{
+				if (LogFileInfo.LastWriteTimeUtc.AddDays(OlderThanDays) < DateTime.UtcNow)
+				{
+					LogFileInfo.Delete();
+				}
+			}
+		}
+
+		/// <summary>
+		/// Create either an initial log file or switch to a new, time-stamped log file
+		/// </summary>
+		public void CreateNewLogFile()
+		{
+			lock( Sync )
+			{
+				var OldLogFile = LogFile;
+
+				var LogFileName = Path.Combine( LogFolder, AppName + "-" + DateTime.UtcNow.ToString( "yyyy-MM-dd" ) + ".txt" );
+
+				// Create a threadsafe text writer
+				LogFile = new StreamWriter( LogFileName, true, Encoding.UTF8 );
+
+				Print( AppName + ", Copyright 2014-2015 Epic Games, Inc." );
+
+				// Don't close the previous log file until the new one is up and running
+				if( OldLogFile != null )
+				{
+					OldLogFile.Flush();
+					OldLogFile.Close();
+				}
+			}
+		}	
+
+		/// <summary>Open file stream messages are written to</summary>
+		volatile TextWriter LogFile;
+
+		/// <summary>Application name to pre-pend to log filenames and put in copyright notice</summary>
+		string AppName;
+
+		/// <summary>Folder log has been created in</summary>
+		string LogFolder;
+
+		object Sync = new object();
+	}
+}
diff --git a/Engine/Source/Programs/CrashReporter/CrashReportCommon/Properties/AssemblyInfo.cs b/Engine/Source/Programs/CrashReporter/CrashReportCommon/Properties/AssemblyInfo.cs
new file mode 100644
--- /dev/null
+++ b/Engine/Source/Programs/CrashReporter/CrashReportCommon/Properties/AssemblyInfo.cs
@@ -0,0 +1,11 @@
+﻿// Copyright 1998-2015 Epic Games, Inc. All Rights Reserved.
+
+using System.Reflection;
+
+// General Information about an assembly is controlled through the following 
+// set of attributes. Change these attribute values to modify the information
+// associated with an assembly.
+[assembly: AssemblyTitle( "CrashReportCommon" )]
+[assembly: AssemblyDescription( "The common data structures used by the CrashReport applications" )]
+[assembly: AssemblyConfiguration( "" )]
+
diff --git a/Engine/Source/Programs/CrashReporter/CrashReportCommon/ReportData.cs b/Engine/Source/Programs/CrashReporter/CrashReportCommon/ReportData.cs
new file mode 100644
--- /dev/null
+++ b/Engine/Source/Programs/CrashReporter/CrashReportCommon/ReportData.cs
@@ -0,0 +1,498 @@
+// Copyright 1998-2015 Epic Games, Inc. All Rights Reserved.
+
+using System;
+using System.Collections.Generic;
+using System.IO;
+
+namespace Tools.CrashReporter.CrashReportCommon
+{
+	/// <summary> Basic support for native FEngineVersion. </summary>
+	public class FEngineVersion
+	{
+		/// <summary> Initialization constructor. </summary>
+		public FEngineVersion(string StringEngineVersion)
+		{
+			// 4.10.0-0+UE4
+			int MinusPos = StringEngineVersion.IndexOf( '-' );
+			int PlusPos = StringEngineVersion.IndexOf( '+' );
+
+			VersionNumber = StringEngineVersion.Substring( 0, MinusPos );
+			string ChangelistSring = StringEngineVersion.Substring( MinusPos + 1, PlusPos - MinusPos - 1 );
+			uint.TryParse( ChangelistSring, out Changelist );
+
+			Branch = StringEngineVersion.Substring( PlusPos+1 );
+		}
+
+		/// <summary>
+		/// Returns the branch name corresponding to this version
+		/// </summary>
+		public string GetCleanedBranch()
+		{
+			string CleanedBranch = Branch.Replace( '+', '/' );
+			CleanedBranch = CleanedBranch.Replace( CrashReporterConstants.P4_DEPOT_PREFIX, "" );
+			CleanedBranch = CleanedBranch.Substring( 0, Math.Min( 31, CleanedBranch.Length ) ); // Database limitation.
+			return CleanedBranch;
+		}
+
+		/// <summary> Version number. </summary>
+		public string VersionNumber;
+
+		/// <summary> Change list/ Built from CL. </summary>
+		public uint Changelist;
+
+		/// <summary> Branch. </summary>
+		public string Branch;
+
+	}
+
+	/// <summary>Helper to read useful information out of the report XML and diagnostics files</summary>
+	public class FReportData
+	{
+		/// <summary>Name of report for logging</summary>
+		public readonly string ReportName;
+
+		/// <summary> Ticks for this report data. </summary>
+		public readonly Int64 Ticks = DateTime.Now.Ticks;
+
+		/// <summary> Helper method, display this report data as a report name. Debugging purpose. </summary>
+		public override string ToString()
+		{
+			return ReportName;
+		}
+
+		/// <returns>Unique name of this report, can be used as name of file.</returns>
+		public string GetReportNameAsFilename()
+		{
+			string EngineVersion = string.Format
+			(
+				"{0}-{1}+{2}",
+				!string.IsNullOrEmpty( BuildVersion ) ? BuildVersion : "NO_BUILD_VERSION",
+				!string.IsNullOrEmpty( BuiltFromCL ) ? BuiltFromCL : "NO_BUILT_FROM_CL",
+				!string.IsNullOrEmpty( BranchName ) ? BranchName.Replace( '/', '+' ) : "NO_BRANCH"
+			);
+			string Directory = string.Format( "{0}_{1}__{2}", EngineVersion, !string.IsNullOrEmpty( GameName ) ? GameName : "NO_GAMENAME", Ticks );
+			return Directory;
+		}
+
+		/// <summary> Returns engine version as a string ie.: 4.7.0-2448202+++depot+UE4-Releases+4.7 </summary>
+		public string GetEngineVersion()
+		{
+			string EngineVersion = string.Format
+			(
+				"{0}-{1}+{2}",
+				!string.IsNullOrEmpty( BuildVersion ) ? BuildVersion : "",
+				!string.IsNullOrEmpty( BuiltFromCL ) ? BuiltFromCL : "",
+				!string.IsNullOrEmpty( BranchNameWithDepot ) ? BranchNameWithDepot.Replace( '/', '+' ) : ""
+			);
+
+			// An invalid version, so return an empty string.
+			if( EngineVersion == "-+" )
+			{
+				return "";
+			}
+
+			return EngineVersion;
+		}
+
+		/// <summary>
+		/// Read the callstack from the diagnostics file
+		/// <param name="DiagnosticsPath">Absolute path of the diagnostics file.</param>
+		/// <returns>The callstack as an array of entries</returns>
+		/// </summary>
+		static public string[] GetCallStack(string DiagnosticsPath)
+		{
+			return ExtractSection(DiagnosticsPath, "<CALLSTACK START>", "<CALLSTACK END>");
+		}
+
+		/// <summary>
+		/// Read the source context from the diagnostics file
+		/// <param name="DiagnosticsPath">Absolute path of the diagnostics file.</param>
+		/// <returns>Source code around the error point as an array of lines</returns>
+		/// </summary>
+		static public string[] GetSourceContext(string DiagnosticsPath)
+		{
+			return ExtractSection(DiagnosticsPath, "<SOURCE START>", "<SOURCE END>");
+		}
+
+		/// <summary>Accessor for the game name from the report metadata</summary>
+		public string GameName
+		{
+			get
+			{
+				try
+				{
+					return Metadata.ProblemSignatures.Parameter0;
+				}
+				catch( System.Exception )
+				{
+					return "";
+				}
+			}
+		}
+
+		/// <summary>Accessor for the build version string from the report metadata</summary>
+		public string BuildVersion
+		{
+			get
+			{
+				try
+				{
+					string FourParms = Metadata.ProblemSignatures.Parameter1;
+
+					var Subs = FourParms.Split( ".".ToCharArray(), StringSplitOptions.RemoveEmptyEntries );
+
+					if( Subs.Length >= 3 )
+					{
+						return string.Format( "{0}.{1}.{2}", Subs[0], Subs[1], Subs[2] );
+					}
+
+					return "";
+				}
+				catch( System.Exception )
+				{
+					return "";
+				}
+			}
+		}
+
+		/// <summary>Accessor for the language identifier from the report metadata</summary>
+		public string Language
+		{
+			get
+			{
+				try
+				{
+					return Metadata.DynamicSignatures.Parameter2;
+				}
+				catch( System.Exception )
+				{
+					return "";
+				}
+			}
+		}
+
+		/// <summary>Accessor for the command line from the report metadata</summary>
+		public string CommandLine
+		{
+			get
+			{
+				string Result = "";
+				try
+				{
+					if( Metadata.ProblemSignatures != null && Metadata.ProblemSignatures.Parameter8 != null )
+					{
+						string[] Components = Metadata.ProblemSignatures.Parameter8.Split( new[] { '!' } );
+						if( Components.Length > 0 )
+						{
+							Result = Components[1];
+						}
+					}
+				}
+				finally
+				{}
+				return Result;
+			}
+		}
+
+		/// <summary>Accessor for the system language from the report metadata</summary>
+		public string SystemLanguage
+		{
+			get
+			{
+				try
+				{
+					return Metadata.OSVersionInformation.LCID;
+				}
+				catch( System.Exception )
+				{
+					return "";
+				}	
+			}
+		}
+
+		/// <summary>Extract the branch name from the report metadata</summary>
+		public string BranchName
+		{
+			get
+			{
+				string BranchName = ExtractFromBaseDirString(0).Replace('+', '/');
+				if (BranchName.StartsWith(CrashReporterConstants.P4_DEPOT_PREFIX))
+				{
+					BranchName = BranchName.Substring( CrashReporterConstants.P4_DEPOT_PREFIX.Length );
+				}
+
+				BranchName = BranchName.Substring( 0, Math.Min( 31, BranchName.Length ) ); // Database limitation.
+
+				return BranchName;
+			}
+		}
+
+		/// <summary>Extract the branch name from the report metadata</summary>
+		public string BranchNameWithDepot
+		{
+			get
+			{
+				string BranchName = ExtractFromBaseDirString( 0 ).Replace( '+', '/' );
+				return BranchName;
+			}
+		}
+
+		/// <summary>Extract the base directory from the report metadata</summary>
+		public string BaseDir
+		{
+			get
+			{
+				return ExtractFromBaseDirString(1);
+			}
+		}
+
+		/// <summary>Extract the platform name from the report metadata</summary>
+		public string Platform
+		{
+			get
+			{
+				// Extract the platform name encoded in parameter9 and add the Windows product info, e.g. Windows 7 Professional
+				string PlatformName = "";
+				string BranchPath = ExtractFromBaseDirString(1);
+				if (BranchPath != "")
+				{
+					string[] FolderComponents = BranchPath.Split(new[] {'/'}, StringSplitOptions.RemoveEmptyEntries);
+					if (FolderComponents.Length != 0)
+					{
+						PlatformName = FolderComponents[FolderComponents.Length - 1];
+					}
+				}
+
+				if (!string.IsNullOrEmpty( Metadata.OSVersionInformation.Product ) && !string.IsNullOrEmpty( Metadata.OSVersionInformation.Architecture ))
+				{
+					PlatformName = string.Format( "{0} [{2} {1}]", PlatformName, Metadata.OSVersionInformation.Architecture == "X64" ? "64b" : "32b", Metadata.OSVersionInformation.Product );
+				}
+
+				return PlatformName;
+			}
+		}
+
+		/// <summary> This should be ErrorMessage=, but leave it as the same for backward compatibility. </summary>
+		public const string ErrorMessageMarker = "AssertLog=\"";
+		
+		const string ExceptionMarker = "Exception was \"";
+
+		const string MachineIdMarker = "MachineId:";
+		const string EpicAccountIdMarker = "EpicAccountId:";
+		const string UserNameMarker = "Name:";
+
+
+		/// <summary>
+		/// The error message, can be assertion message, ensure message or message from the fatal error.
+		/// @see GErrorMessage.
+		/// </summary>
+		public string[] ErrorMessage
+		{
+			get
+			{
+				var EmptyArray = new string[] { };
+				try
+				{
+					// This may fail if we encounter a very long error message.
+					int StartIndex = Metadata.ProblemSignatures.Parameter8.IndexOf( ErrorMessageMarker );
+					if( StartIndex > 0 )
+					{
+						string ErrorMessageFlat = Metadata.ProblemSignatures.Parameter8.Substring( StartIndex );
+
+						ErrorMessageFlat = ErrorMessageFlat.Replace( ErrorMessageMarker, "" ).Replace( "\"", "" );
+						ErrorMessageFlat = ErrorMessageFlat.Substring( 0, Math.Min( 511, ErrorMessageFlat.Length ) ); // Database limitation.
+
+						EmptyArray = ErrorMessageFlat.Split( new[] { '#' }, StringSplitOptions.RemoveEmptyEntries );
+					}
+
+					return EmptyArray;
+				}
+				catch( System.Exception )
+				{
+					return EmptyArray;
+				}
+			}
+		}
+
+		/// <summary>
+		/// Exception description, used if the ErrorMessage is empty
+		/// </summary>
+		static public string[] GetExceptionDescription(string DiagnosticsPath)
+		{
+			var ResultArray = new string[] { "" };
+			try
+			{
+				string[] Lines = File.ReadAllLines( DiagnosticsPath );
+				foreach( string Line in Lines )
+				{
+					if( Line.Contains( ExceptionMarker ) )
+					{
+						int StartIndex = Line.IndexOf( ExceptionMarker );
+						int EndIndex = Line.LastIndexOf( "\"" );
+						ResultArray = new string[] { Line.Substring( StartIndex, EndIndex - StartIndex ).Replace( ExceptionMarker, "" ) };
+						break;
+					}
+				}
+
+				return ResultArray;
+			}
+			catch( System.Exception )
+			{
+				return ResultArray;
+			}
+		}
+
+		/// <summary>Extract the engine mode from the report metadata</summary>
+		public string EngineMode
+		{
+			get
+			{
+				return ExtractFromBaseDirString(2);
+			}
+		}
+
+		/// <summary>Extract the CL# from the report metadata</summary>
+		public string BuiltFromCL
+		{
+			get
+			{
+				return ExtractFromBaseDirString(3);
+			}
+		}
+
+		/// <summary>User-provided description from the report</summary>
+		public string[] UserDescription
+		{
+			get
+			{
+				var EmptyArray = new string[] { };
+				if( Metadata.DynamicSignatures != null && Metadata.DynamicSignatures.Parameter3 != null )
+				{
+					// Description in the DB has only 512 characters.
+					string ExtractedDescription = Metadata.DynamicSignatures.Parameter3;
+					return ExtractedDescription.Substring( 0, Math.Min( 480, ExtractedDescription.Length ) ).Split( new[] { '\n' }, StringSplitOptions.RemoveEmptyEntries );
+				}
+				return EmptyArray;
+			}
+		}
+
+		/// <summary>The unique ID used to identify the machine the crash occurred on.</summary>
+		public string MachineId
+		{
+			get
+			{
+				return ExtractMarkedString( MachineIdMarker );
+			}
+		}
+
+		/// <summary>The Epic account ID for the user who last used the Launcher.</summary>
+		public string EpicAccountId
+		{
+			get
+			{
+				return ExtractMarkedString( EpicAccountIdMarker );
+			}
+		}
+
+		/// <summary> The name of the user that caused this crash. </summary>
+		public string UserName
+		{
+			get
+			{
+				return ExtractMarkedString( UserNameMarker );
+			}
+		}
+
+		/// <summary> Whether the user allowed us to be contacted. </summary>
+		public bool AllowToBeContacted
+		{
+			get
+			{
+				return Metadata.DynamicSignatures != null ? Metadata.DynamicSignatures.bAllowToBeContacted : false;
+			}
+		}
+
+		/// <summary> Extracts a string marked with the specified string value. </summary>
+		private string ExtractMarkedString( string NameMarker )
+		{
+			if( Metadata.DynamicSignatures != null && Metadata.DynamicSignatures.Parameter4 != null )
+			{
+				string[] Components = Metadata.DynamicSignatures.Parameter4.Split( new[] { '!' }, StringSplitOptions.RemoveEmptyEntries );
+				foreach( string Component in Components )
+				{
+					if( Component.StartsWith( NameMarker ) )
+					{
+						string Result = Component.Replace( NameMarker, "" );
+						return Result;
+					}
+				}
+			}
+			return "";
+		}
+
+		/// <summary>
+		/// Constructor: store the metadata and extract the callstack from diagnostics file
+		/// </summary>
+		/// <param name="Data">Report meta-data object, parsed from XML</param>
+		/// <param name="InReportNamePath">Name of report for logging</param>
+		public FReportData(	WERReportMetadata Data, string InReportNamePath	)
+		{
+			Metadata = Data;
+			ReportName = InReportNamePath;
+		}
+
+		/// <summary>
+		/// Split the combined base directory string into individual components
+		/// <param name="Index">Index of component to retrieve</param>
+		/// </summary>
+		string ExtractFromBaseDirString(int Index)
+		{
+			try
+			{
+				string[] Components = Metadata.ProblemSignatures.Parameter9.Split( new[] { '!' } );
+				return Components.Length > Index ? Components[Index] : "";
+			}
+			catch( System.Exception )
+			{
+				return "";
+			}	
+		}
+
+		/// <summary>
+		/// Extract a delimited section from the diagnostics file.
+		/// </summary>
+		/// <param name="Filepath">File name of the diagnostics file to interrogate.</param>
+		/// <param name="StartKey">A string key to mark the start of the selection.</param>
+		/// <param name="EndKey">A string key to mark the end of the selection.</param>
+		/// <returns>A string array of the extracted section.</returns>
+		static string[] ExtractSection( string Filepath, string StartKey, string EndKey )
+		{
+			List<string> Section = new List<string>();
+
+			string[] Lines = File.ReadAllLines( Filepath );
+
+			bool bCollecting = false;
+			foreach( string Line in Lines )
+			{
+				if( Line.StartsWith( StartKey ) )
+				{
+					bCollecting = true;
+				}
+				else if( Line.StartsWith( EndKey ) )
+				{
+					bCollecting = false;
+					break;
+				}
+				else if( bCollecting )
+				{
+					Section.Add( Line.Trim( Environment.NewLine.ToCharArray() ) );
+				}
+			}
+
+			return Section.ToArray();
+		}
+
+		/// <summary>Metadata read from xml file</summary>
+		WERReportMetadata Metadata;
+	}
+}
\ No newline at end of file
diff --git a/Engine/Source/Programs/CrashReporter/CrashReportCommon/WerReport.cs b/Engine/Source/Programs/CrashReporter/CrashReportCommon/WerReport.cs
new file mode 100644
--- /dev/null
+++ b/Engine/Source/Programs/CrashReporter/CrashReportCommon/WerReport.cs
@@ -0,0 +1,172 @@
+// Copyright 1998-2015 Epic Games, Inc. All Rights Reserved.
+
+using System;
+using System.IO;
+using System.Xml.Serialization;
+using System.Collections;
+using System.Collections.Generic;
+using Tools.DotNETCommon.XmlHandler;
+using System.Xml;
+using System.Linq;
+using System.Security;
+
+namespace Tools.CrashReporter.CrashReportCommon
+{
+	/// <summary>
+	/// Descriptive class for the WER report data
+	/// </summary>
+	public class WERReportMetadata
+	{
+		/// <summary>Host operating system version information.</summary>
+		public WERReportMetadataOSVersionInformation OSVersionInformation;
+		/// <summary>Details about the process that spawned the process that crashed.</summary>
+		public WERReportMetadataParentProcessInformation ParentProcessInformation;
+		/// <summary>A set of up to ten key value pairs representing the signature of the crash.</summary>
+		public WERReportMetadataProblemSignatures ProblemSignatures;
+		/// <summary>Additional key value pairs containing information about the crash.</summary>
+		public WERReportMetadataDynamicSignatures DynamicSignatures;
+		/// <summary>Information about the system the crash occurred on.</summary>
+		public WERReportMetadataSystemInformation SystemInformation;
+
+		/// <summary>A simple default constructor to allow Xml serialisation.</summary>
+		public WERReportMetadata()
+		{
+		}
+	}
+
+	/// <summary>
+	/// The problem signature
+	/// </summary>
+	public class WERReportMetadataProblemSignatures
+	{
+		/// <summary>The type of event e.g. crash or hang.</summary>
+		public string EventType = "";
+
+		/// <summary>The application name.</summary>
+		public string Parameter0 = "";
+		/// <summary>The application version.</summary>
+		public string Parameter1 = "";
+		/// <summary>The application timestamp.</summary>
+		public string Parameter2 = "";
+		/// <summary>The faulting module name.</summary>
+		public string Parameter3 = "";
+		/// <summary>The faulting module version.</summary>
+		public string Parameter4 = "";
+		/// <summary>The faulting module timestamp.</summary>
+		public string Parameter5 = "";
+		/// <summary>The exception code.</summary>
+		public string Parameter6 = "";
+		/// <summary>The exception offset.</summary>
+		public string Parameter7 = "";
+		/// <summary>The command line of the crashed process.</summary>
+		public string Parameter8 = "";
+		/// <summary>Branch, BaseDir, and engine mode.</summary>
+		public string Parameter9 = "";
+
+		/// <summary>A simple default constructor to allow Xml serialisation.</summary>
+		public WERReportMetadataProblemSignatures()
+		{
+		}
+	}
+
+	/// <summary>
+	/// The dynamic signature.
+	/// </summary>
+	public class WERReportMetadataDynamicSignatures
+	{
+		/// <summary>The long version of the operating system version.</summary>
+		public string Parameter1 = "";
+		/// <summary>The LCID of the application.</summary>
+		public string Parameter2 = "";
+		/// <summary>Generic parameter</summary>
+		public string Parameter3 = "";
+		/// <summary>Generic parameter</summary>
+		public string Parameter4 = "";
+		/// <summary>Generic parameter</summary>
+		public string Parameter5 = "";
+		/// <summary>Generic parameter</summary>
+		public string Parameter6 = "";
+		/// <summary>Generic parameter</summary>
+		public string Parameter7 = "";
+		/// <summary>Generic parameter</summary>
+		public string Parameter8 = "";
+		/// <summary>Generic parameter</summary>
+		public string Parameter9 = "";
+
+		/// <summary> Whether the user allowed us to be contacted. </summary>
+		public bool bAllowToBeContacted = false;
+
+		/// <summary>A simple default constructor to allow Xml serialisation.</summary>
+		public WERReportMetadataDynamicSignatures()
+		{
+		}
+	}
+
+	/// <summary>
+	/// Details about the host OS the crash occurred on
+	/// </summary>
+	public class WERReportMetadataOSVersionInformation
+	{
+		/// <summary>The Windows version e.g. 6.1.</summary>
+		public string WindowsNTVersion = "";
+		/// <summary>The build number and service pack. e.g. '7601 Service Pack 1'.</summary>
+		public string Build = "";
+		/// <summary>The product name.</summary>
+		public string Product = "";
+		/// <summary>The type of Windows operating system e.g. Professional.</summary>
+		public string Edition = "";
+		/// <summary>The detailed version.e.g. '7601.17944.amd64fre.win7sp1_gdr.120830-0333'.</summary>
+		public string BuildString = "";
+		/// <summary>The revision.</summary>
+		public string Revision = "";
+		/// <summary>The flavor. e.g. 'Multiprocessor Free'.</summary>
+		public string Flavor = "";
+		/// <summary>The processor architecture e.g. x64.</summary>
+		public string Architecture = "";
+		/// <summary>The system language code.</summary>
+		public string LCID = "";
+
+		/// <summary>A simple default constructor to allow Xml serialisation.</summary>
+		public WERReportMetadataOSVersionInformation()
+		{
+		}
+	}
+
+	/// <summary>
+	/// Details about the process that launched the crashed application.
+	/// </summary>
+	public class WERReportMetadataParentProcessInformation
+	{
+		/// <summary>The id of the parent process.</summary>
+		public string ParentProcessId = "";
+		/// <summary>The full path of the parent process.</summary>
+		public string ParentProcessPath = "";
+		/// <summary>The command line used to launch the parent process, including the executable name.</summary>
+		public string ParentProcessCmdLine = "";
+
+		/// <summary>A simple default constructor to allow Xml serialisation.</summary>
+		public WERReportMetadataParentProcessInformation()
+		{
+		}
+	}
+
+	/// <summary>
+	/// Details about the machine the crash occurred on
+	/// </summary>
+	public class WERReportMetadataSystemInformation
+	{
+		/// <summary>A GUID for the machine.</summary>
+		public string MID = "";
+		/// <summary>The manufacturer of the machine e.g. Hewlett Packard.</summary>
+		public string SystemManufacturer = "";
+		/// <summary>The system product name.</summary>
+		public string SystemProductName = "";
+		/// <summary>The current BIOS version.</summary>
+		public string BIOSVersion = "";
+
+		/// <summary>A simple default constructor to allow Xml serialisation.</summary>
+		public WERReportMetadataSystemInformation()
+		{
+		}
+	}
+}
\ No newline at end of file
diff --git a/Engine/Source/Programs/CrashReporter/CrashReportProcess/App.config b/Engine/Source/Programs/CrashReporter/CrashReportProcess/App.config
new file mode 100644
--- /dev/null
+++ b/Engine/Source/Programs/CrashReporter/CrashReportProcess/App.config
@@ -0,0 +1,43 @@
+<?xml version="1.0"?>
+<configuration>
+	<configSections>
+		<sectionGroup name="applicationSettings" type="System.Configuration.ApplicationSettingsGroup, System, Version=4.0.0.0, Culture=neutral, PublicKeyToken=b77a5c561934e089" >
+			<section name="Tools.CrashReporter.CrashReportProcess.Properties.Settings" type="System.Configuration.ClientSettingsSection, System, Version=4.0.0.0, Culture=neutral, PublicKeyToken=b77a5c561934e089" requirePermission="false" />
+		</sectionGroup>
+ </configSections>
+	<connectionStrings />
+	<startup>
+		<supportedRuntime version="v4.0" sku=".NETFramework,Version=v4.0,Profile=Client"/>
+	</startup>
+	<applicationSettings>
+		<Tools.CrashReporter.CrashReportProcess.Properties.Settings>
+   <setting name="ProcessedReports" serializeAs="String">
+    <value>\\devweb-02\Sites\CrashReporterFiles</value>
+   </setting>
+   <setting name="DaysToSunsetReport" serializeAs="String">
+    <value>14</value>
+   </setting>
+   <setting name="CrashReportWebSite" serializeAs="String">
+    <value>CrashReporterUE4</value>
+   </setting>
+   <setting name="ProcessedVideos" serializeAs="String">
+    <value>\\devweb-02\ReportVideos</value>
+   </setting>
+   <setting name="DepotRoot" serializeAs="String">
+    <value>D:\P4EPIC</value>
+   </setting>
+   <setting name="InternalLandingZone" serializeAs="String">
+    <value>D:\!CrashReport\DMZ</value>
+   </setting>
+   <setting name="ExternalLandingZone" serializeAs="String">
+    <value />
+   </setting>
+   <setting name="InvalidReportsDirectory" serializeAs="String">
+    <value>D:\!CrashReport\InvalidCrashReports</value>
+   </setting>
+   <setting name="DepotIndex" serializeAs="String">
+    <value>crash-01</value>
+   </setting>
+  </Tools.CrashReporter.CrashReportProcess.Properties.Settings>
+	</applicationSettings>
+</configuration>
diff --git a/Engine/Source/Programs/CrashReporter/CrashReportProcess/Configs/CrashReportProcess.exe.config b/Engine/Source/Programs/CrashReporter/CrashReportProcess/Configs/CrashReportProcess.exe.config
new file mode 100644
--- /dev/null
+++ b/Engine/Source/Programs/CrashReporter/CrashReportProcess/Configs/CrashReportProcess.exe.config
@@ -0,0 +1,21 @@
+<?xml version="1.0"?>
+<configuration>
+    <configSections>
+        <sectionGroup name="applicationSettings" type="System.Configuration.ApplicationSettingsGroup, System, Version=4.0.0.0, Culture=neutral, PublicKeyToken=b77a5c561934e089" >
+            <section name="Tools.CrashReporter.CrashReportProcess.Properties.Settings" type="System.Configuration.ClientSettingsSection, System, Version=4.0.0.0, Culture=neutral, PublicKeyToken=b77a5c561934e089" requirePermission="false" />
+        </sectionGroup>
+    </configSections>
+    <applicationSettings>
+        <Tools.CrashReporter.CrashReportProcess.Properties.Settings>
+            <setting name="InternalLandingZone" serializeAs="String">
+                <value>F:\InternalCrashReports\CrashReportRepository</value>
+            </setting>
+            <setting name="ExternalLandingZone" serializeAs="String">
+                <value>\\na-web-cluster.dmz.epicgames.net\crash_reporter\CrashReportRepository</value>
+            </setting>
+			<setting name="DepotIndex" serializeAs="String">
+				<value>crash-01</value>
+			</setting>
+        </Tools.CrashReporter.CrashReportProcess.Properties.Settings>
+    </applicationSettings>
+</configuration>
diff --git a/Engine/Source/Programs/CrashReporter/CrashReportProcess/Configs/readme.txt b/Engine/Source/Programs/CrashReporter/CrashReportProcess/Configs/readme.txt
new file mode 100644
--- /dev/null
+++ b/Engine/Source/Programs/CrashReporter/CrashReportProcess/Configs/readme.txt
@@ -0,0 +1 @@
+This config file contains overrides for crash report processor, to allow the internal landing zone to be set without recompiling. It should be copied into the same folder as processor exe.
\ No newline at end of file
diff --git a/Engine/Source/Programs/CrashReporter/CrashReportProcess/CrashReportProcess.csproj b/Engine/Source/Programs/CrashReporter/CrashReportProcess/CrashReportProcess.csproj
new file mode 100644
--- /dev/null
+++ b/Engine/Source/Programs/CrashReporter/CrashReportProcess/CrashReportProcess.csproj
@@ -0,0 +1,127 @@
+﻿<?xml version="1.0" encoding="utf-8"?>
+<Project ToolsVersion="4.0" DefaultTargets="Build" xmlns="http://schemas.microsoft.com/developer/msbuild/2003">
+  <Import Project="$(MSBuildExtensionsPath)\$(MSBuildToolsVersion)\Microsoft.Common.props" Condition="Exists('$(MSBuildExtensionsPath)\$(MSBuildToolsVersion)\Microsoft.Common.props')" />
+  <PropertyGroup>
+    <Configuration Condition=" '$(Configuration)' == '' ">Debug</Configuration>
+    <Platform Condition=" '$(Platform)' == '' ">AnyCPU</Platform>
+    <ProjectGuid>{AE49361C-2EC7-4462-B617-FA363ADD71A6}</ProjectGuid>
+    <OutputType>Exe</OutputType>
+    <AppDesignerFolder>Properties</AppDesignerFolder>
+    <RootNamespace>Tools.CrashReporter.CrashReportProcess</RootNamespace>
+    <AssemblyName>CrashReportProcess</AssemblyName>
+    <TargetFrameworkVersion>v4.5</TargetFrameworkVersion>
+    <FileAlignment>512</FileAlignment>
+    <TargetFrameworkProfile>
+    </TargetFrameworkProfile>
+  </PropertyGroup>
+  <PropertyGroup Condition=" '$(Configuration)|$(Platform)' == 'Debug|AnyCPU' ">
+    <PlatformTarget>AnyCPU</PlatformTarget>
+    <DebugSymbols>true</DebugSymbols>
+    <DebugType>full</DebugType>
+    <Optimize>false</Optimize>
+    <OutputPath>..\..\..\..\Binaries\DotNET\</OutputPath>
+    <DefineConstants>DEBUG;TRACE</DefineConstants>
+    <ErrorReport>prompt</ErrorReport>
+    <WarningLevel>4</WarningLevel>
+    <DocumentationFile>..\..\..\..\Binaries\DotNET\CrashReportProcess.xml</DocumentationFile>
+    <Prefer32Bit>false</Prefer32Bit>
+    <AllowUnsafeBlocks>true</AllowUnsafeBlocks>
+  </PropertyGroup>
+  <PropertyGroup>
+    <StartupObject />
+  </PropertyGroup>
+  <PropertyGroup Condition="'$(Configuration)|$(Platform)' == 'Development|AnyCPU'">
+    <OutputPath>..\..\..\..\Binaries\DotNET\</OutputPath>
+    <DefineConstants>TRACE</DefineConstants>
+    <DocumentationFile>..\..\..\..\Binaries\DotNET\CrashReportProcess.xml</DocumentationFile>
+    <Optimize>true</Optimize>
+    <DebugType>pdbonly</DebugType>
+    <PlatformTarget>AnyCPU</PlatformTarget>
+    <ErrorReport>prompt</ErrorReport>
+    <CodeAnalysisRuleSet>MinimumRecommendedRules.ruleset</CodeAnalysisRuleSet>
+    <RunCodeAnalysis>true</RunCodeAnalysis>
+    <Prefer32Bit>false</Prefer32Bit>
+    <AllowUnsafeBlocks>true</AllowUnsafeBlocks>
+  </PropertyGroup>
+  <PropertyGroup>
+    <SignManifests>true</SignManifests>
+  </PropertyGroup>
+  <PropertyGroup>
+    <ManifestCertificateThumbprint>7B0CF23A7AEEF0637C57BC4BD76F2F43119EFAD6</ManifestCertificateThumbprint>
+  </PropertyGroup>
+  <ItemGroup>
+    <Reference Include="System" />
+    <Reference Include="System.Configuration.Install" />
+    <Reference Include="System.Core" />
+    <Reference Include="System.Data.Linq" />
+    <Reference Include="System.Management" />
+    <Reference Include="System.Xml.Linq" />
+    <Reference Include="System.Data.DataSetExtensions" />
+    <Reference Include="Microsoft.CSharp" />
+    <Reference Include="System.Data" />
+    <Reference Include="System.ServiceProcess" />
+    <Reference Include="System.Xml" />
+  </ItemGroup>
+  <ItemGroup>
+    <Compile Include="CrashReportProcessService.cs">
+      <SubType>Component</SubType>
+    </Compile>
+    <Compile Include="CrashReportProcessService.Designer.cs">
+      <DependentUpon>CrashReportProcessService.cs</DependentUpon>
+    </Compile>
+    <Compile Include="CrashReportProcessServiceInstaller.cs">
+      <SubType>Component</SubType>
+    </Compile>
+    <Compile Include="CrashReportProcessServiceInstaller.Designer.cs">
+      <DependentUpon>CrashReportProcessServiceInstaller.cs</DependentUpon>
+    </Compile>
+    <Compile Include="Program.cs" />
+    <Compile Include="Properties\AssemblyInfo.cs" />
+    <Compile Include="..\..\DotNETCommon\MetaData.cs">
+      <Link>Properties\MetaData.cs</Link>
+    </Compile>
+    <Compile Include="Properties\Settings.Designer.cs">
+      <AutoGen>True</AutoGen>
+      <DesignTimeSharedInput>True</DesignTimeSharedInput>
+      <DependentUpon>Settings.settings</DependentUpon>
+    </Compile>
+    <Compile Include="ReportProcessor.cs" />
+    <Compile Include="ReportQueue.cs" />
+    <Compile Include="ReportWatcher.cs" />
+  </ItemGroup>
+  <ItemGroup>
+    <None Include="App.config" />
+    <None Include="Properties\Settings.settings">
+      <Generator>SettingsSingleFileGenerator</Generator>
+      <LastGenOutput>Settings.Designer.cs</LastGenOutput>
+    </None>
+  </ItemGroup>
+  <ItemGroup>
+  </ItemGroup>
+  <ItemGroup>
+    <ProjectReference Include="..\..\DotNETCommon\DotNETUtilities\DotNETUtilities.csproj">
+      <Project>{5d7d66e8-8c76-4af9-b3ec-2ef03421d730}</Project>
+      <Name>DotNETUtilities</Name>
+    </ProjectReference>
+    <ProjectReference Include="..\CrashReportCommon\CrashReportCommon.csproj">
+      <Project>{b83de8c9-9429-4765-86bd-8eeea89b138f}</Project>
+      <Name>CrashReportCommon</Name>
+      <Private>False</Private>
+    </ProjectReference>
+  </ItemGroup>
+  <ItemGroup>
+    <Service Include="{3259AA49-8AA1-44D3-9025-A0B520596A8C}" />
+  </ItemGroup>
+  <Import Project="$(MSBuildToolsPath)\Microsoft.CSharp.targets" />
+  <PropertyGroup>
+    <PostBuildEvent>
+    </PostBuildEvent>
+  </PropertyGroup>
+  <!-- To modify your build process, add your task inside one of the targets below and uncomment it. 
+       Other similar extension points exist, see Microsoft.Common.targets.
+  <Target Name="BeforeBuild">
+  </Target>
+  <Target Name="AfterBuild">
+  </Target>
+  -->
+</Project>
\ No newline at end of file
diff --git a/Engine/Source/Programs/CrashReporter/CrashReportProcess/CrashReportProcessService.cs b/Engine/Source/Programs/CrashReporter/CrashReportProcess/CrashReportProcessService.cs
new file mode 100644
--- /dev/null
+++ b/Engine/Source/Programs/CrashReporter/CrashReportProcess/CrashReportProcessService.cs
@@ -0,0 +1,145 @@
+﻿// Copyright 1998-2015 Epic Games, Inc. All Rights Reserved.
+
+using System;
+using System.Collections.Generic;
+using System.ComponentModel;
+using System.Data;
+using System.Diagnostics;
+using System.IO;
+using System.Linq;
+using System.Reflection;
+using System.ServiceProcess;
+using System.Text;
+using System.Threading.Tasks;
+
+namespace Tools.CrashReporter.CrashReportProcess
+{
+	/// <summary>
+	/// A class to handle processing received crash reports for displaying on the website.
+	/// </summary>
+	public partial class CrashReporterProcessServicer : ServiceBase
+	{
+		/// <summary>A class the handle detection of new reports.</summary>
+		public ReportWatcher Watcher = null;
+
+		/// <summary>A class to lazily process reports as they are detected.</summary>
+		public ReportProcessor Processor = null;
+
+		/// <summary>Current log file to write debug progress info to</summary>
+		static public CrashReportCommon.LogWriter Log = null;
+
+		/// <summary>Folder in which to store log files</summary>
+		static private string LogFolder = null;
+
+		/// <summary>
+		/// Write a status update to the event log.
+		/// </summary>
+		static public void WriteEvent( string Message )
+		{
+			if( Message != null && Message.Length > 2 )
+			{
+				Log.Print( "[STATUS] " + Message );
+			}
+		}
+
+		/// <summary>
+		/// Write a status update to the event log, from the MinidumpDiagnostics.
+		/// </summary>
+		static public void WriteMDD( string Message )
+		{
+			if( Message != null && Message.Length > 2 )
+			{
+				Log.Print( "[MDD   ] " + Message );
+			}
+		}
+
+		/// <summary>
+		/// Write a status update to the event log, from the P4.
+		/// </summary>
+		static public void WriteP4( string Message )
+		{
+			if( Message != null && Message.Length > 2 )
+			{
+				Log.Print( "[P4    ] " + Message );
+			}
+		}
+
+		/// <summary>
+		/// Write a failure to the event log.
+		/// </summary>
+		static public void WriteFailure( string Message )
+		{
+			if( Message != null && Message.Length > 2 )
+			{
+				Log.Print( "[FAILED] " + Message );
+			}
+		}
+
+		/// <summary>
+		/// Write an exception message to the event log.
+		/// </summary>
+		static public void WriteException( string Message )
+		{
+			if( Message != null && Message.Length > 2 )
+			{
+				Log.Print( "[EXCEPT] " + Message );
+			}
+		}
+
+		/// <summary>
+		/// Initialise all the components, and create an event log.
+		/// </summary>
+		public CrashReporterProcessServicer()
+		{
+			InitializeComponent();
+
+			var StartupPath = Path.GetDirectoryName( Assembly.GetExecutingAssembly().Location );
+			LogFolder = Path.Combine( StartupPath, "Logs" );
+		}
+
+		/// <summary>
+		/// Start the service, and stop the service if there were any errors found.
+		/// </summary>
+		/// <param name="Arguments">Command line arguments (unused).</param>
+		protected override void OnStart( string[] Arguments )
+		{
+			// Create a log file for any start-up messages
+			Log = new CrashReportCommon.LogWriter("CrashReportProcess", LogFolder);
+
+			// Add directory watchers
+			Watcher = new ReportWatcher();
+
+			Processor = new ReportProcessor( Watcher );
+
+			WriteEvent( "Successfully started at " + DateTime.Now.ToString() );
+		}
+
+		/// <summary>
+		/// Stop the service.
+		/// </summary>
+		protected override void OnStop()
+		{
+			// Clean up the directory watcher and crash processor threads
+			Processor.Dispose();
+			Processor = null;
+
+			Watcher.Dispose();
+			Watcher = null;
+
+			// Flush the log to disk
+			Log.Dispose();
+			Log = null;
+		}
+
+		/// <summary>
+		/// Run the service in debug mode. This spews all logging to the console rather than suppressing it.
+		/// </summary>
+		public void DebugRun()
+		{
+			OnStart( null );
+			Console.WriteLine( "Press enter to exit" );
+			Console.Read();
+			OnStop();
+		}
+	}
+}
\ No newline at end of file
diff --git a/Engine/Source/Programs/CrashReporter/CrashReportProcess/CrashReportProcessServiceInstaller.cs b/Engine/Source/Programs/CrashReporter/CrashReportProcess/CrashReportProcessServiceInstaller.cs
new file mode 100644
--- /dev/null
+++ b/Engine/Source/Programs/CrashReporter/CrashReportProcess/CrashReportProcessServiceInstaller.cs
@@ -0,0 +1,50 @@
+﻿// Copyright 1998-2015 Epic Games, Inc. All Rights Reserved.
+
+using System;
+using System.Collections;
+using System.Collections.Generic;
+using System.ComponentModel;
+using System.Configuration.Install;
+using System.Linq;
+using System.ServiceProcess;
+
+namespace Tools.CrashReporter.CrashReportProcess
+{
+	/// <summary>
+	/// The containing class for the Crash Report Processor installer.
+	/// </summary>
+	[RunInstaller( true )]
+	public partial class CrashReportProcessServiceInstaller : System.Configuration.Install.Installer
+	{
+		private ServiceInstaller Installer;
+		private ServiceProcessInstaller ProcessInstaller;
+
+		/// <summary>
+		/// The installer to handle the installation of the Crash Report Processor service.
+		/// </summary>
+		public CrashReportProcessServiceInstaller()
+		{
+			InitializeComponent();
+
+			Installer = new ServiceInstaller();
+			Installer.StartType = ServiceStartMode.Automatic;
+			Installer.ServiceName = "CrashReportProcessService";
+			Installer.DisplayName = "CrashReport Processor Service";
+			Installer.Description = "A web service that processes crash reports for display on the crash report website.";
+			Installers.Add( Installer );
+
+			ProcessInstaller = new ServiceProcessInstaller();
+			ProcessInstaller.Account = ServiceAccount.NetworkService;
+			Installers.Add( ProcessInstaller );
+		}
+
+		/// <summary>
+		/// Must exists, because installation will fail without it.
+		/// </summary>
+		/// <param name="StateSaver">Dictionary to pass down to the base class.</param>
+		public override void Install( IDictionary StateSaver )
+		{
+			base.Install( StateSaver );
+		}
+	}
+}
diff --git a/Engine/Source/Programs/CrashReporter/CrashReportProcess/Program.cs b/Engine/Source/Programs/CrashReporter/CrashReportProcess/Program.cs
new file mode 100644
--- /dev/null
+++ b/Engine/Source/Programs/CrashReporter/CrashReportProcess/Program.cs
@@ -0,0 +1,34 @@
+﻿// Copyright 1998-2015 Epic Games, Inc. All Rights Reserved.
+
+using System;
+using System.ServiceProcess;
+using System.Threading.Tasks;
+
+namespace Tools.CrashReporter.CrashReportProcess
+{
+	/// <summary>
+	/// The container for the Crash Report Processor application.
+	/// </summary>
+	static class CrashReportProcessProgram
+	{
+		/// <summary>
+		/// The main entry point for the application.
+		/// </summary>
+		static void Main()
+		{
+			CrashReporterProcessServicer CrashReporterProcess = new CrashReporterProcessServicer();
+#if !DEBUG
+			if( !Environment.UserInteractive )
+			{
+				// Launch the service as normal if we aren't in the debugger
+				ServiceBase.Run( CrashReporterProcess );
+			}
+			else
+#endif
+			{
+				// Call OnStart, wait for a console key press, call OnStop
+				CrashReporterProcess.DebugRun();
+			}
+		}
+	}
+}
diff --git a/Engine/Source/Programs/CrashReporter/CrashReportProcess/Properties/AssemblyInfo.cs b/Engine/Source/Programs/CrashReporter/CrashReportProcess/Properties/AssemblyInfo.cs
new file mode 100644
--- /dev/null
+++ b/Engine/Source/Programs/CrashReporter/CrashReportProcess/Properties/AssemblyInfo.cs
@@ -0,0 +1,11 @@
+﻿// Copyright 1998-2015 Epic Games, Inc. All Rights Reserved.
+
+using System.Reflection;
+
+// General Information about an assembly is controlled through the following 
+// set of attributes. Change these attribute values to modify the information
+// associated with an assembly.
+[assembly: AssemblyTitle( "CrashReportProcess" )]
+[assembly: AssemblyDescription( "A service to process uploaded crash reports." )]
+[assembly: AssemblyConfiguration( "" )]
+
diff --git a/Engine/Source/Programs/CrashReporter/CrashReportProcess/Properties/Settings.settings b/Engine/Source/Programs/CrashReporter/CrashReportProcess/Properties/Settings.settings
new file mode 100644
--- /dev/null
+++ b/Engine/Source/Programs/CrashReporter/CrashReportProcess/Properties/Settings.settings
@@ -0,0 +1,33 @@
+﻿<?xml version='1.0' encoding='utf-8'?>
+<SettingsFile xmlns="http://schemas.microsoft.com/VisualStudio/2004/01/settings" CurrentProfile="(Default)" GeneratedClassNamespace="Tools.CrashReporter.CrashReportProcess.Properties" GeneratedClassName="Settings">
+  <Profiles />
+  <Settings>
+    <Setting Name="ProcessedReports" Type="System.String" Scope="Application">
+      <Value Profile="(Default)">\\devweb-02\Sites\CrashReporterFiles</Value>
+    </Setting>
+    <Setting Name="DaysToSunsetReport" Type="System.Int32" Scope="Application">
+      <Value Profile="(Default)">14</Value>
+    </Setting>
+    <Setting Name="CrashReportWebSite" Type="System.String" Scope="Application">
+      <Value Profile="(Default)">CrashReporterUE4</Value>
+    </Setting>
+    <Setting Name="ProcessedVideos" Type="System.String" Scope="Application">
+      <Value Profile="(Default)">\\devweb-02\ReportVideos</Value>
+    </Setting>
+    <Setting Name="DepotRoot" Type="System.String" Scope="Application">
+      <Value Profile="(Default)">D:\P4EPIC</Value>
+    </Setting>
+    <Setting Name="InternalLandingZone" Type="System.String" Scope="Application">
+      <Value Profile="(Default)">D:\!CrashReport\DMZ</Value>
+    </Setting>
+    <Setting Name="ExternalLandingZone" Type="System.String" Scope="Application">
+      <Value Profile="(Default)" />
+    </Setting>
+    <Setting Name="InvalidReportsDirectory" Type="System.String" Scope="Application">
+      <Value Profile="(Default)">D:\!CrashReport\InvalidCrashReports</Value>
+    </Setting>
+    <Setting Name="DepotIndex" Type="System.String" Scope="Application">
+      <Value Profile="(Default)">crash-01</Value>
+    </Setting>
+  </Settings>
+</SettingsFile>
\ No newline at end of file
