Copying //UE4/Dev-Core to //UE4/Main

diff --git a/Engine/Source/Programs/CrashReporter/CrashReportProcess/ReportProcessor.cs b/Engine/Source/Programs/CrashReporter/CrashReportProcess/ReportProcessor.cs
new file mode 100644
--- /dev/null
+++ b/Engine/Source/Programs/CrashReporter/CrashReportProcess/ReportProcessor.cs
@@ -0,0 +1,765 @@
+// Copyright 1998-2015 Epic Games, Inc. All Rights Reserved.
+
+using System;
+using System.Collections.Generic;
+using System.Diagnostics;
+using System.Globalization;
+using System.IO;
+using System.Linq;
+using System.Net;
+using System.Text;
+using System.Threading;
+using System.Threading.Tasks;
+using Tools.CrashReporter.CrashReportCommon;
+using Tools.DotNETCommon.LaunchProcess;
+using Tools.DotNETCommon.SimpleWebRequest;
+using Tools.DotNETCommon.XmlHandler;
+
+namespace Tools.CrashReporter.CrashReportProcess
+{
+
+	/// <summary>
+	/// A class to handle processing of received crash reports.
+	/// </summary>
+	public sealed class ReportProcessor : IDisposable
+	{
+		/// <summary>Absolute path to Perforce (default installation assumed)</summary>
+		const string PerforceExePath = "C:\\Program Files\\Perforce\\p4.exe";
+
+		/// <summary>Give up on call stack analysis after 30 minutes</summary>
+		const int MinidumpDiagnosticsTimeoutSeconds = 30 * 60;
+
+		/// <summary>Give up on Perforce syncs after 2 minutes</summary>
+		const int SyncTimeoutSeconds = 2 * 60;
+
+		/// <summary> Number of all processed reports. </summary>
+		public static int ProcessedReports = 1;
+
+		/// <summary> Number of processed reports that have been added to the website. </summary>
+		public static int WebAddedReports = 0;
+
+		/// <summary> Global timer used to measure web added reports per day. </summary>
+		public static Stopwatch Timer = Stopwatch.StartNew();
+
+		/// <summary>The thread that handles detection of new crashes.</summary>
+		public ReportWatcher Watcher = null;
+
+		/// <summary>Task to handle processing.</summary>
+		Task ProcessorTask;
+
+		/// <summary>Object to allow the task to be cancelled at shut-down.</summary>
+		CancellationTokenSource CancelSource;
+
+		/// <summary> Task to handle adding a new record to the database. </summary>
+		Task AddReportTask;
+
+		/// <summary>
+		/// Global initialisation of the processor.
+		/// </summary>
+		/// <param name="InWatcher">The object that watches for new crash reports coming in.</param>
+		public ReportProcessor( ReportWatcher InWatcher )
+		{
+			// Create the thread to handle processing
+			Watcher = InWatcher;
+			CancelSource = new CancellationTokenSource();
+			ProcessNewReports();
+		}
+
+		/// <summary>
+		/// Shutdown: stop the thread
+		/// </summary>
+		public void Dispose()
+		{
+			CancelSource.Cancel();
+			ProcessorTask.Wait();
+	
+			CancelSource.Dispose();
+		}
+
+		/// <summary>
+		/// Main processing thread.
+		/// </summary>
+		/// <remarks>All exceptions are caught and written to the event log.</remarks>
+		private void ProcessNewReports()
+		{
+			// Use the latest MinidumpDiagnostics from the main branch.
+			string Win64BinariesDirectory = Path.Combine( Properties.Settings.Default.DepotRoot, "UE4", "Engine", "Binaries", "Win64" );
+#if !DEBUG
+			string MinidumpDiagnosticsName = Path.Combine(Win64BinariesDirectory, "MinidumpDiagnostics*");
+			if (!SyncRequiredFiles(MinidumpDiagnosticsName))
+			{
+				return;
+			}
+#endif
+
+			DateTime LastCleanupTime = DateTime.Now;
+
+			int LastDay = DateTime.UtcNow.Day;
+
+			var Cancel = CancelSource.Token;
+			ProcessorTask = Task.Factory.StartNew(() =>
+			{
+				while (!Cancel.IsCancellationRequested)
+				{
+					try
+					{
+						foreach( var Queue in Watcher.ReportQueues )
+						{
+							FGenericCrashContext NewContext = null;
+							if (Queue.TryDequeueReport( out NewContext ))
+							{
+								ProcessReport( NewContext );
+
+								int CurrentDay = DateTime.UtcNow.Day;
+								if( CurrentDay > LastDay )
+								{
+									// Check the log and create a new one for a new day.
+									CrashReporterProcessServicer.Log.CreateNewLogFile();
+									LastDay = CurrentDay;
+								}
+
+								// The effect of this break is to prioritize ReportQueues by their order in the list, from highest to lowest
+								break;
+							}
+						}
+					}
+					catch (Exception Ex)
+					{
+						CrashReporterProcessServicer.WriteException( "ProcessNewReports: " + Ex.ToString() );
+					}
+
+					// Remove all folders older than n days to prevent unbounded growth
+					if ((DateTime.Now - LastCleanupTime) > TimeSpan.FromMinutes(15))
+					{
+						CleanRepository();
+						LastCleanupTime = DateTime.Now;
+					}
+
+					// Don't use the CPU if we don't need.
+					Thread.Sleep( 100 );
+				}
+			});
+		}
+
+		/// <summary> 
+		/// Finalizes report files. 
+		/// Thread-safe.
+		/// </summary>
+		private void FinalizeReport( bool bAdded, DirectoryInfo DirInfo, FGenericCrashContext NewContext )
+		{
+			// Only remove if we added the report, otherwise leave all files to further investigation.
+			if( bAdded )
+			{
+				// Remove the report files as we're done with them.
+				CleanReport( DirInfo );
+			}
+			else
+			{
+				MoveReportToInvalid( NewContext.CrashDirectory, NewContext.GetAsFilename() );
+			}
+		}
+
+		/// <summary>
+		/// Delete a report directory.
+		/// Thread-safe.
+		/// </summary>
+		/// <param name="DirInfo">The directory to delete</param>
+		public static void CleanReport(DirectoryInfo DirInfo)
+		{
+			bool bWriteException = true;
+			for( int Retry = 0; Retry < 3; ++Retry )
+			{
+				try
+				{
+					foreach( FileInfo Info in DirInfo.GetFiles() )
+					{
+						Info.IsReadOnly = false;
+					}
+					DirInfo.Delete( true /* delete contents */);
+					break;
+				}
+				catch( Exception Ex )
+				{
+					if( bWriteException )
+					{
+						CrashReporterProcessServicer.WriteException( "CleanReport: " + Ex.ToString() );
+						bWriteException = false;
+					}
+
+				}
+				System.Threading.Thread.Sleep( 100 );
+			}
+		}
+
+		/// <summary>
+		/// Delete a report directory.
+		/// Thread-safe.
+		/// </summary>
+		/// <param name="DirPath">The directory to delete</param>
+		public static void CleanReport(string DirPath)
+		{
+			CleanReport( new DirectoryInfo( DirPath ) );
+		}
+
+		/// <summary> 
+		/// Moves specified report to the directory where are stored invalid reports. 
+		/// Thread-safe.
+		/// </summary>
+		public static void MoveReportToInvalid( string ReportName, string ReportNameAsFilename )
+		{
+			try
+			{
+				DirectoryInfo DirInfo = new DirectoryInfo( ReportName );
+				// Rename the report directory, so we should be able to quickly find the invalid reports.
+				Directory.CreateDirectory( Properties.Settings.Default.InvalidReportsDirectory );
+				string CleanFilename = String.Concat( ReportNameAsFilename.Split( Path.GetInvalidFileNameChars() ) );
+				string DestinationDirectory = Path.Combine( Properties.Settings.Default.InvalidReportsDirectory, CleanFilename );
+
+				Directory.CreateDirectory( DestinationDirectory );
+
+				// Copy all files from the source directory. We can't use MoveTo due to different disc location.
+				foreach( FileInfo File in DirInfo.GetFiles() )
+				{
+					string DestinationFilepath = Path.Combine( DestinationDirectory, File.Name );
+					File.CopyTo( DestinationFilepath, true );
+				}
+				DirInfo.Delete( true );
+
+				CrashReporterProcessServicer.WriteEvent( string.Format( "Moved to {0}", DestinationDirectory ) );
+				UpdateProcessedReports();
+			}
+			catch( System.Exception Ex )
+			{
+				CrashReporterProcessServicer.WriteException( "MoveReportToInvalid: " + Ex.ToString() );
+			}	
+		}
+
+		/// <summary> 
+		/// Update the number of processed reports.
+		/// Thread-safe.
+		/// </summary>
+		private static void UpdateProcessedReports()
+		{
+			Interlocked.Increment( ref ProcessedReports );
+		}
+
+		/// <summary> 
+		/// Delete report folders older than a certain age to avoid unbounded growth of the crash repository.
+		/// </summary>
+		/// <remarks>The folders for the deduplication process older than the property 'DaysToSunsetReport' days old are deleted.</remarks>
+		private void CleanRepository()
+		{
+			try
+			{
+				foreach (var Queue in Watcher.ReportQueues)
+				{
+					Queue.CleanLandingZone();
+				}
+
+				CrashReporterProcessServicer.Log.CleanOutOldLogs(Properties.Settings.Default.DaysToSunsetReport);
+			}
+			catch( Exception Ex )
+			{
+				CrashReporterProcessServicer.WriteException( "CleanRepository: " + Ex.ToString() );
+			}
+		}
+
+		/// <summary>
+		/// Create an Xml payload representing a new crash.
+		/// </summary>
+		/// <param name="DirInfo">The DirectoryInfo of the report folder.</param>
+		/// <param name="NewContext">The generic crash context.</param>
+		/// <param name="bHasVideoFile">Whether the report contains a video file.</param>
+		/// <param name="bHasLog">Whether the report contains a log file.</param>
+		/// <param name="bHasMinidump">Whether the report contains a minidump.</param>
+		/// <returns>A string of Xml payload representing the newly found crash report.</returns>
+		private string CreateCrash( DirectoryInfo DirInfo, FGenericCrashContext NewContext, bool bHasVideoFile, bool bHasLog, bool bHasMinidump )
+		{
+			string XmlPayload = "";
+			try
+			{
+				FEngineVersion EngineVersion = new FEngineVersion( NewContext.PrimaryCrashProperties.EngineVersion ); 
+
+				// Create a new crash description for uploading
+				CrashDescription NewCrash = new CrashDescription();
+
+				NewCrash.BranchName = EngineVersion.GetCleanedBranch();
+				NewCrash.GameName = NewContext.PrimaryCrashProperties.GameName;
+				NewCrash.Platform = NewContext.PrimaryCrashProperties.PlatformFullName;
+				NewCrash.EngineMode = NewContext.PrimaryCrashProperties.EngineMode;
+				NewCrash.BuildVersion = EngineVersion.VersionNumber;
+				NewCrash.CommandLine = NewContext.PrimaryCrashProperties.CommandLine;
+				NewCrash.BaseDir = NewContext.PrimaryCrashProperties.BaseDir;
+
+				NewCrash.Language = NewContext.PrimaryCrashProperties.AppDefaultLocale;
+
+// 				// Create a locate and get the system language.
+// 				int SystemLanguageCode = 0;
+// 				int.TryParse( ReportData.SystemLanguage, out SystemLanguageCode );
+// 				try
+// 				{
+// 					if( SystemLanguageCode > 0 )
+// 					{
+// 						CultureInfo SystemLanguageCI = new CultureInfo( SystemLanguageCode );
+// 						NewCrash.SystemLanguage = SystemLanguageCI.Name;
+// 					}
+// 				}
+// 				catch( System.Exception )
+// 				{
+// 					// Default to en-US
+// 					NewCrash.SystemLanguage = "en-US";
+// 				}
+
+				NewCrash.MachineGuid = NewContext.PrimaryCrashProperties.MachineId; // Valid for all kind of builds, previously only for UE4 releases.
+				NewCrash.UserName = NewContext.PrimaryCrashProperties.UserName; // Only valid for non-UE4 releases.
+				NewCrash.EpicAccountId = NewContext.PrimaryCrashProperties.EpicAccountId; // Only valid for UE4 releases.
+				NewCrash.CallStack = NewContext.PrimaryCrashProperties.GetCallstack();
+				NewCrash.SourceContext = NewContext.PrimaryCrashProperties.GetSourceContext();
+				NewCrash.ErrorMessage = NewContext.PrimaryCrashProperties.GetErrorMessage();
+				NewCrash.UserDescription = NewContext.PrimaryCrashProperties.GetUserDescription();
+				
+				// Iterate through all files and find a file with the earliest date.
+				DateTime TimeOfCrash = DateTime.UtcNow;
+				foreach( var File in DirInfo.GetFiles() )
+				{
+					if( File.CreationTimeUtc < TimeOfCrash )
+					{
+						TimeOfCrash = File.CreationTimeUtc;
+					}
+				}
+
+				//NewCrash.TimeofCrash = NewContext.PrimaryCrashProperties.TimeofCrash;
+				NewCrash.TimeofCrash = TimeOfCrash;
+				
+				NewCrash.bHasMiniDump = bHasMinidump;
+				NewCrash.bHasLog = bHasLog;
+				NewCrash.bHasVideo = bHasVideoFile;
+				NewCrash.BuiltFromCL = (int)EngineVersion.Changelist;
+				NewCrash.bAllowToBeContacted = NewContext.PrimaryCrashProperties.bAllowToBeContacted;
+
+				// Ignore any callstack that is shorter than expected, usually the callstack is invalid.
+				if( NewCrash.CallStack.Length <= CrashReporterConstants.MinCallstackDepth )
+				{
+					CrashReporterProcessServicer.WriteFailure( "! BadStack: BuiltFromCL=" + string.Format( "{0,7}", NewContext.PrimaryCrashProperties.EngineVersion ) + " Path=" + NewContext.CrashDirectory );
+				}
+				else
+				{
+					XmlPayload = XmlHandler.ToXmlString<CrashDescription>( NewCrash );
+				}
+			}
+			catch (Exception Ex)
+			{
+				CrashReporterProcessServicer.WriteException( "CreateCrash: " + Ex.ToString() );
+			}
+
+			return XmlPayload;
+		}
+
+		/// <summary>
+		/// Call the web service function with an Xml payload to add a crash to the database.
+		/// </summary>
+		/// <param name="Payload">Xml representation of a new crash to upload.</param>
+		/// <returns>The database id of the newly added row.</returns>
+		private int UploadCrash( string Payload )
+		{
+			int NewID = -1;
+
+			try
+			{
+				// Simple suppression by blanking out the URL for local testing
+				if( Properties.Settings.Default.CrashReportWebSite.Length > 0 )
+				{
+
+					bool bDebug = false;
+					string RequestString;
+					if( !bDebug )
+					{
+						RequestString = "http://" + Properties.Settings.Default.CrashReportWebSite + ":80/Crashes/AddCrash/-1";
+					}
+					else
+					{ 
+						RequestString = "http://localhost:80/Crashes/AddCrash/-1"; 
+					}
+
+					string ResponseString = SimpleWebRequest.GetWebServiceResponse( RequestString, Payload );
+					if( ResponseString.Length > 0 )
+					{
+						// Convert response into a string
+						CrashReporterResult Result = XmlHandler.FromXmlString<CrashReporterResult>( ResponseString );
+						if( Result.ID > 0 )
+						{
+							NewID = Result.ID;
+						}
+						else
+						{
+							CrashReporterProcessServicer.WriteFailure("UploadCrash: " + Result.Message);
+						}
+					}
+				}
+			}
+			catch( Exception Ex )
+			{
+				CrashReporterProcessServicer.WriteException( "UploadCrash: " + Ex.ToString() );
+			}
+
+			return NewID;
+		}
+
+		/// <summary>
+		/// A function to add a report to the database, and rename the relevant files.
+		/// Thread-safe.
+		/// </summary>
+		/// <param name="DirInfo">The DirectoryInfo of the report folder.</param>
+		/// <param name="NewContext">The generic crash context.</param>
+		/// <param name="LogFileName">The file name of the log file in the report.</param>
+		/// <param name="DumpFileName">The file name of the minidump in the report.</param>
+		/// <param name="VideoFileName">The file name of the video file in the report, or null if there isn't one.</param>
+		private bool AddReport( DirectoryInfo DirInfo, FGenericCrashContext NewContext, string LogFileName, string DumpFileName, string VideoFileName )
+		{
+			try
+			{
+				Stopwatch AddReportTime = Stopwatch.StartNew();
+
+				// Create an XML representation of a crash
+				string CrashDetails = CreateCrash( DirInfo, NewContext, VideoFileName != null, LogFileName != null, DumpFileName != null );
+				if (CrashDetails == "")
+				{
+					CrashReporterProcessServicer.WriteFailure( "! NoDetail: Path=" + NewContext.CrashDirectory );
+					return false;
+				}
+
+				// Upload the crash to the database, and retrieve the new row id
+				int ReportID = UploadCrash(CrashDetails);
+				if (ReportID <= 0)
+				{
+					CrashReporterProcessServicer.WriteFailure( "! NoUpload: Path=" + NewContext.CrashDirectory );
+					return false;
+				}
+
+				string IDThenUnderscore = string.Format("{0}_", ReportID);
+
+				// Use the row id to name and move the files the way the web site requires
+				string DestinationFolder = Path.Combine(Properties.Settings.Default.ProcessedReports, IDThenUnderscore);
+
+				// Move report files to crash reporter file store
+				if (LogFileName != null)
+				{
+					LogFileName = Path.Combine(DirInfo.FullName, LogFileName);
+					FileInfo LogInfo = new FileInfo(LogFileName);
+					if(LogInfo.Exists)
+					{
+						LogInfo.MoveTo(DestinationFolder + "Launch.log");
+					}
+				}
+
+				string CrashContextRuntimeName = Path.Combine( DirInfo.FullName, FGenericCrashContext.CrashContextRuntimeXMLName );
+				FileInfo CrashContextInfo = new FileInfo( CrashContextRuntimeName );
+				if (CrashContextInfo.Exists)
+				{
+					CrashContextInfo.MoveTo( DestinationFolder + FGenericCrashContext.CrashContextRuntimeXMLName );
+				}
+
+// 				WERMetaDataName = Path.Combine(DirInfo.FullName, WERMetaDataName);
+// 				FileInfo MetaInfo = new FileInfo(WERMetaDataName);
+// 				if (MetaInfo.Exists)
+// 				{
+// 					MetaInfo.MoveTo(DestinationFolder + "WERMeta.xml");
+// 				}
+
+				if( DumpFileName != null )
+				{
+					DumpFileName = Path.Combine( DirInfo.FullName, DumpFileName );
+					FileInfo DumpInfo = new FileInfo( DumpFileName );
+					if (DumpInfo.Exists && NewContext.PrimaryCrashProperties.CrashDumpMode != 1 /* ECrashDumpMode.FullDump = 1*/ )
+					{
+						DumpInfo.MoveTo( DestinationFolder + "MiniDump.dmp" );
+					}
+				}
+
+// 				DiagnosticsFileName = Path.Combine(DirInfo.FullName, DiagnosticsFileName);
+// 				FileInfo DiagnosticsInfo = new FileInfo(DiagnosticsFileName);
+// 				if (DiagnosticsInfo.Exists)
+// 				{
+// 					DiagnosticsInfo.MoveTo(DestinationFolder + CrashReporterConstants.DiagnosticsFileName);
+// 				}
+
+				// Move the video (if it exists) to an alternate store
+				if (VideoFileName != null)
+				{
+					DestinationFolder = Path.Combine(Properties.Settings.Default.ProcessedVideos, IDThenUnderscore);
+
+					VideoFileName = Path.Combine(DirInfo.FullName, VideoFileName);
+					FileInfo VideoInfo = new FileInfo(VideoFileName);
+					if (VideoInfo.Exists)
+					{
+						VideoInfo.MoveTo(DestinationFolder + CrashReporterConstants.VideoFileName);
+					}
+				}
+
+				CrashReporterProcessServicer.WriteEvent( "# WebAdded: ReportID   =" + string.Format( "{0,7}", ReportID ) + " Path=" + NewContext.CrashDirectory );
+
+				UpdateProcessedReports();
+				WebAddedReports++;
+				double Ratio = (double)WebAddedReports / (double)ProcessedReports * 100;
+
+				double AddedPerDay = (double)WebAddedReports / Timer.Elapsed.TotalDays;
+
+				CrashReporterProcessServicer.WriteEvent( string.Format( "Ratio={0,2} Processed={1,7} WebAdded={2,7} AddReportTime={3} AddedPerDay={4}", (int)Ratio, ProcessedReports, WebAddedReports, AddReportTime.Elapsed.TotalSeconds.ToString( "0.00" ), (int)AddedPerDay ) );
+				return true;
+			}
+			catch( Exception Ex )
+			{
+				CrashReporterProcessServicer.WriteException( "AddReport: " + DirInfo.Name + "\n\n" + Ex.ToString() );
+			}
+
+			return false;
+		}
+
+		/// <summary>
+		/// Sync the MinidumpDiagnostics binary and the engine config files to #head.
+		/// </summary>
+		/// <param name="MinidumpDiagnosticsName">The name of the binary to sync to process the minidump.</param>
+		/// <returns>true if all the files synced without issue, false otherwise.</returns>
+		/// <remarks>As MinidumpDiagnostics is synced to #head, it requires the engine config files that match to run properly.</remarks>
+		private bool SyncRequiredFiles(string MinidumpDiagnosticsName)
+		{
+			using( var MDDSyncProc = new LaunchProcess( PerforceExePath, null, CrashReporterProcessServicer.WriteP4, "sync", MinidumpDiagnosticsName ) )
+			{
+				if( MDDSyncProc.WaitForExit( SyncTimeoutSeconds * 1000 ) == EWaitResult.TimedOut )
+				{
+					CrashReporterProcessServicer.WriteFailure( "Failed to sync MinidumpDiagnostics " + MinidumpDiagnosticsName );
+
+					return false;
+				}
+			}
+			
+			string ConfigFiles = Path.Combine(Properties.Settings.Default.DepotRoot, "UE4", "Engine/Config/...");
+			using (var ConfigSyncProc = new LaunchProcess(PerforceExePath, null, CrashReporterProcessServicer.WriteP4, "sync", ConfigFiles))
+			{
+				if (ConfigSyncProc.WaitForExit(SyncTimeoutSeconds * 1000) == EWaitResult.TimedOut)
+				{
+					CrashReporterProcessServicer.WriteFailure( "Failed to sync config files " + ConfigFiles );
+
+					return false;
+				}
+			}
+
+			string SSLFiles = Path.Combine(Properties.Settings.Default.DepotRoot, "UE4", "Engine/Binaries/ThirdParty/OpenSSL/..."); // Required by Perforce
+			using (var MiscSyncProc = new LaunchProcess(PerforceExePath, null, CrashReporterProcessServicer.WriteP4, "sync", SSLFiles))
+			{
+				if (MiscSyncProc.WaitForExit(SyncTimeoutSeconds * 1000) == EWaitResult.TimedOut)
+				{
+					CrashReporterProcessServicer.WriteFailure( "Failed to sync OpenSSL files " + SSLFiles );
+
+					return false;
+				}
+			}
+
+			return true;
+		}
+
+		/// <summary>
+		/// A function to process a newly landed crash report.
+		/// </summary>
+		/// <param name="NewContext">An instance of the generic crash context</param>
+		private void ProcessReport( FGenericCrashContext NewContext )
+		{
+			Stopwatch ProcessReportSW = Stopwatch.StartNew();
+
+			// Just to verify that the report is still there.
+			DirectoryInfo DirInfo = new DirectoryInfo( NewContext.CrashDirectory );
+			if( !DirInfo.Exists )
+			{
+				// Something very odd happened.
+				CrashReporterProcessServicer.WriteEvent( "ProcessReport: Directory not found: " + NewContext.CrashDirectory );
+				return;
+			}
+
+			double WaitTime = ReadProcessAddReport( DirInfo, NewContext );
+			
+			// Make sure any log messages have been written to disk
+			CrashReporterProcessServicer.WriteEvent( string.Format( "ProcessReportTime={0} WaitTime={1}", ProcessReportSW.Elapsed.TotalSeconds.ToString( "0.0" ), WaitTime.ToString( "0.00" ) ) );
+		}
+
+		/// <summary> Reads callstack, source context and error message from the diagnostics file, if exists. </summary>
+		private void ReadDiagnosticsFile( FGenericCrashContext NewContext )
+		{
+			// Read the callstack and the source context from the diagnostics files.
+			string CrashDiagnosticsPath = Path.Combine( NewContext.CrashDirectory, CrashReporterConstants.DiagnosticsFileName );
+
+			if (File.Exists( CrashDiagnosticsPath ))
+			{
+				NewContext.PrimaryCrashProperties.CallStack = FGenericCrashContext.EscapeXMLString( string.Join( "\n", FReportData.GetCallStack( CrashDiagnosticsPath ) ) );
+				NewContext.PrimaryCrashProperties.SourceContext = FGenericCrashContext.EscapeXMLString( string.Join( "\n", FReportData.GetSourceContext( CrashDiagnosticsPath ) ) );
+
+				if (NewContext.PrimaryCrashProperties.ErrorMessage.Length == 0)
+				{
+					NewContext.PrimaryCrashProperties.ErrorMessage = FGenericCrashContext.EscapeXMLString( string.Join( "\n", FReportData.GetExceptionDescription( CrashDiagnosticsPath ) ) );
+				}
+			}
+		}
+
+		double ReadProcessAddReport( DirectoryInfo DirInfo, FGenericCrashContext NewContext )
+		{
+			try
+			{
+				string DumpFileName = null;
+				string LogFileName = null;
+				string DiagnosticsFileName = "";
+				string VideoFileName = null;
+
+				string CrashContextName = "";
+
+				foreach( FileInfo Info in DirInfo.GetFiles() )
+				{
+					switch( Info.Extension.ToLower() )
+					{
+					case ".avi":
+					case ".mp4":
+						VideoFileName = Info.Name;
+						break;
+
+					case ".runtime-xml":
+						CrashContextName = Info.Name;
+						break;
+
+					case ".log":
+						LogFileName = Info.Name;
+						break;
+
+					case ".txt":
+						if( Info.Name == CrashReporterConstants.DiagnosticsFileName )
+						{
+							DiagnosticsFileName = Info.Name;
+							ReadDiagnosticsFile( NewContext );				
+						}
+						break;
+
+					case ".dmp":
+					case ".mdmp":
+						DumpFileName = Info.Name;
+						// Check to see if this has been processed locally
+						FileInfo DiagnosticsInfo = new FileInfo( Path.Combine( DirInfo.FullName, CrashReporterConstants.DiagnosticsFileName ) );
+						if (!DiagnosticsInfo.Exists && !NewContext.HasProcessedData())
+						{
+							ProcessDumpFile( Info.FullName, NewContext );
+						}
+
+						// Check to see if a new one has been created
+						DiagnosticsInfo.Refresh();
+						if( DiagnosticsInfo.Exists )
+						{
+							DiagnosticsFileName = DiagnosticsInfo.Name;	
+						}
+						break;
+
+					default:
+						break;
+					}
+				}
+
+				// Check if the new context has processed data.
+				if( NewContext.HasProcessedData() )
+				{
+					Stopwatch WaitSW = Stopwatch.StartNew();
+					// Wait for previous task, should not really happen.
+					if( AddReportTask != null )
+					{
+						AddReportTask.Wait();
+					}
+					double WaitTime = WaitSW.Elapsed.TotalSeconds;
+
+					//bool bAdded = AddReport(DirInfo, NewContext, LogFileName, DumpFileName, VideoFileName );
+					
+					// Save/update crash context to the file.
+					NewContext.ToFile();
+					AddReportTask = Task.Factory.StartNew( () =>
+					{
+						bool bAdded = AddReport( DirInfo, NewContext, LogFileName, DumpFileName, VideoFileName );
+						FinalizeReport( bAdded, DirInfo, NewContext );
+					} );
+
+					return WaitTime;
+				}
+				else
+				{
+					CrashReporterProcessServicer.WriteFailure( "! Invalid : BuiltFromCL=" + string.Format( "{0,7}", NewContext.PrimaryCrashProperties.EngineVersion ) + " Path=" + NewContext.CrashDirectory );
+					FinalizeReport( false, DirInfo, NewContext );
+				}
+			}
+			catch( Exception Ex )
+			{
+				CrashReporterProcessServicer.WriteException( "ProcessReport: " + NewContext.CrashDirectory + "\n\n: " + Ex.ToString() );
+			}
+
+			return 0.0;
+		}
+
+		void ProcessDumpFile( string DiagnosticsPath, FGenericCrashContext NewContext )
+		{
+			// Use the latest MinidumpDiagnostics from the main branch.
+			string Win64BinariesDirectory = Path.Combine(Properties.Settings.Default.DepotRoot, "UE4", "Engine", "Binaries", "Win64");
+#if DEBUG
+			// Note: the debug executable must be built locally or synced from Perforce manually
+			string MinidumpDiagnosticsName = Path.Combine(Win64BinariesDirectory, "MinidumpDiagnostics-Win64-Debug.exe");
+			//string MinidumpDiagnosticsName = Path.Combine(Win64BinariesDirectory, "MinidumpDiagnostics.exe");
+#else
+			string MinidumpDiagnosticsName = Path.Combine(Win64BinariesDirectory, "MinidumpDiagnostics.exe");
+#endif
+			// Purge logs every 2048 processed crashes
+			string PurgeLogsDays = ProcessedReports % 2048 == 0 ? "2" : "-1";
+
+			FEngineVersion EngineVersion = new FEngineVersion( NewContext.PrimaryCrashProperties.EngineVersion ); 
+
+			List<string> MinidumpDiagnosticsParams = new List<string>
+			(
+				new string[] 
+				{
+					"\""+DiagnosticsPath+"\"",
+					"-BranchName=" + EngineVersion.Branch,			// Backward compatibility
+					"-BuiltFromCL=" + EngineVersion.Changelist,		// Backward compatibility
+					"-GameName=" + NewContext.PrimaryCrashProperties.GameName,
+					"-EngineVersion=" + NewContext.PrimaryCrashProperties.EngineVersion,
+					"-bUsePDBCache=true",
+					"-Annotate",
+					"-SyncSymbols",
+					"-SyncMicrosoftSymbols",
+					"-unattended",
+					"-Log=" + NewContext.GetAsFilename() + "-backup-.log",
+					"-DepotIndex=" + Properties.Settings.Default.DepotIndex,
+					"-ini:Engine:[LogFiles]:PurgeLogsDays="+PurgeLogsDays,
+					"-LOGTIMESINCESTART"
+				}
+			);
+
+			LaunchProcess.CaptureMessageDelegate CaptureMessageDelegate = null;
+			if (Environment.UserInteractive)
+			{
+				CaptureMessageDelegate = CrashReporterProcessServicer.WriteMDD;
+			}
+			else
+			{
+				MinidumpDiagnosticsParams.AddRange
+				(
+					new[] 
+					{ 
+						"-buildmachine", 
+						"-forcelogflush" 
+					}
+				);
+
+				// Write some debugging message.
+				CrashReporterProcessServicer.WriteMDD( MinidumpDiagnosticsName + " Params: " + String.Join( " ", MinidumpDiagnosticsParams ) );
+			}
+
+			LaunchProcess ReportParser = new LaunchProcess(MinidumpDiagnosticsName, Path.GetDirectoryName(MinidumpDiagnosticsName), CaptureMessageDelegate, MinidumpDiagnosticsParams.ToArray());
+
+			if (ReportParser.WaitForExit(MinidumpDiagnosticsTimeoutSeconds * 1000) == EWaitResult.TimedOut)
+			{
+				CrashReporterProcessServicer.WriteFailure( "ProcessDumpFile: Timed out running MinidumpDiagnostics" );
+			}
+
+			ReadDiagnosticsFile( NewContext );
+		}
+	}
+}
diff --git a/Engine/Source/Programs/CrashReporter/CrashReportProcess/ReportQueue.cs b/Engine/Source/Programs/CrashReporter/CrashReportProcess/ReportQueue.cs
new file mode 100644
--- /dev/null
+++ b/Engine/Source/Programs/CrashReporter/CrashReportProcess/ReportQueue.cs
@@ -0,0 +1,472 @@
+// Copyright 1998-2015 Epic Games, Inc. All Rights Reserved.
+
+using System;
+using System.Collections.Generic;
+using System.Collections.Concurrent;
+using System.Diagnostics;
+using System.IO;
+using System.IO.Compression;
+using System.Linq;
+
+using Tools.CrashReporter.CrashReportCommon;
+using Tools.DotNETCommon.XmlHandler;
+
+
+namespace Tools.CrashReporter.CrashReportProcess
+{
+	using ReportIdSet = HashSet<string>;
+	using System.Runtime.InteropServices;
+	using System.Globalization;
+
+	/// <summary>
+	/// A queue of pending reports in a specific folder
+	/// </summary>
+	public class ReportQueue
+	{
+		/// <summary>
+		/// Constructor taking the landing zone
+		/// </summary>
+		public ReportQueue(string LandingZonePath)
+		{
+			LandingZone = LandingZonePath;
+		}
+
+		/// <summary>
+		/// Look for new report folders and add them to the publicly available thread-safe queue
+		/// </summary>
+		public void CheckForNewReports()
+		{	
+			try
+			{
+				var NewReportPath = "";
+				var ReportsInLandingZone = new ReportIdSet();
+
+				// Check the landing zones for new reports
+				DirectoryInfo DirInfo = new DirectoryInfo(LandingZone);
+
+				// Couldn't find a landing zone, skip and try again later.
+				// Crash receiver will recreate them if needed.
+				if( !DirInfo.Exists )
+				{
+					CrashReporterProcessServicer.WriteFailure( "LandingZone not found: " + LandingZone );
+					return;
+				}
+
+				// Add any new reports
+				foreach (var SubDirInfo in DirInfo.GetDirectories())
+				{
+					try
+					{
+						if( Directory.Exists( SubDirInfo.FullName ) )
+						{
+							NewReportPath = SubDirInfo.FullName;
+							ReportsInLandingZone.Add( NewReportPath );
+							if( !ReportsInLandingZoneLastTimeWeChecked.Contains( NewReportPath ) )
+							{
+								EnqueueNewReport( NewReportPath );
+							}
+						}
+					}
+					catch( System.Exception Ex )
+					{
+						CrashReporterProcessServicer.WriteException( "CheckForNewReportsInner: " + Ex.ToString() );
+						ReportProcessor.MoveReportToInvalid( NewReportPath, "NEWRECORD_FAIL_" + DateTime.Now.Ticks );
+					}
+				}
+				//CrashReporterProcessServicer.WriteEvent( string.Format( "ReportsInLandingZone={0}, ReportsInLandingZoneLastTimeWeChecked={1}", ReportsInLandingZone.Count, ReportsInLandingZoneLastTimeWeChecked.Count ) );
+				ReportsInLandingZoneLastTimeWeChecked = ReportsInLandingZone;
+			}
+			catch( Exception Ex )
+			{
+				CrashReporterProcessServicer.WriteException( "CheckForNewReportsOuter: " + Ex.ToString() );
+			}
+		}
+
+		/// <summary>
+		/// Tidy up the landing zone folder
+		/// </summary>
+		public void CleanLandingZone()
+		{
+			var DirInfo = new DirectoryInfo(LandingZone);
+			foreach(var SubDirInfo in DirInfo.GetDirectories())
+			{
+				if (SubDirInfo.LastWriteTimeUtc.AddDays(Properties.Settings.Default.DaysToSunsetReport) < DateTime.UtcNow)
+				{
+					SubDirInfo.Delete(true);
+				}
+			}
+		}
+
+		/// <summary> Looks for the CrashContext.runtime-xml, if found, will return a new instance of the FGenericCrashContext. </summary>
+		private FGenericCrashContext FindCrashContext( string NewReportPath )
+		{
+			string CrashContextPath = Path.Combine( NewReportPath, FGenericCrashContext.CrashContextRuntimeXMLName );
+			bool bExist = File.Exists( CrashContextPath );
+
+			if( bExist )
+			{
+				return FGenericCrashContext.FromFile( CrashContextPath );
+			}
+
+			return null;
+		}
+
+		/// <summary> Converts WER metadata xml file to the crash context. </summary>
+		private FGenericCrashContext ConvertMetadataToCrashContext( WERReportMetadata Metadata, string NewReportPath )
+		{
+			FGenericCrashContext CrashContext = new FGenericCrashContext();
+
+			FReportData ReportData = new FReportData( Metadata, NewReportPath );
+
+			CrashContext.PrimaryCrashProperties.CrashVersion = (int)ECrashDescVersions.VER_1_NewCrashFormat;
+
+			CrashContext.PrimaryCrashProperties.ProcessId = 0;
+			CrashContext.PrimaryCrashProperties.CrashGUID = new DirectoryInfo( NewReportPath ).Name;  
+// 			CrashContext.PrimaryCrashProperties.IsInternalBuild
+// 			CrashContext.PrimaryCrashProperties.IsPerforceBuild
+// 			CrashContext.PrimaryCrashProperties.IsSourceDistribution
+// 			CrashContext.PrimaryCrashProperties.SecondsSinceStart
+			CrashContext.PrimaryCrashProperties.GameName = ReportData.GameName;
+// 			CrashContext.PrimaryCrashProperties.ExecutableName
+// 			CrashContext.PrimaryCrashProperties.BuildConfiguration
+// 			CrashContext.PrimaryCrashProperties.PlatformName
+// 			CrashContext.PrimaryCrashProperties.PlatformNameIni
+			CrashContext.PrimaryCrashProperties.PlatformFullName = ReportData.Platform;
+			CrashContext.PrimaryCrashProperties.EngineMode = ReportData.EngineMode;
+			CrashContext.PrimaryCrashProperties.EngineVersion = ReportData.GetEngineVersion();
+			CrashContext.PrimaryCrashProperties.CommandLine = ReportData.CommandLine;
+// 			CrashContext.PrimaryCrashProperties.LanguageLCID
+
+			// Create a locate and get the language.
+			int LanguageCode = 0;
+			int.TryParse( ReportData.Language, out LanguageCode );
+			try
+			{
+				if (LanguageCode > 0)
+				{
+					CultureInfo LanguageCI = new CultureInfo( LanguageCode );
+					CrashContext.PrimaryCrashProperties.AppDefaultLocale = LanguageCI.Name;
+				}
+			}
+			catch (System.Exception)
+			{
+				// Default to en-US
+				CrashContext.PrimaryCrashProperties.AppDefaultLocale = "en-US";
+			}
+
+
+// 			CrashContext.PrimaryCrashProperties.IsUE4Release
+			CrashContext.PrimaryCrashProperties.UserName = ReportData.UserName;
+			CrashContext.PrimaryCrashProperties.BaseDir = ReportData.BaseDir;
+// 			CrashContext.PrimaryCrashProperties.RootDir
+			CrashContext.PrimaryCrashProperties.MachineId = ReportData.MachineId;
+			CrashContext.PrimaryCrashProperties.EpicAccountId = ReportData.EpicAccountId;
+// 			CrashContext.PrimaryCrashProperties.CallStack
+// 			CrashContext.PrimaryCrashProperties.SourceContext
+			CrashContext.PrimaryCrashProperties.UserDescription = string.Join( "\n", ReportData.UserDescription );
+			CrashContext.PrimaryCrashProperties.ErrorMessage = string.Join( "\n", ReportData.ErrorMessage );
+// 			CrashContext.PrimaryCrashProperties.CrashDumpMode
+// 			CrashContext.PrimaryCrashProperties.Misc.NumberOfCores
+// 			CrashContext.PrimaryCrashProperties.Misc.NumberOfCoresIncludingHyperthreads
+// 			CrashContext.PrimaryCrashProperties.Misc.Is64bitOperatingSystem
+// 			CrashContext.PrimaryCrashProperties.Misc.CPUVendor
+// 			CrashContext.PrimaryCrashProperties.Misc.CPUBrand
+// 			CrashContext.PrimaryCrashProperties.Misc.PrimaryGPUBrand
+// 			CrashContext.PrimaryCrashProperties.Misc.OSVersionMajor
+// 			CrashContext.PrimaryCrashProperties.Misc.OSVersionMinor
+// 			CrashContext.PrimaryCrashProperties.Misc.AppDiskTotalNumberOfBytes
+// 			CrashContext.PrimaryCrashProperties.Misc.AppDiskNumberOfFreeBytes
+// 			CrashContext.PrimaryCrashProperties.MemoryStats.TotalPhysical
+// 			CrashContext.PrimaryCrashProperties.MemoryStats.TotalVirtual
+// 			CrashContext.PrimaryCrashProperties.MemoryStats.PageSize
+// 			CrashContext.PrimaryCrashProperties.MemoryStats.TotalPhysicalGB
+// 			CrashContext.PrimaryCrashProperties.MemoryStats.AvailablePhysical
+// 			CrashContext.PrimaryCrashProperties.MemoryStats.AvailableVirtual
+// 			CrashContext.PrimaryCrashProperties.MemoryStats.UsedPhysical
+// 			CrashContext.PrimaryCrashProperties.MemoryStats.PeakUsedPhysical
+// 			CrashContext.PrimaryCrashProperties.MemoryStats.UsedVirtual
+// 			CrashContext.PrimaryCrashProperties.MemoryStats.PeakUsedVirtual
+// 			CrashContext.PrimaryCrashProperties.MemoryStats.bIsOOM
+// 			CrashContext.PrimaryCrashProperties.MemoryStats.OOMAllocationSize
+// 			CrashContext.PrimaryCrashProperties.MemoryStats.OOMAllocationAlignment
+			CrashContext.PrimaryCrashProperties.TimeofCrash = new DateTime( ReportData.Ticks );
+			CrashContext.PrimaryCrashProperties.bAllowToBeContacted = ReportData.AllowToBeContacted;
+
+			return CrashContext;
+		}
+
+		/// <summary> Looks for the WER metadata xml file, if found, will return a new instance of the WERReportMetadata. </summary>
+		private WERReportMetadata FindMetadata( string NewReportPath )
+		{
+			WERReportMetadata MetaData = null;
+
+			// Just to verify that the report is still there.
+			DirectoryInfo DirInfo = new DirectoryInfo( NewReportPath );
+			if( !DirInfo.Exists )
+			{
+				CrashReporterProcessServicer.WriteFailure( "FindMetadata: Directory not found " + NewReportPath );
+			}
+			else
+			{
+				// Check to see if we wish to suppress processing of this report.	
+				foreach( var Info in DirInfo.GetFiles( "*.xml" ) )
+				{
+					var MetaDataToCheck = XmlHandler.ReadXml<WERReportMetadata>( Info.FullName );
+					if( CheckMetaData( MetaDataToCheck ) )
+					{
+						MetaData = MetaDataToCheck;
+						break;
+					}
+				}
+			}
+			
+			return MetaData;
+		}
+
+		/// <summary> Looks for the Diagnostics.txt file and returns true if exists. </summary>
+		private bool FindDiagnostics( string NewReportPath )
+		{
+			bool bFound = false;
+
+			DirectoryInfo DirInfo = new DirectoryInfo( NewReportPath );
+			foreach( var Info in DirInfo.GetFiles( CrashReporterConstants.DiagnosticsFileName ) )
+			{
+				bFound = true;
+				break;
+			}
+			return bFound;
+		}
+
+		/// <summary>
+		/// Optionally don't process some reports based on the Windows Error report meta data.
+		/// </summary>
+		/// <param name="WERData">The Windows Error Report meta data.</param>
+		/// <returns>false to reject the report.</returns>
+		private bool CheckMetaData( WERReportMetadata WERData )
+		{
+			if( WERData == null )
+			{
+				return false;
+			}
+
+			// Reject any crashes with the invalid metadata.
+			if( WERData.ProblemSignatures == null || WERData.DynamicSignatures == null || WERData.OSVersionInformation == null )
+			{
+				return false;
+			}
+
+			// Reject any crashes from the minidump processor.
+			if( WERData.ProblemSignatures.Parameter0 != null && WERData.ProblemSignatures.Parameter0.ToLower() == "MinidumpDiagnostics".ToLower() )
+			{
+				return false;
+			}
+
+			return true;
+		}
+
+		// From CrashUpload.cpp
+		/*
+		struct FCompressedCrashFile : FNoncopyable
+		{
+			int32 CurrentFileIndex; // 4 bytes for file index
+			FString Filename; // 4 bytes for length + 260 bytes for char data
+			TArray<uint8> Filedata; // 4 bytes for length + N bytes for data
+		}
+		*/
+	
+		[DllImport( "CompressionHelper.dll", CallingConvention = CallingConvention.Cdecl )]
+		private extern static bool UE4UncompressMemoryZLIB( /*void**/IntPtr UncompressedBuffer, Int32 UncompressedSize, /*const void**/IntPtr CompressedBuffer, Int32 CompressedSize );
+
+		/// <summary> 
+		/// Decompresses a compressed crash report. 
+		/// </summary>
+		unsafe private bool DecompressReport( string CompressedReportPath, FCompressedCrashInformation Meta )
+		{
+			string ReportPath = Path.GetDirectoryName( CompressedReportPath );
+			using (FileStream FileStream = File.OpenRead( CompressedReportPath ))
+			{
+				Int32 UncompressedSize = Meta.GetUncompressedSize();
+				Int32 CompressedSize = Meta.GetCompressedSize();
+
+				if (FileStream.Length != CompressedSize)
+				{
+					return false;
+				}
+
+				byte[] UncompressedBufferArray = new byte[UncompressedSize];
+				byte[] CompressedBufferArray = new byte[CompressedSize];
+
+				FileStream.Read( CompressedBufferArray, 0, CompressedSize );
+
+				fixed (byte* UncompressedBufferPtr = UncompressedBufferArray, CompressedBufferPtr = CompressedBufferArray)
+				{
+					bool bResult = UE4UncompressMemoryZLIB( (IntPtr)UncompressedBufferPtr, UncompressedSize, (IntPtr)CompressedBufferPtr, CompressedSize );
+					if (!bResult)
+					{
+						CrashReporterProcessServicer.WriteFailure( "! ZLibFail: Path=" + ReportPath );
+						return false;
+					}
+				}
+
+				MemoryStream MemoryStream = new MemoryStream( UncompressedBufferArray, false );
+				BinaryReader BinaryData = new BinaryReader( MemoryStream );
+
+				for (int FileIndex = 0; FileIndex < Meta.GetNumberOfFiles(); FileIndex++)
+				{
+					Int32 CurrentFileIndex = BinaryData.ReadInt32();
+					if (CurrentFileIndex != FileIndex)
+					{
+						CrashReporterProcessServicer.WriteFailure( "! ReadFail: Required=" + FileIndex + " Read=" + CurrentFileIndex );
+						return false;
+					}
+
+					string Filename = FBinaryReaderHelper.ReadFixedSizeString( BinaryData );
+					Int32 FiledataLength = BinaryData.ReadInt32();
+					byte[] Filedata = BinaryData.ReadBytes( FiledataLength );
+
+					File.WriteAllBytes( Path.Combine( ReportPath, Filename ), Filedata );
+				}
+			}
+
+			return true;
+		}
+
+		/// <summary> Enqueues a new crash. </summary>
+		void EnqueueNewReport( string NewReportPath )
+		{
+			string ReportName = Path.GetFileName( NewReportPath );
+
+			string CompressedReportPath = Path.Combine( NewReportPath, ReportName + ".ue4crash" );
+			string MetadataPath = Path.Combine( NewReportPath, ReportName + ".xml" );
+			bool bIsCompressed = File.Exists( CompressedReportPath ) && File.Exists( MetadataPath );
+			if (bIsCompressed)
+			{
+				FCompressedCrashInformation CompressedCrashInformation = XmlHandler.ReadXml<FCompressedCrashInformation>( MetadataPath );
+				bool bResult = DecompressReport( CompressedReportPath, CompressedCrashInformation );
+				if (!bResult)
+				{
+					ReportProcessor.CleanReport( new DirectoryInfo( NewReportPath ) );
+					ReportProcessor.MoveReportToInvalid( NewReportPath, "DECOMPRESSION_FAIL_" + DateTime.Now.Ticks + "_" + ReportName );
+					return;
+				}
+				else
+				{
+					// Rename metadata file to not interfere with the WERReportMetadata.
+					File.Move( MetadataPath, Path.ChangeExtension( MetadataPath, "processed_xml" ) );
+				}
+			}
+
+			// Unified crash reporting
+			FGenericCrashContext GenericContext = FindCrashContext( NewReportPath );
+			FGenericCrashContext Context = GenericContext;
+
+			bool bFromWER = false;
+			if (Context == null || !Context.HasProcessedData())
+			{
+				WERReportMetadata MetaData = FindMetadata( NewReportPath );
+				if (MetaData != null)
+				{
+					FReportData ReportData = new FReportData( MetaData, NewReportPath );
+					Context = ConvertMetadataToCrashContext( MetaData, NewReportPath );
+					bFromWER = true;
+				}
+			}
+
+			if (Context == null)
+			{
+				CrashReporterProcessServicer.WriteFailure( "! NoCntx  : Path=" + NewReportPath );
+				ReportProcessor.CleanReport( new DirectoryInfo( NewReportPath ) );
+			}
+			else
+			{
+				if (GenericContext != null && GenericContext.PrimaryCrashProperties.ErrorMessage.Length > 0)
+				{
+					// Get error message from the crash context and fix value in the metadata.
+					Context.PrimaryCrashProperties.ErrorMessage = GenericContext.PrimaryCrashProperties.ErrorMessage;
+				}
+
+				Context.CrashDirectory = NewReportPath;
+				Context.PrimaryCrashProperties.SetPlatformFullName();
+
+				// If based on WER, save to the file.
+				if (bFromWER && GenericContext == null)
+				{
+					Context.ToFile();
+				}
+				
+				FEngineVersion EngineVersion = new FEngineVersion( Context.PrimaryCrashProperties.EngineVersion ); 
+
+				uint BuiltFromCL = EngineVersion.Changelist;
+
+				string BranchName = EngineVersion.Branch;
+				if (string.IsNullOrEmpty( BranchName ))
+				{
+					CrashReporterProcessServicer.WriteFailure( "% NoBranch: BuiltFromCL=" + string.Format( "{0,7}", BuiltFromCL ) + " Path=" + NewReportPath );
+					ReportProcessor.MoveReportToInvalid( NewReportPath, Context.GetAsFilename() );
+					return;
+				}
+
+				if (BranchName.Equals( CrashReporterConstants.LicenseeBranchName, StringComparison.InvariantCultureIgnoreCase ))
+				{
+					CrashReporterProcessServicer.WriteFailure( "% UE4-QA  : BuiltFromCL=" + string.Format( "{0,7}", BuiltFromCL ) + " Path=" + NewReportPath );
+					ReportProcessor.CleanReport( NewReportPath );
+					return;
+				}
+
+				// Look for the Diagnostics.txt, if we have a diagnostics file we can continue event if the CL is marked as broken.
+				bool bHasDiagnostics = FindDiagnostics( NewReportPath );
+
+				if (BuiltFromCL == 0 && ( !bHasDiagnostics && !Context.HasProcessedData() ))
+				{
+					// For now ignore all locally made crashes.
+					CrashReporterProcessServicer.WriteFailure( "! BROKEN0 : BuiltFromCL=" + string.Format( "{0,7}", BuiltFromCL ) + " Path=" + NewReportPath );
+					ReportProcessor.CleanReport( NewReportPath );
+					return;
+				}
+
+
+				lock (NewReportsLock)
+				{
+					NewCrashContexts.Add( Context );
+				}
+				CrashReporterProcessServicer.WriteEvent( "+ Enqueued: BuiltFromCL=" + string.Format( "{0,7}", BuiltFromCL ) + " Path=" + NewReportPath );
+			}
+			
+		}
+
+		/// <summary> Tries to dequeue a report from the list. </summary>
+		public bool TryDequeueReport( out FGenericCrashContext Context )
+		{
+			lock( NewReportsLock )
+			{
+				int LastIndex = NewCrashContexts.Count - 1;
+
+				if( LastIndex >= 0 )
+				{
+					Context = NewCrashContexts[LastIndex];
+					NewCrashContexts.RemoveAt( LastIndex );
+					CrashReporterProcessServicer.WriteEvent( "- Dequeued: BuiltFromCL=" + string.Format( "{0,7}", Context.PrimaryCrashProperties.EngineVersion ) + " Path=" + Context.CrashDirectory );
+					return true;
+				}
+				else
+				{
+					Context = null;
+					return false;
+				}
+			}
+		}
+
+
+		/// <summary> A list of freshly landed crash reports. </summary>
+		private List<FGenericCrashContext> NewCrashContexts = new List<FGenericCrashContext>();
+
+		/// <summary> Object used to synchronize the access to the NewReport. </summary>
+		private Object NewReportsLock = new Object();
+
+		/// <summary>Incoming report path</summary>
+		string LandingZone;
+
+		/// <summary>Report IDs that were in the folder on the last pass</summary>
+		ReportIdSet ReportsInLandingZoneLastTimeWeChecked = new ReportIdSet();
+	}
+}
diff --git a/Engine/Source/Programs/CrashReporter/CrashReportProcess/ReportWatcher.cs b/Engine/Source/Programs/CrashReporter/CrashReportProcess/ReportWatcher.cs
new file mode 100644
--- /dev/null
+++ b/Engine/Source/Programs/CrashReporter/CrashReportProcess/ReportWatcher.cs
@@ -0,0 +1,97 @@
+// Copyright 1998-2015 Epic Games, Inc. All Rights Reserved.
+
+using System;
+using System.Collections.Generic;
+using System.Threading;
+using System.Threading.Tasks;
+
+namespace Tools.CrashReporter.CrashReportProcess
+{
+	/// <summary>
+	/// A class to monitor the crash report repository for new crash reports.
+	/// </summary>
+	public sealed class ReportWatcher : IDisposable
+	{
+		/// <summary>A queue of freshly landed crash reports ready to be processed.</summary>
+		public List<ReportQueue> ReportQueues = new List<ReportQueue>();
+
+		/// <summary>Task to periodically check for the arrival of new report folders.</summary>
+		Task WatcherTask;
+
+		/// <summary>Object providing cancellation of the watcher task.</summary>
+		CancellationTokenSource CancelSource;
+
+		/// <summary>
+		/// Create a directory watcher to monitor the crash report landing zone.
+		/// </summary>
+		public ReportWatcher()
+		{
+			CancelSource = new CancellationTokenSource();
+			Start();
+		}
+
+		/// <summary>
+		/// Shutdown: stop the thread
+		/// </summary>
+		public void Dispose()
+		{
+			// Cancel the task and wait for it to quit
+			CancelSource.Cancel();
+			WatcherTask.Wait();
+
+			CancelSource.Dispose();
+		}
+
+		/// <summary>
+		/// A thread to watch for new crash reports landing.
+		/// </summary>
+		/// <remarks>The NFS storage does not support file system watchers, so this has to be done laboriously.</remarks>
+		void Start()
+		{
+			CrashReporterProcessServicer.WriteEvent("CrashReportProcessor watching directories:");
+			var Settings = Properties.Settings.Default;
+
+			if (!string.IsNullOrEmpty(Settings.InternalLandingZone))
+			{
+				if( System.IO.Directory.Exists( Settings.InternalLandingZone ) )
+				{
+					ReportQueues.Add(new ReportQueue(Settings.InternalLandingZone));
+					CrashReporterProcessServicer.WriteEvent(string.Format("\t{0} (internal, high priority)", Settings.InternalLandingZone));
+				}
+				else
+				{
+					CrashReporterProcessServicer.WriteFailure( string.Format( "\t{0} (internal, high priority) is not accessible", Settings.InternalLandingZone ) );
+				}
+			}
+
+#if !DEBUG
+			if (!string.IsNullOrEmpty(Settings.ExternalLandingZone))
+			{
+				if( System.IO.Directory.Exists( Settings.ExternalLandingZone ) )
+				{
+					ReportQueues.Add( new ReportQueue( Settings.ExternalLandingZone ) );
+					CrashReporterProcessServicer.WriteEvent( string.Format( "\t{0}", Settings.ExternalLandingZone ) );
+				}
+				else
+				{
+					CrashReporterProcessServicer.WriteFailure( string.Format( "\t{0} is not accessible", Settings.ExternalLandingZone ) );
+				}
+			}
+#endif //!DEBUG
+
+			var Cancel = CancelSource.Token;
+			WatcherTask = Task.Factory.StartNew(async () =>
+			{
+				while (!Cancel.IsCancellationRequested)
+				{
+					// Check the landing zones for new reports
+					foreach (var Queue in ReportQueues)
+					{
+						Queue.CheckForNewReports();
+					}
+					await Task.Delay(60000, Cancel);
+				}
+			});
+		}
+	}
+}
\ No newline at end of file
diff --git a/Engine/Source/Programs/CrashReporter/CrashReportReceiver/App.config b/Engine/Source/Programs/CrashReporter/CrashReportReceiver/App.config
new file mode 100644
--- /dev/null
+++ b/Engine/Source/Programs/CrashReporter/CrashReportReceiver/App.config
@@ -0,0 +1,24 @@
+<?xml version="1.0"?>
+<configuration>
+    <configSections>
+        <sectionGroup name="applicationSettings" type="System.Configuration.ApplicationSettingsGroup, System, Version=4.0.0.0, Culture=neutral, PublicKeyToken=b77a5c561934e089">
+            <section name="Tools.CrashReporter.CrashReportReceiver.Properties.Settings" type="System.Configuration.ClientSettingsSection, System, Version=4.0.0.0, Culture=neutral, PublicKeyToken=b77a5c561934e089" requirePermission="false"/>
+        </sectionGroup>
+    </configSections>
+    <startup> 
+        <supportedRuntime version="v4.0" sku=".NETFramework,Version=v4.0"/>
+    </startup>
+    <applicationSettings>
+        <Tools.CrashReporter.CrashReportReceiver.Properties.Settings>
+            <setting name="MaxAllowedBandwidthMB" serializeAs="String">
+                <value>10</value>
+            </setting>
+            <setting name="MaxAllowedCPUtilizationPct" serializeAs="String">
+                <value>75</value>
+            </setting>
+            <setting name="CrashReportRepository" serializeAs="String">
+                <value>U:\CRASHREPORT\DMZ</value>
+            </setting>
+        </Tools.CrashReporter.CrashReportReceiver.Properties.Settings>
+    </applicationSettings>
+</configuration>
diff --git a/Engine/Source/Programs/CrashReporter/CrashReportReceiver/Configs/CrashReportReceiver.exe.config b/Engine/Source/Programs/CrashReporter/CrashReportReceiver/Configs/CrashReportReceiver.exe.config
new file mode 100644
--- /dev/null
+++ b/Engine/Source/Programs/CrashReporter/CrashReportReceiver/Configs/CrashReportReceiver.exe.config
@@ -0,0 +1,24 @@
+<?xml version="1.0"?>
+<configuration>
+	<configSections>
+		<sectionGroup name="applicationSettings" type="System.Configuration.ApplicationSettingsGroup, System, Version=4.0.0.0, Culture=neutral, PublicKeyToken=b77a5c561934e089">
+			<section name="Tools.CrashReporter.CrashReportReceiver.Properties.Settings" type="System.Configuration.ClientSettingsSection, System, Version=4.0.0.0, Culture=neutral, PublicKeyToken=b77a5c561934e089" requirePermission="false"/>
+		</sectionGroup>
+	</configSections>
+	<startup>
+		<supportedRuntime version="v4.0" sku=".NETFramework,Version=v4.0"/>
+	</startup>
+	<applicationSettings>
+		<Tools.CrashReporter.CrashReportReceiver.Properties.Settings>
+			<setting name="MaxAllowedBandwidthMB" serializeAs="String">
+				<value>10</value>
+			</setting>
+			<setting name="MaxAllowedCPUtilizationPct" serializeAs="String">
+				<value>75</value>
+			</setting>
+			<setting name="CrashReportRepository" serializeAs="String">
+				<value>D:!\CRASHREPORT\DMZ</value>
+			</setting>
+		</Tools.CrashReporter.CrashReportReceiver.Properties.Settings>
+	</applicationSettings>
+</configuration>
diff --git a/Engine/Source/Programs/CrashReporter/CrashReportReceiver/Configs/readme.txt b/Engine/Source/Programs/CrashReporter/CrashReportReceiver/Configs/readme.txt
new file mode 100644
--- /dev/null
+++ b/Engine/Source/Programs/CrashReporter/CrashReportReceiver/Configs/readme.txt
@@ -0,0 +1 @@
+This config file contains overrides for the internal crash report receiver. It should be copied into the same folder as receiver exe on the machine that receives internal reports.
\ No newline at end of file
diff --git a/Engine/Source/Programs/CrashReporter/CrashReportReceiver/CrashReportReceiver.csproj b/Engine/Source/Programs/CrashReporter/CrashReportReceiver/CrashReportReceiver.csproj
new file mode 100644
--- /dev/null
+++ b/Engine/Source/Programs/CrashReporter/CrashReportReceiver/CrashReportReceiver.csproj
@@ -0,0 +1,112 @@
+﻿<?xml version="1.0" encoding="utf-8"?>
+<Project ToolsVersion="4.0" DefaultTargets="Build" xmlns="http://schemas.microsoft.com/developer/msbuild/2003">
+  <Import Project="$(MSBuildExtensionsPath)\$(MSBuildToolsVersion)\Microsoft.Common.props" Condition="Exists('$(MSBuildExtensionsPath)\$(MSBuildToolsVersion)\Microsoft.Common.props')" />
+  <PropertyGroup>
+    <Configuration Condition=" '$(Configuration)' == '' ">Debug</Configuration>
+    <Platform Condition=" '$(Platform)' == '' ">AnyCPU</Platform>
+    <ProjectGuid>{A38E68D6-0CB0-450C-A753-764651D7E2B6}</ProjectGuid>
+    <OutputType>Exe</OutputType>
+    <AppDesignerFolder>Properties</AppDesignerFolder>
+    <RootNamespace>Tools.CrashReporter.CrashReportReceiver</RootNamespace>
+    <AssemblyName>CrashReportReceiver</AssemblyName>
+    <TargetFrameworkVersion>v4.5</TargetFrameworkVersion>
+    <FileAlignment>512</FileAlignment>
+    <TargetFrameworkProfile>
+    </TargetFrameworkProfile>
+  </PropertyGroup>
+  <PropertyGroup Condition=" '$(Configuration)|$(Platform)' == 'Debug|AnyCPU' ">
+    <PlatformTarget>AnyCPU</PlatformTarget>
+    <DebugSymbols>true</DebugSymbols>
+    <DebugType>full</DebugType>
+    <Optimize>false</Optimize>
+    <OutputPath>..\..\..\..\Binaries\DotNET\</OutputPath>
+    <DefineConstants>DEBUG;TRACE</DefineConstants>
+    <ErrorReport>prompt</ErrorReport>
+    <WarningLevel>4</WarningLevel>
+    <DocumentationFile>..\..\..\..\Binaries\DotNET\CrashReportReceiver.xml</DocumentationFile>
+  </PropertyGroup>
+  <PropertyGroup>
+    <StartupObject />
+  </PropertyGroup>
+  <PropertyGroup Condition="'$(Configuration)|$(Platform)' == 'Development|AnyCPU'">
+    <OutputPath>..\..\..\..\Binaries\DotNET\</OutputPath>
+    <DefineConstants>TRACE</DefineConstants>
+    <DocumentationFile>..\..\..\..\Binaries\DotNET\CrashReportReceiver.xml</DocumentationFile>
+    <Optimize>true</Optimize>
+    <DebugType>pdbonly</DebugType>
+    <PlatformTarget>AnyCPU</PlatformTarget>
+    <ErrorReport>prompt</ErrorReport>
+    <CodeAnalysisRuleSet>MinimumRecommendedRules.ruleset</CodeAnalysisRuleSet>
+    <RunCodeAnalysis>true</RunCodeAnalysis>
+  </PropertyGroup>
+  <ItemGroup>
+    <Reference Include="System" />
+    <Reference Include="System.Configuration.Install" />
+    <Reference Include="System.Core" />
+    <Reference Include="System.Management" />
+    <Reference Include="System.Xml.Linq" />
+    <Reference Include="System.Data.DataSetExtensions" />
+    <Reference Include="Microsoft.CSharp" />
+    <Reference Include="System.Data" />
+    <Reference Include="System.ServiceProcess" />
+    <Reference Include="System.Xml" />
+  </ItemGroup>
+  <ItemGroup>
+    <Compile Include="CrashReportReceiverServiceInstaller.cs">
+      <SubType>Component</SubType>
+    </Compile>
+    <Compile Include="CrashReportReceiverServiceInstaller.Designer.cs">
+      <DependentUpon>CrashReportReceiverServiceInstaller.cs</DependentUpon>
+    </Compile>
+    <Compile Include="LandingZone.cs" />
+    <Compile Include="..\..\DotNETCommon\MetaData.cs">
+      <Link>Properties\MetaData.cs</Link>
+    </Compile>
+    <Compile Include="Properties\Settings.Designer.cs">
+      <AutoGen>True</AutoGen>
+      <DesignTimeSharedInput>True</DesignTimeSharedInput>
+      <DependentUpon>Settings.settings</DependentUpon>
+    </Compile>
+    <Compile Include="CrashReportReceiverService.cs">
+      <SubType>Component</SubType>
+    </Compile>
+    <Compile Include="CrashReportReceiverService.Designer.cs">
+      <DependentUpon>CrashReportReceiverService.cs</DependentUpon>
+    </Compile>
+    <Compile Include="Program.cs" />
+    <Compile Include="Properties\AssemblyInfo.cs" />
+    <Compile Include="UploadsInProgress.cs" />
+    <Compile Include="WebHandler.cs" />
+  </ItemGroup>
+  <ItemGroup>
+    <None Include="App.config" />
+    <None Include="Properties\Settings.settings">
+      <Generator>SettingsSingleFileGenerator</Generator>
+      <LastGenOutput>Settings.Designer.cs</LastGenOutput>
+    </None>
+  </ItemGroup>
+  <ItemGroup>
+    <ProjectReference Include="..\..\DotNETCommon\DotNETUtilities\DotNETUtilities.csproj">
+      <Project>{5d7d66e8-8c76-4af9-b3ec-2ef03421d730}</Project>
+      <Name>DotNETUtilities</Name>
+    </ProjectReference>
+    <ProjectReference Include="..\CrashReportCommon\CrashReportCommon.csproj">
+      <Project>{b83de8c9-9429-4765-86bd-8eeea89b138f}</Project>
+      <Name>CrashReportCommon</Name>
+    </ProjectReference>
+  </ItemGroup>
+  <ItemGroup>
+  </ItemGroup>
+  <Import Project="$(MSBuildToolsPath)\Microsoft.CSharp.targets" />
+  <PropertyGroup>
+    <PostBuildEvent>
+    </PostBuildEvent>
+  </PropertyGroup>
+  <!-- To modify your build process, add your task inside one of the targets below and uncomment it. 
+       Other similar extension points exist, see Microsoft.Common.targets.
+  <Target Name="BeforeBuild">
+  </Target>
+  <Target Name="AfterBuild">
+  </Target>
+  -->
+</Project>
\ No newline at end of file
diff --git a/Engine/Source/Programs/CrashReporter/CrashReportReceiver/CrashReportReceiverService.cs b/Engine/Source/Programs/CrashReporter/CrashReportReceiver/CrashReportReceiverService.cs
new file mode 100644
--- /dev/null
+++ b/Engine/Source/Programs/CrashReporter/CrashReportReceiver/CrashReportReceiverService.cs
@@ -0,0 +1,89 @@
+﻿// Copyright 1998-2015 Epic Games, Inc. All Rights Reserved.
+
+using System;
+using System.Collections.Generic;
+using System.ComponentModel;
+using System.Data;
+using System.Diagnostics;
+using System.IO;
+using System.Reflection;
+using System.ServiceProcess;
+using System.Text;
+using System.Threading.Tasks;
+
+namespace Tools.CrashReporter.CrashReportReceiver
+{
+	/// <summary>
+	/// A class to handle receiving of files from clients, both internally and externally.
+	/// </summary>
+	public partial class CrashReporterReceiverServicer : ServiceBase, IDisposable
+	{
+		/// <summary>An instance of the web handler class that listens for incoming crash reports.</summary>
+		public WebHandler Web = null;
+
+		/// <summary>Current log file to write debug progress info to</summary>
+		static public CrashReportCommon.LogWriter Log = null;
+
+		/// <summary>
+		/// Write a status update to the event log.
+		/// </summary>
+		/// <param name="Message">A string containing the status update.</param>
+		static public void WriteEvent( string Message )
+		{
+			// Status messages go to the log file and/or console
+			Log.Print( "[STATUS] " + Message );
+		}
+
+		/// <summary>
+		/// Write a performance update to the event log.
+		/// </summary>
+		static public void WritePerf( string Message )
+		{
+			Log.Print( "[PERF  ] " + Message );
+		}
+
+		/// <summary>
+		/// Initialise all the components.
+		/// </summary>
+		public CrashReporterReceiverServicer()
+		{
+			InitializeComponent();
+
+			var StartupPath = Path.GetDirectoryName(Assembly.GetExecutingAssembly().Location);
+			Log = new CrashReportCommon.LogWriter( "CrashReportReceiver", Path.Combine( StartupPath, "Logs" ) );
+		}
+
+		/// <summary>
+		/// Start the service, and stop the service if there were any errors found.
+		/// </summary>
+		/// <param name="Arguments">An array of arguments passed from the command line (unused).</param>
+		protected override void OnStart( string[] Arguments )
+		{
+			Web = new WebHandler();
+			
+			if( !Web.bStartedSuccessfully )
+			{
+				OnStop();
+			}
+		}
+
+		/// <summary>
+		/// Stop the service.
+		/// </summary>
+		protected override void OnStop()
+		{
+			Web.Release();
+		}
+
+		/// <summary>
+		/// Run the service in debug mode. This spews all logging to the console rather than suppressing it.
+		/// </summary>
+		public void DebugRun()
+		{
+			OnStart( null );
+			Console.WriteLine( "Press enter to exit" );
+			Console.Read();
+			OnStop();
+		}
+	}
+}
diff --git a/Engine/Source/Programs/CrashReporter/CrashReportReceiver/CrashReportReceiverServiceInstaller.cs b/Engine/Source/Programs/CrashReporter/CrashReportReceiver/CrashReportReceiverServiceInstaller.cs
new file mode 100644
--- /dev/null
+++ b/Engine/Source/Programs/CrashReporter/CrashReportReceiver/CrashReportReceiverServiceInstaller.cs
@@ -0,0 +1,50 @@
+﻿// Copyright 1998-2015 Epic Games, Inc. All Rights Reserved.
+
+using System;
+using System.Collections;
+using System.Collections.Generic;
+using System.ComponentModel;
+using System.Configuration.Install;
+using System.Linq;
+using System.ServiceProcess;
+
+namespace Tools.CrashReporter.CrashReportReceiver
+{
+	/// <summary>
+	/// An installer class for the CrashReportReceiver.
+	/// </summary>
+	[RunInstaller( true )]
+	public partial class CrashReportReceiverServiceInstaller : System.Configuration.Install.Installer
+	{
+		private ServiceInstaller Installer;
+		private ServiceProcessInstaller ProcessInstaller;
+
+		/// <summary>
+		/// Set all the parameters required to install the service.
+		/// </summary>
+		public CrashReportReceiverServiceInstaller()
+		{
+			InitializeComponent();
+
+			Installer = new ServiceInstaller();
+			Installer.StartType = ServiceStartMode.Automatic;
+			Installer.ServiceName = "CrashReportReceiverService";
+			Installer.DisplayName = "CrashReport Receiver Service";
+			Installer.Description = "A web service that receives files and copies them locally for later processing.";
+			Installers.Add( Installer );
+
+			ProcessInstaller = new ServiceProcessInstaller();
+			ProcessInstaller.Account = ServiceAccount.NetworkService;
+			Installers.Add( ProcessInstaller );
+		}
+
+		/// <summary>
+		/// Must exist, because installation will fail without it.
+		/// </summary>
+		/// <param name="StateSaver">Passed through to the base method.</param>
+		public override void Install( IDictionary StateSaver )
+		{
+			base.Install( StateSaver );
+		}
+	}
+}
diff --git a/Engine/Source/Programs/CrashReporter/CrashReportReceiver/LandingZone.cs b/Engine/Source/Programs/CrashReporter/CrashReportReceiver/LandingZone.cs
new file mode 100644
--- /dev/null
+++ b/Engine/Source/Programs/CrashReporter/CrashReportReceiver/LandingZone.cs
@@ -0,0 +1,209 @@
+// Copyright 1998-2015 Epic Games, Inc. All Rights Reserved.
+
+using System;
+using System.Collections.Generic;
+using System.Collections.Concurrent;
+using System.IO;
+using System.Linq;
+using System.Xml.Linq;
+using System.Threading;
+
+namespace Tools.CrashReporter.CrashReportReceiver
+{
+	using ReportsDict = ConcurrentDictionary<string, DateTime>;
+
+	/// <summary>
+	/// Class to deal with the 'landing zone' folder, where crashes are uploaded
+	/// </summary>
+	public class LandingZoneMonitor
+	{
+		#region Public methods
+		/// <summary>
+		/// Constructor: set up a landing zone directory
+		/// </summary>
+		/// <param name="DirectoryToUse">Directory to store the WER sub-directories</param>
+		/// <param name="Log">Optional file to write messages to</param>
+		public LandingZoneMonitor(string DirectoryToUse, CrashReportCommon.LogWriter Log)
+		{
+			Directory = DirectoryToUse;
+			LogFile = Log;
+
+			DirectoryInfo DirInfo = new DirectoryInfo(Directory);
+			if (DirInfo.Exists)
+			{
+				ReadMemorizedReceivedReports();
+				MemorizeReceivedReports();
+			}
+			else
+			{
+				DirInfo.Create();
+			}
+
+			LogFile.Print("Landing zone: " + DirectoryToUse);
+		}
+
+		/// <summary>
+		/// Move a successfully uploaded WER into the landing zone folder
+		/// </summary>
+		/// <param name="SourceDirInfo">Temporary directory WER was uploaded to</param>
+		/// <param name="ReportId">Unique ID of the report, used as the folder name in the landing zone</param>
+		public void ReceiveReport(DirectoryInfo SourceDirInfo, string ReportId)
+		{
+			string ReportPath = Path.Combine(Directory, ReportId);
+
+			// Recreate the landing zone directory, just in case.
+			System.IO.Directory.CreateDirectory( Directory );
+
+			DirectoryInfo DirInfo = new DirectoryInfo(ReportPath);
+			if (DirInfo.Exists)
+			{
+				DirInfo.Delete(true);
+				DirInfo.Refresh();
+			}
+
+			// Retry the move in case the OS keeps some locks on the uploaded files longer than expected
+			for (int Retry = 0; Retry != MoveUploadedReportRetryAttempts; ++Retry)
+			{
+				try
+				{
+					SourceDirInfo.MoveTo(DirInfo.FullName);
+					break;
+				}
+				catch (Exception)
+				{
+				}
+				System.Threading.Thread.Sleep(100);
+			}
+
+			CrashReporterReceiverServicer.WriteEvent( "A new report has been received into " + DirInfo.FullName );
+
+			ReceivedReports[ReportId] = DateTime.Now;
+
+			// Ensure the receive reports memo file doesn't grow indefinitely
+			ForgetOldReports();
+
+			// Update file on disk, in case service is restarted
+			MemorizeReceivedReports();
+		}
+
+		/// <summary>
+		/// Check whether a report of a given ID has already been through the landing zone
+		/// </summary>
+		/// <param name="ReportId">Report ID to check for</param>
+		/// <returns>Whether the report of the given ID has been successfully uploaded before</returns>
+		public bool HasReportAlreadyBeenReceived(string ReportId)
+		{
+			return ReceivedReports.ContainsKey(ReportId);
+		}
+
+		#endregion Public methods
+
+		#region Private methods
+		/// <summary>
+		/// Save out the current list of received reports to disk
+		/// </summary>
+		void MemorizeReceivedReports()
+		{
+			var MemoPath = Path.Combine(Directory, MemoFilename);
+			// Recreate the landing zone directory, just in case.
+			System.IO.Directory.CreateDirectory( Directory );
+			try
+			{
+				if (Interlocked.Exchange(ref MemoBeingWritten, 1) == 1)
+				{
+					return;
+				}
+
+
+				new XElement
+				( 
+					"Root",
+					from Item in ReceivedReports.OrderBy( x => x.Value )	
+					select new XElement( "CrashReport", Item.Key, new XAttribute( "time", Item.Value.ToString() ) )
+				)
+				.Save( MemoPath );
+				CrashReporterReceiverServicer.WriteEvent( "Saved to " + MemoPath );
+				Interlocked.Exchange(ref MemoBeingWritten, 0);
+			}
+			catch (Exception Ex)
+			{
+				LogFile.Print("Failed to memorize received reports: " + Ex.Message);
+				Interlocked.Exchange( ref MemoBeingWritten, 0 );
+			}
+		}
+
+		/// <summary>
+		/// Load a list of received reports from disk
+		/// </summary>
+		void ReadMemorizedReceivedReports()
+		{
+			var MemoPath = Path.Combine(Directory, MemoFilename);
+			try
+			{
+				ReceivedReports = new ReportsDict
+				(
+					from element in XElement.Load(MemoPath).Descendants()
+					select new KeyValuePair<string, DateTime>(element.Value.ToString(),	DateTime.Parse(element.Attribute("time").Value))
+				);
+			}
+			catch (Exception)
+			{
+				// Not an error: the file isn't there the first time
+			}
+		}
+
+		/// <summary>
+		/// Ensure the receive reports memo file doesn't grow indefinitely
+		/// High thread contention?
+		/// </summary>
+		void ForgetOldReports()
+		{
+			// Recreate the dictionary with only recent reports
+			ReceivedReports = new ReportsDict
+			(
+				from Item in ReceivedReports
+				where Item.Value.AddDays(DaysToRememberReportsFor) > DateTime.Now
+				select Item
+			);
+		}
+		#endregion Private methods
+
+		#region Private members
+		/// <summary>
+		/// Filename to use to persist received reports dictionary
+		/// </summary>
+		const string MemoFilename = "ReceivedReports.xml";
+
+		/// <summary>
+		/// Atomic flag to indicate memo file is being written to
+		/// </summary>
+		int MemoBeingWritten = 0;
+
+		/// <summary>
+		/// Landing zone root directory
+		/// </summary>
+		string Directory;
+
+		/// <summary>
+		/// Log file to write to
+		/// </summary>
+		CrashReportCommon.LogWriter LogFile;
+
+		/// <summary>
+		/// Dictionary of reports that have been seen before mapping to the time they were received
+		/// </summary>
+		ReportsDict ReceivedReports = new ReportsDict();
+
+		/// <summary>
+		/// Number of days to remember reports, to prevent the multiple uploads of the same report
+		/// </summary>
+		const int DaysToRememberReportsFor = 7;
+
+		/// <summary>
+		/// Number of times to retry moving the report folder to the landing zone once uploaded
+		/// </summary>
+		const int MoveUploadedReportRetryAttempts = 5;
+
+		#endregion Private members
+	}
+}
diff --git a/Engine/Source/Programs/CrashReporter/CrashReportReceiver/Program.cs b/Engine/Source/Programs/CrashReporter/CrashReportReceiver/Program.cs
new file mode 100644
--- /dev/null
+++ b/Engine/Source/Programs/CrashReporter/CrashReportReceiver/Program.cs
@@ -0,0 +1,32 @@
+﻿// Copyright 1998-2015 Epic Games, Inc. All Rights Reserved.
+
+using System;
+using System.ServiceProcess;
+using System.Threading.Tasks;
+
+namespace Tools.CrashReporter.CrashReportReceiver
+{
+	/// <summary>
+	/// The container for the Crash Report Receiver application.
+	/// </summary>
+	static class CrashReportReceiverProgram
+	{
+		/// <summary>
+		/// The main entry point for crash report receiver application.
+		/// </summary>
+		static void Main()
+		{
+			CrashReporterReceiverServicer CrashReporterReceiver = new CrashReporterReceiverServicer();
+			if( !Environment.UserInteractive )
+			{
+				// Launch the service as normal if we aren't in the debugger, or aren't in an interactive environment
+				ServiceBase.Run( CrashReporterReceiver );
+			}
+			else
+			{
+				// Call OnStart, wait for a console key press, call OnStop
+				CrashReporterReceiver.DebugRun();
+			}
+		}
+	}
+}
diff --git a/Engine/Source/Programs/CrashReporter/CrashReportReceiver/Properties/AssemblyInfo.cs b/Engine/Source/Programs/CrashReporter/CrashReportReceiver/Properties/AssemblyInfo.cs
new file mode 100644
--- /dev/null
+++ b/Engine/Source/Programs/CrashReporter/CrashReportReceiver/Properties/AssemblyInfo.cs
@@ -0,0 +1,11 @@
+﻿// Copyright 1998-2015 Epic Games, Inc. All Rights Reserved.
+
+using System.Reflection;
+
+// General Information about an assembly is controlled through the following 
+// set of attributes. Change these attribute values to modify the information
+// associated with an assembly.
+[assembly: AssemblyTitle( "CrashReportReceiver" )]
+[assembly: AssemblyDescription( "A service to receive uploaded crash reports." )]
+[assembly: AssemblyConfiguration( "" )]
+
diff --git a/Engine/Source/Programs/CrashReporter/CrashReportReceiver/Properties/Settings.settings b/Engine/Source/Programs/CrashReporter/CrashReportReceiver/Properties/Settings.settings
new file mode 100644
--- /dev/null
+++ b/Engine/Source/Programs/CrashReporter/CrashReportReceiver/Properties/Settings.settings
@@ -0,0 +1,15 @@
+﻿<?xml version='1.0' encoding='utf-8'?>
+<SettingsFile xmlns="http://schemas.microsoft.com/VisualStudio/2004/01/settings" CurrentProfile="(Default)" GeneratedClassNamespace="Tools.CrashReporter.CrashReportReceiver.Properties" GeneratedClassName="Settings">
+  <Profiles />
+  <Settings>
+    <Setting Name="MaxAllowedBandwidthMB" Type="System.Int32" Scope="Application">
+      <Value Profile="(Default)">10</Value>
+    </Setting>
+    <Setting Name="MaxAllowedCPUtilizationPct" Type="System.Single" Scope="Application">
+      <Value Profile="(Default)">75</Value>
+    </Setting>
+    <Setting Name="CrashReportRepository" Type="System.String" Scope="Application">
+      <Value Profile="(Default)">U:\CRASHREPORT\DMZ</Value>
+    </Setting>
+  </Settings>
+</SettingsFile>
\ No newline at end of file
diff --git a/Engine/Source/Programs/CrashReporter/CrashReportReceiver/UploadsInProgress.cs b/Engine/Source/Programs/CrashReporter/CrashReportReceiver/UploadsInProgress.cs
new file mode 100644
--- /dev/null
+++ b/Engine/Source/Programs/CrashReporter/CrashReportReceiver/UploadsInProgress.cs
@@ -0,0 +1,354 @@
+// Copyright 1998-2015 Epic Games, Inc. All Rights Reserved.
+
+using System;
+using System.Collections.Generic;
+using Tools.CrashReporter.CrashReportCommon;
+using System.Threading;
+using System.Diagnostics;
+
+namespace Tools.CrashReporter.CrashReportReceiver
+{
+	/// <summary>Manager to keep track of what reports the server is currently dealing with </summary>
+	/// <remarks>
+	/// Mainly to cope with multiple clients fighting to upload the same report. This shouldn't happen, but the service
+	/// should be resilient against it even so.
+	/// </remarks>
+	public class FUploadsInProgress
+	{
+		/// <summary> Background thread used to monitor CPU/network usage. </summary>
+		private Thread PerformancePooler;
+
+		/// <summary> Performance counter used to continuously monitor CPU usage. </summary>
+		private PerformanceCounter PerfCounterProcessorTime = new PerformanceCounter( "Processor", "% Processor Time", "_Total" );
+
+		/// <summary> The total amount of data received via reading from the stream. </summary>
+		public Int64 TotalReceivedData = 0;
+
+		/// <summary> The current amount of data received via reading from the stream, for the last second. </summary>
+		public Int64 CurrentReceivedData = 0;
+
+		/// <summary> The current CPU usage, from 0 to 100. </summary>
+		private float CurrentCPUUsage = 0.0f;
+
+		/// <summary> 
+		/// Maximum allowed bandwidth that can be used by the crash receiver service.
+		/// If this value is exceeded the service will start to ignore larger files like CrashVideo.avi or large log files.
+		/// Currently 10MB/s.
+		/// Default 1024MB/s.
+		/// @todo Should be set based on the current network connection.
+		/// </summary>
+		private long MaxAllowedBandwidth = 1024;
+
+		/// <summary> 
+		/// Maximum allowed CPU utilization for the machine where the crash receiver service is running.
+		/// If this value is exceeded the service will start to ignore larger files like CrashVideo.avi or large log files. 
+		/// Currently 75%.
+		/// </summary>
+		private float MaxAllowedCPUtilizationPct = 75.0f;
+
+		/// <summary> 
+		/// Files smaller than this size will not be ignored during the dynamic balancing.
+		/// Currently 376KB. Should handle all WER XML files, most of the log and minidump files.
+		/// Will be reduced once crash report client will support compressing of the data.
+		/// </summary>
+		private Int64 MaxForcedFilesize = 3 * 128 * 1024;
+
+		private Int64 NumFilesReceived = 0;
+		private Int64 NumFilesSkipped = 0;
+
+		/// <summary> To convert bytes into megabytes. </summary>
+		public const double MBInv = 1.0 / 1024.0 / 1024.0;
+
+		/// <summary> Periodically pools the performance data. </summary>
+		private static void StaticPoolPerformanceData( object Obj )
+		{
+			FUploadsInProgress This = (FUploadsInProgress)Obj;
+			This.PoolPerformanceData();
+		}
+
+		private void PoolPerformanceData()
+		{
+			while( true )
+			{
+				Int64 NumTotalFiles = NumFilesReceived + NumFilesSkipped;
+
+				// Get the CPU for the last second.
+				CurrentCPUUsage = PerfCounterProcessorTime.NextValue();
+
+				if( CurrentReceivedData > 0 && CurrentTotalUploadCount > 0 )
+				{
+					CrashReporterReceiverServicer.WritePerf
+					(
+						string.Format( "CPU Usage: {0:00}%, Received data: {1,5}MB/{2:00.00} of {4:00.00} MB, Uploads:{3,3}, Files:{5,6},{6,6} : {7:00}%",
+						CurrentCPUUsage, 
+						(Int64)Math.Truncate( MBInv * TotalReceivedData ), 
+						MBInv * CurrentReceivedData, 
+						UploadCount, 
+						MBInv * MaxAllowedBandwidth,
+						NumFilesReceived, NumFilesSkipped,
+						NumTotalFiles > 0 ? 100.0f * NumFilesSkipped / NumTotalFiles : 0.0f )
+					);
+				}
+
+				Interlocked.Add( ref TotalReceivedData, CurrentReceivedData );
+
+				// Reset the the performance data.
+				Interlocked.Exchange( ref CurrentReceivedData, 0 );
+				Interlocked.Exchange( ref CurrentTotalUploadCount, 0 );
+
+				// Sleep.
+				System.Threading.Thread.Sleep( 1000 );
+			}
+		}
+
+		/// <summary>
+		/// Constructor: set up throttle from app settings
+		/// </summary>
+		public FUploadsInProgress()
+		{
+			var Settings = Properties.Settings.Default;
+			MaxAllowedBandwidth = Settings.MaxAllowedBandwidthMB * 1024 * 1024;
+			MaxAllowedCPUtilizationPct = Settings.MaxAllowedCPUtilizationPct;			
+
+			PerformancePooler = new Thread( Tools.CrashReporter.CrashReportReceiver.FUploadsInProgress.StaticPoolPerformanceData );
+			PerformancePooler.Start( this );
+		}
+
+		/// <summary>
+		/// Check for parallel attempts to upload the same file or complete a report
+		/// </summary>
+		/// <param name="ReportId">Directory name being used as report id</param>
+		/// <param name="Filename">Leaf name of file to be uploaded</param>
+		/// <param name="Filesize">Size of the file</param>
+		/// <param name="FailureMessage">Message to log and return in case of failure</param>
+		/// <param name="bIsOverloaded">If True, this attempt will temporary overload the connection</param>
+		/// <returns>True if this is the first client to attempt to upload this file</returns>
+		public bool TryReceiveFile( string ReportId, string Filename, long Filesize, ref string FailureMessage, ref bool bIsOverloaded )
+		{
+			lock (Mutex)
+			{
+				var ReportRecord = EnsureRecord(ReportId);
+				if (ReportRecord.bReportComplete)
+				{
+					FailureMessage = string.Format("Attempt to upload file \"{0}\" to completed report \"{1}\"", Filename, ReportId);
+					return false;
+				}
+
+				string FilenameKey = Filename.ToLower();
+				if (ReportRecord.FilesUploaded.Contains(FilenameKey))
+				{
+					FailureMessage = string.Format("File \"{0}\" has already been uploaded to report \"{1}\"", Filename, ReportId);
+					return false;
+				}
+
+				if( !CanStartOperation( Filename, Filesize, ref bIsOverloaded ) )
+				{
+					FailureMessage = string.Format( "Skipping \"{0}\"{1:0.0}MB in report \"{2}\" due to excessive load", Filename, MBInv*Filesize, ReportId );
+					return false;
+				}
+
+				ReportRecord.FilesUploaded.Add( FilenameKey );
+				Interlocked.Increment( ref UploadCount );
+				Interlocked.Increment( ref CurrentTotalUploadCount );
+				return true;
+			}
+		}
+
+
+		private bool CanStartOperation( string Filename, long Filesize, ref bool bIsOverloaded )
+		{
+			bool bCPUExceeded = CurrentCPUUsage > MaxAllowedCPUtilizationPct;
+			bool bBandwidthExceeded = CurrentReceivedData > MaxAllowedBandwidth;
+
+			if( bCPUExceeded || bBandwidthExceeded )
+			{
+				if( Filesize > MaxForcedFilesize || MaxAllowedBandwidth < 0 )
+				{
+					return TryWaitForStartOperation();
+				}
+				else
+				{
+					// We're temporary overloading the connection, so reduce the max allowed bandwidth.
+					Interlocked.Add( ref MaxAllowedBandwidth, -Filesize );
+					bIsOverloaded = true;
+					return true;
+				}
+			}
+
+			return true;
+		}
+
+		/// <summary> Spin loop for a while, maybe we hit an opportunity to start the upload. </summary>
+		private bool TryWaitForStartOperation()
+		{
+			bool bCPUExceeded = false;
+			bool bBandwidthExceeded = false;
+			for( int Index = 0; Index < 4; Index++ )
+			{
+				Thread.Sleep( 1000 );
+				bCPUExceeded = CurrentCPUUsage > MaxAllowedCPUtilizationPct;
+				bBandwidthExceeded = CurrentReceivedData > MaxAllowedBandwidth;
+
+				if( !bCPUExceeded && !bBandwidthExceeded )
+				{
+					// Found, continue.
+					CrashReporterReceiverServicer.WritePerf( "TryWaitForStartOperation was successful" );
+					return true;
+				}
+			}
+
+			// Give up.
+			Interlocked.Increment( ref NumFilesSkipped );
+			return false;
+		}
+
+		/// <summary>
+		/// Notification that a file upload has finished, regardless of success
+		/// <param name="Filesize">Size of the file</param>
+		/// <param name="bIsOverloaded">If true, it means that we need to restore the max allowed bandwidth.</param>
+		/// 
+		/// </summary>
+		public void FileUploadAttemptDone( long Filesize, bool bIsOverloaded )
+		{
+			Interlocked.Increment( ref NumFilesReceived );
+			Interlocked.Decrement( ref UploadCount );
+			if( bIsOverloaded )
+			{
+				// Uploading is done, so restore the max allowed bandwidth.
+				Interlocked.Add( ref MaxAllowedBandwidth, Filesize );
+			}
+		}
+
+		/// <summary>
+		/// Check for parallel attempts to complete a report
+		/// </summary>
+		/// <param name="ReportId">Directory name being used as report id</param>
+		/// <returns>True if this is the first client to attempt to complete this report</returns>
+		public bool TrySetReportComplete(string ReportId)
+		{
+			lock (Mutex)
+			{
+				var UploadRecord = EnsureRecord(ReportId);
+
+				if (UploadRecord.bReportComplete)
+				{
+					return false;
+				}
+
+				UploadRecord.bReportComplete = true;
+				UploadRecord.FilesUploaded = null;
+
+				return true;
+			}
+		}
+
+		/// <summary>
+		/// Get the last access time of a report if we know about it
+		/// </summary>
+		/// <param name="ReportId">Report to look for</param>
+		/// <param name="LastAccessTime">Time to fill in of last attempted file upload to report, if found</param>
+		/// <returns>Whether report was found</returns>
+		public bool TryGetLastAccessTime(string ReportId, ref DateTime LastAccessTime)
+		{
+			lock (Mutex)
+			{
+				FReportRecord UploadRecord;
+				if (UploadRecords.TryGetValue(ReportId.ToLower(), out UploadRecord))
+				{
+					LastAccessTime = UploadRecord.LastAccessTime;
+					return true;
+				}	
+				return false;
+			}
+		}
+
+		#region Private
+		/// <summary>
+		/// Create a record of files that have been uploaded for a report, if there's not one already
+		/// </summary>
+		/// <param name="ReportId">Directory name being used as report id</param>
+		/// <returns>The record object with a set of files that have been uploaded so far</returns>
+		FReportRecord EnsureRecord(string ReportId)
+		{
+			string ReportIdKey = ReportId.ToLower();
+			FReportRecord UploadRecord;
+			if (!UploadRecords.TryGetValue(ReportIdKey, out UploadRecord))
+			{
+				UploadRecord = new FReportRecord();
+				UploadRecords[ReportIdKey] = UploadRecord;
+			}
+
+			UploadRecord.LastAccessTime = DateTime.Now;
+
+			DoHouseKeeping();
+
+			return UploadRecord;
+		}
+
+		/// <summary>
+		/// Clear out records that have been around for a while (see MinutesToRememberReports below)
+		/// </summary>
+		void DoHouseKeeping()
+		{
+			var Now = DateTime.Now;
+
+			if ((Now - OldestAccessTime).Minutes > MinutesToRememberReports)
+			{
+				OldestAccessTime = Now;
+
+				// This is doing a little bit of work while a lock is held,
+				// but only happens at most once a minute
+				var OldUploadRecords = UploadRecords;
+				UploadRecords = new Dictionary<string, FReportRecord>();
+				foreach (var Entry in OldUploadRecords)
+				{
+					var LastAccessTime = Entry.Value.LastAccessTime;
+					if ((Now - LastAccessTime).Minutes > MinutesToRememberReports)
+					{
+						// Old report: forget it
+						continue;
+					}
+					UploadRecords[Entry.Key] = Entry.Value;
+					if (LastAccessTime < OldestAccessTime)
+					{
+						OldestAccessTime = LastAccessTime;
+					}
+				}
+			}
+		}
+
+		/// <summary>
+		/// A record of which files have been uploaded for a report and whether the report is complete
+		/// </summary>
+		class FReportRecord
+		{
+			/// <summary>Flag to indicate report has been completed</summary>	
+			public bool bReportComplete = false;
+
+			/// <summary>Lowercase filenames of files uploaded for this report</summary>	
+			public HashSet<string> FilesUploaded = new HashSet<string>();
+
+			/// <summary>Most recent time EnsureRecord was called on this record</summary>
+			public DateTime LastAccessTime;
+		};
+
+		/// <summary>Dictionary to keep track of existence of files</summary>
+		Dictionary<string, FReportRecord> UploadRecords = new Dictionary<string, FReportRecord>();
+
+		/// <summary>Remember reports for a day</summary>
+		const int MinutesToRememberReports = 60 * 24;
+
+		/// <summary>Oldest report in UploadRecords, to avoid trawling the whole dictionary on each house-keeping pass</summary>
+		DateTime OldestAccessTime = DateTime.Now;
+
+		/// <summary>Number of simultaneous file uploads in progress</summary>
+		int UploadCount = 0;
+
+		/// <summary>Total number of file uploads, for the last second.</summary>
+		int CurrentTotalUploadCount = 0;
+
+		/// <summary>Mutex to ensure synchronous access to members of this class</summary>
+		object Mutex = new object();
+		#endregion
+	}
+}
