Copying //UE4/Dev-Sequencer to Dev-Main (//UE4/Dev-Main)

diff --git a/Engine/Source/Editor/Sequencer/Private/TimeSliderController.cpp b/Engine/Source/Editor/Sequencer/Private/TimeSliderController.cpp
--- a/Engine/Source/Editor/Sequencer/Private/TimeSliderController.cpp
+++ b/Engine/Source/Editor/Sequencer/Private/TimeSliderController.cpp
@@ -107,6 +107,7 @@ FSequencerTimeSliderController::FSequencerTimeSliderController( const FTimeSlide
 	ScrubHandleUp = FEditorStyle::GetBrush( TEXT( "Sequencer.Timeline.ScrubHandleUp" ) ); 
 	ScrubHandleDown = FEditorStyle::GetBrush( TEXT( "Sequencer.Timeline.ScrubHandleDown" ) );
 	ScrubHandleSize = 13.f;
+	ContextMenuSupression = 0;
 }
 
 struct FDrawTickArgs
@@ -418,7 +419,7 @@ int32 FSequencerTimeSliderController::DrawPlaybackRange(const FGeometry& Allotte
 	FSlateDrawElement::MakeBox(
 		OutDrawElements,
 		LayerId+1,
-		AllottedGeometry.ToPaintGeometry(FVector2D(0.f, 0.f), FVector2D(PlaybackRangeL - 1.f, AllottedGeometry.Size.Y)),
+		AllottedGeometry.ToPaintGeometry(FVector2D(0.f, 0.f), FVector2D(PlaybackRangeL, AllottedGeometry.Size.Y)),
 		FEditorStyle::GetBrush("WhiteBrush"),
 		MyClippingRect,
 		ESlateDrawEffect::None,
@@ -465,11 +466,35 @@ FReply FSequencerTimeSliderController::OnMouseButtonUp( SWidget& WidgetOwner, co
 {
 	bool bHandleLeftMouseButton = MouseEvent.GetEffectingButton() == EKeys::LeftMouseButton && WidgetOwner.HasMouseCapture();
 	bool bHandleRightMouseButton = MouseEvent.GetEffectingButton() == EKeys::RightMouseButton && WidgetOwner.HasMouseCapture() && TimeSliderArgs.AllowZoom ;
-	
+
 	if ( bHandleRightMouseButton )
 	{
 		if (!bPanning)
 		{
+			// Open a context menu if allowed
+			if (ContextMenuSupression == 0 && TimeSliderArgs.PlaybackRange.IsSet())
+			{
+				FScrubRangeToScreen RangeToScreen( TimeSliderArgs.ViewRange.Get(), MyGeometry.Size );
+				FVector2D CursorPos = MyGeometry.AbsoluteToLocal( MouseEvent.GetScreenSpacePosition() );
+
+				float MouseValue = RangeToScreen.LocalXToInput( CursorPos.X );
+				if (TimeSliderArgs.Settings->GetIsSnapEnabled())
+				{
+					MouseValue = TimeSliderArgs.Settings->SnapTimeToInterval(MouseValue);
+				}
+
+				TSharedRef<SWidget> MenuContent = OpenSetPlaybackRangeMenu(MouseValue);
+				FSlateApplication::Get().PushMenu(
+					WidgetOwner.AsShared(),
+					MouseEvent.GetEventPath() != nullptr ? *MouseEvent.GetEventPath() : FWidgetPath(),
+					MenuContent,
+					MouseEvent.GetScreenSpacePosition(),
+					FPopupTransitionEffect( FPopupTransitionEffect::ContextMenu )
+					);
+
+				return FReply::Handled().SetUserFocus(MenuContent, EFocusCause::SetDirectly).ReleaseMouseCapture();
+			}
+
 			// return unhandled in case our parent wants to use our right mouse button to open a context menu
 			return FReply::Unhandled().ReleaseMouseCapture();
 		}
@@ -688,7 +713,7 @@ FReply FSequencerTimeSliderController::OnMouseWheel( SWidget& WidgetOwner, const
 	TOptional<TRange<float>> NewTargetRange;
 
 	float MouseFractionX = MyGeometry.AbsoluteToLocal(MouseEvent.GetScreenSpacePosition()).X / MyGeometry.GetLocalSize().X;
-	if ( TimeSliderArgs.AllowZoom && (MouseEvent.IsLeftControlDown() || MouseEvent.IsRightControlDown()) )
+	if ( TimeSliderArgs.AllowZoom && MouseEvent.IsControlDown() )
 	{
 		const float ZoomDelta = -0.2f * MouseEvent.GetWheelDelta();
 		if (ZoomByDelta(ZoomDelta, MouseFractionX))
@@ -696,7 +721,7 @@ FReply FSequencerTimeSliderController::OnMouseWheel( SWidget& WidgetOwner, const
 			return FReply::Handled();
 		}
 	}
-	else if (MouseEvent.IsLeftShiftDown() || MouseEvent.IsRightShiftDown())
+	else if (MouseEvent.IsShiftDown())
 	{
 		PanByDelta(-MouseEvent.GetWheelDelta());
 		return FReply::Handled();
@@ -787,6 +812,50 @@ int32 FSequencerTimeSliderController::OnPaintSectionView( const FGeometry& Allot
 	return LayerId;
 }
 
+TSharedRef<SWidget> FSequencerTimeSliderController::OpenSetPlaybackRangeMenu(float MouseTime)
+{
+	const bool bShouldCloseWindowAfterMenuSelection = true;
+	FMenuBuilder MenuBuilder(bShouldCloseWindowAfterMenuSelection, nullptr);
+
+	FText NumericText;
+	if (SequencerSnapValues::IsTimeSnapIntervalFrameRate(TimeSliderArgs.Settings->GetTimeSnapInterval()) && TimeSliderArgs.Settings->GetShowFrameNumbers())
+	{
+		NumericText = FText::Format(LOCTEXT("FrameTextFormat", "Playback Range (at frame {0}):"), FText::AsNumber(TimeToFrame(MouseTime)));
+	}
+	else
+	{
+		NumericText = FText::Format(LOCTEXT("TimeTextFormat", "Playback Range (at {0}s):"), FText::AsNumber(MouseTime));
+	}
+
+	TRange<float> PlaybackRange = TimeSliderArgs.PlaybackRange.Get();
+
+	MenuBuilder.BeginSection("SequencerPlaybackRangeMenu", NumericText);
+	{
+		MenuBuilder.AddMenuEntry(
+			FText::Format(LOCTEXT("SetPlaybackStart", "Set Start Time"), NumericText),
+			FText(),
+			FSlateIcon(),
+			FUIAction(
+				FExecuteAction::CreateLambda([=]{ return SetPlaybackRangeStart(MouseTime); }),
+				FCanExecuteAction::CreateLambda([=]{ return MouseTime <= PlaybackRange.GetUpperBoundValue(); })
+			)
+		);
+
+		MenuBuilder.AddMenuEntry(
+			FText::Format(LOCTEXT("SetPlaybackEnd", "Set End Time"), NumericText),
+			FText(),
+			FSlateIcon(),
+			FUIAction(
+				FExecuteAction::CreateLambda([=]{ return SetPlaybackRangeEnd(MouseTime); }),
+				FCanExecuteAction::CreateLambda([=]{ return MouseTime >= PlaybackRange.GetLowerBoundValue(); })
+			)
+		);
+	}
+	MenuBuilder.EndSection(); // SequencerPlaybackRangeMenu
+
+	return MenuBuilder.MakeWidget();
+}
+
 int32 FSequencerTimeSliderController::TimeToFrame(float Time) const
 {
 	float FrameRate = 1.0f/TimeSliderArgs.Settings->GetTimeSnapInterval();
@@ -859,7 +928,7 @@ bool FSequencerTimeSliderController::ZoomByDelta( float InDelta, float MousePosi
 	float NewViewOutputMin = LocalViewRangeMin - (OutputChange * MousePositionFraction);
 	float NewViewOutputMax = LocalViewRangeMax + (OutputChange * (1.f - MousePositionFraction));
 
-	if( FMath::Abs( OutputChange ) > 0.01f && NewViewOutputMin < NewViewOutputMax )
+	if( NewViewOutputMin < NewViewOutputMax )
 	{
 		ClampViewRange(NewViewOutputMin, NewViewOutputMax);
 		SetViewRange(NewViewOutputMin, NewViewOutputMax, EViewRangeInterpolation::Animated);
diff --git a/Engine/Source/Editor/Sequencer/Private/TimeSliderController.h b/Engine/Source/Editor/Sequencer/Private/TimeSliderController.h
--- a/Engine/Source/Editor/Sequencer/Private/TimeSliderController.h
+++ b/Engine/Source/Editor/Sequencer/Private/TimeSliderController.h
@@ -4,6 +4,8 @@
 
 #include "Editor/SequencerWidgets/Public/ITimeSlider.h"
 
+struct FContextMenuSuppressor;
+
 struct FPaintPlaybackRangeArgs
 {
 	FPaintPlaybackRangeArgs()
@@ -148,6 +150,8 @@ private:
 	void SetPlaybackRangeStart(float NewStart);
 	void SetPlaybackRangeEnd(float NewEnd);
 
+	TSharedRef<SWidget> OpenSetPlaybackRangeMenu(float MouseTime);
+
 private:
 	FTimeSliderArgs TimeSliderArgs;
 	/** The size of the scrub handle */
@@ -174,4 +178,26 @@ private:
 	FVector2D MouseDownRange;
 	/** Range stack */
 	TArray<FVector2D> RangeStack;
+	/** When > 0, we should not show context menus */
+	int32 ContextMenuSupression;
+	friend FContextMenuSuppressor;
 };
+
+struct FContextMenuSuppressor
+{
+	FContextMenuSuppressor(TSharedRef<FSequencerTimeSliderController> InTimeSliderController)
+		: TimeSliderController(InTimeSliderController)
+	{
+		++TimeSliderController->ContextMenuSupression;
+	}
+	~FContextMenuSuppressor()
+	{
+		--TimeSliderController->ContextMenuSupression;
+	}
+
+private:
+	FContextMenuSuppressor(const FContextMenuSuppressor&);
+	FContextMenuSuppressor& operator=(const FContextMenuSuppressor&);
+
+	TSharedRef<FSequencerTimeSliderController> TimeSliderController;
+};
\ No newline at end of file
diff --git a/Engine/Source/Editor/Sequencer/Private/Tools/EditToolDragOperations.cpp b/Engine/Source/Editor/Sequencer/Private/Tools/EditToolDragOperations.cpp
--- a/Engine/Source/Editor/Sequencer/Private/Tools/EditToolDragOperations.cpp
+++ b/Engine/Source/Editor/Sequencer/Private/Tools/EditToolDragOperations.cpp
@@ -137,6 +137,7 @@ void FEditToolDragOperation::EndTransaction()
 {
 	Transaction.Reset();
 	Sequencer.UpdateRuntimeInstances();
+	Sequencer.UpdatePlaybackRange();
 }
 
 FResizeSection::FResizeSection( FSequencer& InSequencer, TArray<FSectionHandle> InSections, bool bInDraggingByEnd )
diff --git a/Engine/Source/Editor/Sequencer/Public/ISequencer.h b/Engine/Source/Editor/Sequencer/Public/ISequencer.h
--- a/Engine/Source/Editor/Sequencer/Public/ISequencer.h
+++ b/Engine/Source/Editor/Sequencer/Public/ISequencer.h
@@ -17,6 +17,23 @@ class UMovieSceneSection;
 enum class EMapChangeType : uint8;
 
 
+/**
+ * Defines auto-key modes.
+ */
+UENUM()
+enum class EAutoKeyMode
+{
+	/** Key all properties that change. */
+	KeyAll,
+
+	/** Only key changed properties if they have exiting animations. */
+	KeyAnimated,
+
+	/** Don't auto-key */
+	KeyNone
+};
+
+
 /**
  * Interface for sequencers.
  */
@@ -77,11 +94,11 @@ public:
 	 */
 	virtual void AddSubSequence(UMovieSceneSequence* Sequence) = 0;
 
-	/** @return Returns whether auto-key is enabled in this sequencer */
-	virtual bool GetAutoKeyEnabled() const = 0;
+	/** @return Returns the current auto-key mode. */
+	virtual EAutoKeyMode GetAutoKeyMode() const = 0;
 
 	/** Sets whether autokey is enabled in this sequencer. */
-	virtual void SetAutoKeyEnabled(bool bAutoKeyEnabled) = 0;
+	virtual void SetAutoKeyMode(EAutoKeyMode AutoKeyMode) = 0;
 
 	/** @return Returns whether key all is enabled in this sequencer */
 	virtual bool GetKeyAllEnabled() const = 0;
@@ -184,6 +201,8 @@ public:
 
 	virtual void UpdateRuntimeInstances() = 0;
 
+	virtual void UpdatePlaybackRange() = 0;
+
 	virtual FSequencerSelection& GetSelection() = 0;
 	virtual FSequencerSelectionPreview& GetSelectionPreview() = 0;
 
diff --git a/Engine/Source/Editor/Sequencer/Public/ISequencerTrackEditor.h b/Engine/Source/Editor/Sequencer/Public/ISequencerTrackEditor.h
--- a/Engine/Source/Editor/Sequencer/Public/ISequencerTrackEditor.h
+++ b/Engine/Source/Editor/Sequencer/Public/ISequencerTrackEditor.h
@@ -17,12 +17,11 @@ public:
 
 	/**
 	 * Manually adds a key.
-
+	 *
 	 * @param ObjectGuid The Guid of the object that we are adding a key to.
-	 * @param AdditionalAsset An optional asset that can be added with the key.
 	 */
-	virtual void AddKey(const FGuid& ObjectGuid, UObject* AdditionalAsset = nullptr) = 0;
-	
+	virtual void AddKey(const FGuid& ObjectGuid) = 0;
+
 	/**
 	 * Add a new track to the sequence.
 	 */
diff --git a/Engine/Source/Editor/Sequencer/Public/MovieSceneTrackEditor.h b/Engine/Source/Editor/Sequencer/Public/MovieSceneTrackEditor.h
--- a/Engine/Source/Editor/Sequencer/Public/MovieSceneTrackEditor.h
+++ b/Engine/Source/Editor/Sequencer/Public/MovieSceneTrackEditor.h
@@ -74,6 +74,15 @@ public:
 		}
 	}
 
+	void UpdatePlaybackRange()
+	{
+		TSharedPtr<ISequencer> SequencerPin = Sequencer.Pin();
+		if( SequencerPin.IsValid()  )
+		{
+			SequencerPin->UpdatePlaybackRange();
+		}
+	}
+
 	void AnimatablePropertyChanged( FOnKeyProperty OnKeyProperty )
 	{
 		check(OnKeyProperty.IsBound());
@@ -98,6 +107,8 @@ public:
 				// Movie scene data has changed
 				NotifyMovieSceneDataChanged();
 			}
+
+			UpdatePlaybackRange();
 		}
 	}
 
@@ -190,7 +201,7 @@ public:
 
 	// ISequencerTrackEditor interface
 
-	virtual void AddKey( const FGuid& ObjectGuid, UObject* AdditionalAsset = nullptr ) override { }
+	virtual void AddKey( const FGuid& ObjectGuid ) override {}
 
 	virtual UMovieSceneTrack* AddTrack(UMovieScene* FocusedMovieScene, const FGuid& ObjectHandle, TSubclassOf<class UMovieSceneTrack> TrackClass, FName UniqueTypeName) override
 	{
@@ -214,7 +225,7 @@ public:
 	{
 		// @todo sequencer livecapture: This turns on "auto key" for the purpose of capture keys for actor state
 		// during PIE sessions when record mode is active.
-		return Sequencer.Pin()->IsRecordingLive() || Sequencer.Pin()->GetAutoKeyEnabled();
+		return Sequencer.Pin()->IsRecordingLive() || Sequencer.Pin()->GetAutoKeyMode() != EAutoKeyMode::KeyNone;
 	}
 
 	virtual TSharedRef<ISequencerSection> MakeSectionInterface(class UMovieSceneSection& SectionObject, UMovieSceneTrack& Track) = 0;
diff --git a/Engine/Source/Editor/UMGEditor/Private/Animation/MarginTrackEditor.cpp b/Engine/Source/Editor/UMGEditor/Private/Animation/MarginTrackEditor.cpp
--- a/Engine/Source/Editor/UMGEditor/Private/Animation/MarginTrackEditor.cpp
+++ b/Engine/Source/Editor/UMGEditor/Private/Animation/MarginTrackEditor.cpp
@@ -79,18 +79,20 @@ TSharedRef<FPropertySection> FMarginTrackEditor::MakePropertySectionInterface( U
 }
 
 
-void FMarginTrackEditor::GenerateKeysFromPropertyChanged( const FPropertyChangedParams& PropertyChangedParams, TArray<FMarginKey>& GeneratedKeys )
+void FMarginTrackEditor::GenerateKeysFromPropertyChanged( const FPropertyChangedParams& PropertyChangedParams, TArray<FMarginKey>& NewGeneratedKeys, TArray<FMarginKey>& DefaultGeneratedKeys)
 {
 	FName ChannelName = PropertyChangedParams.StructPropertyNameToKey;
 	FMargin Margin = PropertyChangedParams.GetPropertyValue<FMargin>();
 
-	GeneratedKeys.Add( FMarginKey( EKeyMarginChannel::Left, Margin.Left, ChannelName == NAME_None || ChannelName == LeftName ? EKeyMarginValueType::Key : EKeyMarginValueType::Default ) );
-	GeneratedKeys.Add( FMarginKey( EKeyMarginChannel::Top, Margin.Top, ChannelName == NAME_None || ChannelName == TopName ? EKeyMarginValueType::Key : EKeyMarginValueType::Default ) );
-	GeneratedKeys.Add( FMarginKey( EKeyMarginChannel::Right, Margin.Right, ChannelName == NAME_None || ChannelName == RightName ? EKeyMarginValueType::Key : EKeyMarginValueType::Default ) );
-	GeneratedKeys.Add( FMarginKey( EKeyMarginChannel::Bottom, Margin.Bottom, ChannelName == NAME_None || ChannelName == BottomName ? EKeyMarginValueType::Key : EKeyMarginValueType::Default ) );
-}
+	TArray<FMarginKey>& LeftKeys = ChannelName == NAME_None || ChannelName == LeftName ? NewGeneratedKeys : DefaultGeneratedKeys;
+	LeftKeys.Add( FMarginKey( EKeyMarginChannel::Left, Margin.Left ) );
 
-bool FMarginTrackEditor::ShouldAddKey(UMovieSceneMarginTrack* InTrack, FMarginKey InKey, FKeyParams InKeyParams) const
-{
-	return FPropertyTrackEditor::ShouldAddKey(InTrack, InKey, InKeyParams) && InKey.ValueType == EKeyMarginValueType::Key;
+	TArray<FMarginKey>& TopKeys = ChannelName == NAME_None || ChannelName == TopName ? NewGeneratedKeys : DefaultGeneratedKeys;
+	TopKeys.Add( FMarginKey( EKeyMarginChannel::Top, Margin.Top ) );
+
+	TArray<FMarginKey>& RightKeys = ChannelName == NAME_None || ChannelName == RightName ? NewGeneratedKeys : DefaultGeneratedKeys;
+	RightKeys.Add( FMarginKey( EKeyMarginChannel::Right, Margin.Right ) );
+
+	TArray<FMarginKey>& BottomKeys = ChannelName == NAME_None || ChannelName == BottomName ? NewGeneratedKeys : DefaultGeneratedKeys;
+	BottomKeys.Add( FMarginKey( EKeyMarginChannel::Bottom, Margin.Bottom ) );
 }
diff --git a/Engine/Source/Editor/UMGEditor/Private/Animation/MarginTrackEditor.h b/Engine/Source/Editor/UMGEditor/Private/Animation/MarginTrackEditor.h
--- a/Engine/Source/Editor/UMGEditor/Private/Animation/MarginTrackEditor.h
+++ b/Engine/Source/Editor/UMGEditor/Private/Animation/MarginTrackEditor.h
@@ -39,10 +39,7 @@ protected:
 	// FPropertyTrackEditor Interface
 
 	virtual TSharedRef<FPropertySection> MakePropertySectionInterface( UMovieSceneSection& SectionObject, UMovieSceneTrack& Track ) override;
-	virtual void GenerateKeysFromPropertyChanged( const FPropertyChangedParams& PropertyChangedParams, TArray<FMarginKey>& GeneratedKeys ) override;
-
-	// FKeyframeTrackEditor interface
-	virtual bool ShouldAddKey(UMovieSceneMarginTrack* InTrack, FMarginKey InKey, FKeyParams InKeyParams) const override;
+	virtual void GenerateKeysFromPropertyChanged( const FPropertyChangedParams& PropertyChangedParams, TArray<FMarginKey>& NewGeneratedKeys, TArray<FMarginKey>& DefaultGeneratedKeys) override;
 
 private:
 
diff --git a/Engine/Source/Editor/UMGEditor/Private/Animation/Sequencer2DTransformTrackEditor.cpp b/Engine/Source/Editor/UMGEditor/Private/Animation/Sequencer2DTransformTrackEditor.cpp
--- a/Engine/Source/Editor/UMGEditor/Private/Animation/Sequencer2DTransformTrackEditor.cpp
+++ b/Engine/Source/Editor/UMGEditor/Private/Animation/Sequencer2DTransformTrackEditor.cpp
@@ -121,21 +121,23 @@ TSharedRef<FPropertySection> F2DTransformTrackEditor::MakePropertySectionInterfa
 }
 
 
-void F2DTransformTrackEditor::GenerateKeysFromPropertyChanged( const FPropertyChangedParams& PropertyChangedParams, TArray<F2DTransformKey>& GeneratedKeys )
+void F2DTransformTrackEditor::GenerateKeysFromPropertyChanged( const FPropertyChangedParams& PropertyChangedParams, TArray<F2DTransformKey>& NewGeneratedKeys, TArray<F2DTransformKey>& DefaultGeneratedKeys )
 {
 	FName ChannelName = PropertyChangedParams.StructPropertyNameToKey;
 	FWidgetTransform Transform = PropertyChangedParams.GetPropertyValue<FWidgetTransform>();
 
-	GeneratedKeys.Add( F2DTransformKey(EKey2DTransformChannel::Translation, EKey2DTransformAxis::X, Transform.Translation.X, ChannelName == NAME_None || ChannelName == TranslationName ? EKey2DTransformValueType::Key : EKey2DTransformValueType::Default ) );
-	GeneratedKeys.Add( F2DTransformKey( EKey2DTransformChannel::Translation, EKey2DTransformAxis::Y, Transform.Translation.Y, ChannelName == NAME_None || ChannelName == TranslationName ? EKey2DTransformValueType::Key : EKey2DTransformValueType::Default ) );
-	GeneratedKeys.Add( F2DTransformKey( EKey2DTransformChannel::Scale, EKey2DTransformAxis::X, Transform.Scale.X, ChannelName == NAME_None || ChannelName == ScaleName ? EKey2DTransformValueType::Key : EKey2DTransformValueType::Default ) );
-	GeneratedKeys.Add( F2DTransformKey( EKey2DTransformChannel::Scale, EKey2DTransformAxis::Y, Transform.Scale.Y, ChannelName == NAME_None || ChannelName == ScaleName ? EKey2DTransformValueType::Key : EKey2DTransformValueType::Default ) );
-	GeneratedKeys.Add( F2DTransformKey( EKey2DTransformChannel::Shear, EKey2DTransformAxis::X, Transform.Shear.X, ChannelName == NAME_None || ChannelName == ShearName ? EKey2DTransformValueType::Key : EKey2DTransformValueType::Default ) );
-	GeneratedKeys.Add( F2DTransformKey( EKey2DTransformChannel::Shear, EKey2DTransformAxis::Y, Transform.Shear.Y, ChannelName == NAME_None || ChannelName == ShearName ? EKey2DTransformValueType::Key : EKey2DTransformValueType::Default ) );
-	GeneratedKeys.Add( F2DTransformKey( EKey2DTransformChannel::Rotation, EKey2DTransformAxis::None, Transform.Angle, ChannelName == NAME_None || ChannelName == AngleName ? EKey2DTransformValueType::Key : EKey2DTransformValueType::Default ) );
-}
+	TArray<F2DTransformKey>& TranslationKeys = ChannelName == NAME_None || ChannelName == TranslationName ? NewGeneratedKeys : DefaultGeneratedKeys;
+	TranslationKeys.Add( F2DTransformKey( EKey2DTransformChannel::Translation, EKey2DTransformAxis::X, Transform.Translation.X ) );
+	TranslationKeys.Add( F2DTransformKey( EKey2DTransformChannel::Translation, EKey2DTransformAxis::Y, Transform.Translation.Y ) );
 
-bool F2DTransformTrackEditor::ShouldAddKey(UMovieScene2DTransformTrack* InTrack, F2DTransformKey InKey, FKeyParams InKeyParams) const
-{
-	return FPropertyTrackEditor::ShouldAddKey(InTrack, InKey, InKeyParams) && InKey.ValueType == EKey2DTransformValueType::Key;
+	TArray<F2DTransformKey>& ScaleKeys = ChannelName == NAME_None || ChannelName == ScaleName ? NewGeneratedKeys : DefaultGeneratedKeys;
+	ScaleKeys.Add( F2DTransformKey( EKey2DTransformChannel::Scale, EKey2DTransformAxis::X, Transform.Scale.X ) );
+	ScaleKeys.Add( F2DTransformKey( EKey2DTransformChannel::Scale, EKey2DTransformAxis::Y, Transform.Scale.Y ) );
+
+	TArray<F2DTransformKey>& ShearKeys = ChannelName == NAME_None || ChannelName == ShearName ? NewGeneratedKeys : DefaultGeneratedKeys;
+	ShearKeys.Add( F2DTransformKey( EKey2DTransformChannel::Shear, EKey2DTransformAxis::X, Transform.Shear.X ) );
+	ShearKeys.Add( F2DTransformKey( EKey2DTransformChannel::Shear, EKey2DTransformAxis::Y, Transform.Shear.Y ) );
+
+	TArray<F2DTransformKey>& AngleKeys = ChannelName == NAME_None || ChannelName == AngleName ? NewGeneratedKeys : DefaultGeneratedKeys;
+	AngleKeys.Add( F2DTransformKey( EKey2DTransformChannel::Rotation, EKey2DTransformAxis::None, Transform.Angle ) );
 }
diff --git a/Engine/Source/Editor/UMGEditor/Private/Animation/Sequencer2DTransformTrackEditor.h b/Engine/Source/Editor/UMGEditor/Private/Animation/Sequencer2DTransformTrackEditor.h
--- a/Engine/Source/Editor/UMGEditor/Private/Animation/Sequencer2DTransformTrackEditor.h
+++ b/Engine/Source/Editor/UMGEditor/Private/Animation/Sequencer2DTransformTrackEditor.h
@@ -35,10 +35,7 @@ protected:
 
 	virtual TSharedRef<FPropertySection> MakePropertySectionInterface( UMovieSceneSection& SectionObject, UMovieSceneTrack& Track ) override;
 
-	virtual void GenerateKeysFromPropertyChanged( const FPropertyChangedParams& PropertyChangedParams, TArray<F2DTransformKey>& GeneratedKeys ) override;
-
-	// FKeyframeTrackEditor interface
-	virtual bool ShouldAddKey(UMovieScene2DTransformTrack* InTrack, F2DTransformKey InKey, FKeyParams InKeyParams) const override;
+	virtual void GenerateKeysFromPropertyChanged( const FPropertyChangedParams& PropertyChangedParams, TArray<F2DTransformKey>& NewGeneratedKeys, TArray<F2DTransformKey>& DefaultGeneratedKeys ) override;
 
 private:
 	static FName TranslationName;
diff --git a/Engine/Source/Editor/UMGEditor/Private/Details/SWidgetDetailsView.cpp b/Engine/Source/Editor/UMGEditor/Private/Details/SWidgetDetailsView.cpp
--- a/Engine/Source/Editor/UMGEditor/Private/Details/SWidgetDetailsView.cpp
+++ b/Engine/Source/Editor/UMGEditor/Private/Details/SWidgetDetailsView.cpp
@@ -467,7 +467,7 @@ void SWidgetDetailsView::HandleIsVariableChanged(ECheckBoxState CheckState)
 void SWidgetDetailsView::NotifyPreChange(FEditPropertyChain* PropertyAboutToChange)
 {
 	// During auto-key do not migrate values
-	if( !BlueprintEditor.Pin()->GetSequencer()->GetAutoKeyEnabled() )
+	if( BlueprintEditor.Pin()->GetSequencer()->GetAutoKeyMode() != EAutoKeyMode::KeyNone )
 	{
 		TSharedPtr<FWidgetBlueprintEditor> Editor = BlueprintEditor.Pin();
 
@@ -480,7 +480,7 @@ void SWidgetDetailsView::NotifyPostChange(const FPropertyChangedEvent& PropertyC
 {
 	const static FName DesignerRebuildName("DesignerRebuild");
 
-	if ( PropertyChangedEvent.ChangeType != EPropertyChangeType::Interactive && !BlueprintEditor.Pin()->GetSequencer()->GetAutoKeyEnabled() )
+	if ( PropertyChangedEvent.ChangeType != EPropertyChangeType::Interactive && BlueprintEditor.Pin()->GetSequencer()->GetAutoKeyMode() != EAutoKeyMode::KeyNone )
 	{
 		TSharedPtr<FWidgetBlueprintEditor> Editor = BlueprintEditor.Pin();
 
diff --git a/Engine/Source/Editor/UMGEditor/Private/WidgetBlueprintEditor.cpp b/Engine/Source/Editor/UMGEditor/Private/WidgetBlueprintEditor.cpp
--- a/Engine/Source/Editor/UMGEditor/Private/WidgetBlueprintEditor.cpp
+++ b/Engine/Source/Editor/UMGEditor/Private/WidgetBlueprintEditor.cpp
@@ -567,7 +567,7 @@ TSharedPtr<ISequencer>& FWidgetBlueprintEditor::GetSequencer()
 			SequencerInitParams.ViewParams = ViewParams;
 			SequencerInitParams.RootSequence = NullAnimation;
 			SequencerInitParams.bEditWithinLevelEditor = false;
-			SequencerInitParams.ToolkitHost = nullptr;
+			SequencerInitParams.ToolkitHost = GetToolkitHost();
 		};
 
 		Sequencer = FModuleManager::LoadModuleChecked<ISequencerModule>("Sequencer").CreateSequencer(SequencerInitParams);
@@ -597,7 +597,7 @@ void FWidgetBlueprintEditor::ChangeViewedAnimation( UWidgetAnimation& InAnimatio
 		{
 			// Disable sequencer from interaction
 			Sequencer->GetSequencerWidget()->SetEnabled(false);
-			Sequencer->SetAutoKeyEnabled(false);
+			Sequencer->SetAutoKeyMode(EAutoKeyMode::KeyNone);
 			NoAnimationTextBlockPin->SetVisibility(EVisibility::Visible);
 			SequencerOverlayPin->SetVisibility( EVisibility::HitTestInvisible );
 		}
@@ -675,7 +675,11 @@ void FWidgetBlueprintEditor::UpdatePreview(UBlueprint* InBlueprint, bool bInForc
 	}
 
 	OnWidgetPreviewUpdated.Broadcast();
-	Sequencer->UpdateRuntimeInstances();
+
+	if (Sequencer.IsValid())
+	{
+		Sequencer->UpdateRuntimeInstances();
+	}
 }
 
 FGraphAppearanceInfo FWidgetBlueprintEditor::GetGraphAppearance(UEdGraph* InGraph) const
diff --git a/Engine/Source/Editor/UnrealEd/Classes/Editor/EditorEngine.h b/Engine/Source/Editor/UnrealEd/Classes/Editor/EditorEngine.h
--- a/Engine/Source/Editor/UnrealEd/Classes/Editor/EditorEngine.h
+++ b/Engine/Source/Editor/UnrealEd/Classes/Editor/EditorEngine.h
@@ -748,7 +748,8 @@ public:
 	bool	HandleListMapPackageDependenciesCommand( const TCHAR* Str, FOutputDevice& Ar );
 	bool	HandleRebuildVolumesCommand( const TCHAR* Str, FOutputDevice& Ar, UWorld* InWorld );
 	bool	HandleRemoveArchtypeFlagCommand( const TCHAR* Str, FOutputDevice& Ar );
-
+	bool	HandleStartMovieCaptureCommand( const TCHAR* Cmd, FOutputDevice& Ar );
+	
 	/**
 	 * Initializes the Editor.
 	 */
diff --git a/Engine/Source/Editor/UnrealEd/Classes/Settings/EditorExperimentalSettings.h b/Engine/Source/Editor/UnrealEd/Classes/Settings/EditorExperimentalSettings.h
--- a/Engine/Source/Editor/UnrealEd/Classes/Settings/EditorExperimentalSettings.h
+++ b/Engine/Source/Editor/UnrealEd/Classes/Settings/EditorExperimentalSettings.h
@@ -1,9 +1,5 @@
 // Copyright 1998-2015 Epic Games, Inc. All Rights Reserved.
 
-/*=============================================================================
-	EditorExperimentalSettings.h: Declares the UEditorExperimentalSettings class.
-=============================================================================*/
-
 #pragma once
 
 
@@ -45,10 +41,6 @@ public:
 	UPROPERTY(EditAnywhere, config, Category=Tools, meta=(ConfigRestartRequired=true))
 	bool bBlueprintableComponents;
 
-	/** The Messaging Debugger provides a visual utility for debugging the messaging system. */
-	UPROPERTY(EditAnywhere, config, Category=Tools, meta=(DisplayName="Messaging Debugger"))
-	bool bMessagingDebugger;
-
 	/** Allows to use actor merging utilities (Simplygon Proxy LOD, Grouping by Materials)*/
 	UPROPERTY(EditAnywhere, config, Category = Tools, meta = (DisplayName = "Actor Merging"))
 	bool bActorMerging;
@@ -134,7 +126,10 @@ public:
 	 * @return The delegate.
 	 */
 	DECLARE_EVENT_OneParam(UEditorExperimentalSettings, FSettingChangedEvent, FName /*PropertyName*/);
-	FSettingChangedEvent& OnSettingChanged( ) { return SettingChangedEvent; }
+	FSettingChangedEvent& OnSettingChanged( )
+	{
+		return SettingChangedEvent;
+	}
 
 protected:
 
diff --git a/Engine/Source/Editor/UnrealEd/Classes/Settings/LevelEditorPlaySettings.h b/Engine/Source/Editor/UnrealEd/Classes/Settings/LevelEditorPlaySettings.h
--- a/Engine/Source/Editor/UnrealEd/Classes/Settings/LevelEditorPlaySettings.h
+++ b/Engine/Source/Editor/UnrealEd/Classes/Settings/LevelEditorPlaySettings.h
@@ -248,6 +248,9 @@ public:
 	UPROPERTY(config, EditAnywhere, Category=PlayOnDevice)
 	bool bAutoCompileBlueprintsOnLaunch;
 
+	/** A programmatically defined custom PIE window to use */
+	TWeakPtr<SWindow> CustomPIEWindow;
+	
 private:
 
 	/** NetMode to use for Play In Editor. */
diff --git a/Engine/Source/Editor/UnrealEd/Private/Analytics/EditorAnalytics.cpp b/Engine/Source/Editor/UnrealEd/Private/Analytics/EditorAnalytics.cpp
--- a/Engine/Source/Editor/UnrealEd/Private/Analytics/EditorAnalytics.cpp
+++ b/Engine/Source/Editor/UnrealEd/Private/Analytics/EditorAnalytics.cpp
@@ -55,51 +55,51 @@ public:
 			// Scoped lock
 			FSystemWideCriticalSection StoredValuesLock(EditorAnalyticsDefs::GlobalLockName, bWaitForLock ? EditorAnalyticsDefs::GlobalLockWaitTimeout : FTimespan::Zero());
 
-			// Get list of sessions in storage
+		// Get list of sessions in storage
 			if (StoredValuesLock.IsValid() && BeginReadWriteRecords())
+		{
+			TArray<FSessionRecord> SessionRecordsToDelete;
+
+			// Attempt check each stored session
+			for (FSessionRecord& Record : SessionRecords)
 			{
-				TArray<FSessionRecord> SessionRecordsToDelete;
+				FTimespan RecordAge = FDateTime::UtcNow() - Record.Timestamp;
 
-				// Attempt check each stored session
-				for (FSessionRecord& Record : SessionRecords)
+				if (Record.bCrashed)
 				{
-					FTimespan RecordAge = FDateTime::UtcNow() - Record.Timestamp;
-
-					if (Record.bCrashed)
-					{
-						// Crashed sessions
-						SessionRecordsToReport.Add(Record);
-						SessionRecordsToDelete.Add(Record);
-					}
+					// Crashed sessions
+					SessionRecordsToReport.Add(Record);
+					SessionRecordsToDelete.Add(Record);
+				}
 					else if (RecordAge > EditorAnalyticsDefs::SessionRecordExpiration)
-					{
-						// Delete expired session records
-						SessionRecordsToDelete.Add(Record);
-					}
-					else if (RecordAge > EditorAnalyticsDefs::SessionRecordTimeout)
-					{
-						// Timed out sessions
-						SessionRecordsToReport.Add(Record);
-						SessionRecordsToDelete.Add(Record);
-					}
+				{
+					// Delete expired session records
+					SessionRecordsToDelete.Add(Record);
 				}
-
-				for (FSessionRecord& DeletingRecord : SessionRecordsToDelete)
+				else if (RecordAge > EditorAnalyticsDefs::SessionRecordTimeout)
 				{
-					DeleteStoredRecord(DeletingRecord);
+					// Timed out sessions
+					SessionRecordsToReport.Add(Record);
+					SessionRecordsToDelete.Add(Record);
 				}
+			}
 
-				// Create a session record for this session
-				CreateAndWriteRecordForSession();
+			for (FSessionRecord& DeletingRecord : SessionRecordsToDelete)
+			{
+				DeleteStoredRecord(DeletingRecord);
+			}
 
-				// Update and release list of sessions in storage
-				EndReadWriteRecords();
+			// Create a session record for this session
+			CreateAndWriteRecordForSession();
 
-				// Register for crash callbacks
-				FCoreDelegates::OnHandleSystemError.AddRaw(this, &FSessionManager::OnCrashing);
+			// Update and release list of sessions in storage
+			EndReadWriteRecords();
 
-				bInitialized = true;
-			}
+			// Register for crash callbacks
+			FCoreDelegates::OnHandleSystemError.AddRaw(this, &FSessionManager::OnCrashing);
+
+			bInitialized = true;
+		}
 		}
 
 		for (FSessionRecord& ReportingSession : SessionRecordsToReport)
@@ -250,6 +250,35 @@ private:
 		SessionRecords.RemoveAll([&SessionId](const FSessionRecord& X){ return X.SessionId == SessionId; });
 	}
 
+	/**
+  	 * @EventName Editor.AbnormalShutdown
+	 *
+	 * @Trigger Fired only by the editor at startup, once for each "abnormal shutdown" detected that has not already been sent.
+	 *
+	 * @Type Static
+	 *
+	 * @EventParam SessionId Analytics SessionID of the session that abnormally terminated.
+	 * @EventParam EngineVersion EngineVersion of the session that abnormally terminated.
+	 * @EventParam ShutdownType - one of Crashed, Debugger, or AbormalShutdown
+	 *               * Crashed - we definitely detected a crash (whether or not a debugger was attached)
+	 *               * Debugger - the session crashed or shutdown abnormally, but we had a debugger attached at startup, so abnormal termination is much more likely because the user was debugging.
+	 *               * AbnormalShutdown - this happens when we didn't detect a normal shutdown, but none of the above cases is the cause.
+	 * @EventParam Timestamp - the UTC time of the last known time the crashed session was running, within 5 minutes.
+	 *
+	 * @TODO: Debugger should be a completely separate flag, since it's orthogonal to whether we detect a crash or shutdown.
+	 *
+	 * @Comments The editor will only try to check for abnormal terminations if it determines it is a real editor run (not a commandlet or PIE, or editor -game run), and the user has not disabled sending Engine Usage data to Epic via the Editor global preferences.
+	 *
+	 * SessionId can be used to find the actual session associated with this crash in the data.
+	 *
+	 * Use THIS instead of the current event's analytics AppVersion field, since the crash could be from an old version. Could technically be found by searching the data for the SessionId and checking the version there, but sending it here simplifies the backend processing and ensure we have the version here even if we never got the original events.
+	 *
+	 * If multiple versions of the editor or launched, this code will properly track each one and its shutdown status. So during startup, an editor instance may need to fire off several events.
+	 *
+	 * When attributing abnormal terminations to engine versions, be sure to use the EngineVersion associated with this event, and not the AppVersion. AppVersion is for the session that is currently sending the event, not for the session that crashed. That is why EngineVersion is sent separately.
+	 *
+	 * The editor updates this timestamp every 5 minutes, so we should know the time of the crash within 5 minutes. It should technically correlate with the last heartbeat we receive in the data for that session.
+	 */
 	void SendAbnormalShutdownReport(const FSessionRecord& Record)
 	{
 		FGameProjectGenerationModule& GameProjectModule = FModuleManager::LoadModuleChecked<FGameProjectGenerationModule>(TEXT("GameProjectGeneration"));
diff --git a/Engine/Source/Editor/UnrealEd/Private/EditorServer.cpp b/Engine/Source/Editor/UnrealEd/Private/EditorServer.cpp
--- a/Engine/Source/Editor/UnrealEd/Private/EditorServer.cpp
+++ b/Engine/Source/Editor/UnrealEd/Private/EditorServer.cpp
@@ -30,6 +30,7 @@
 #include "Matinee/InterpData.h"
 #include "Animation/SkeletalMeshActor.h"
 #include "InstancedFoliageActor.h"
+#include "MovieSceneCaptureModule.h"
 
 #include "Editor/UnrealEd/Public/Kismet2/KismetEditorUtilities.h"
 #include "Editor/PropertyEditor/Public/PropertyEditorModule.h"
@@ -70,7 +71,6 @@
 #include "AI/Navigation/NavLinkRenderingComponent.h"
 #include "PhysicsPublic.h"
 #include "AnimationRecorder.h"
-#include "Settings/EditorSettings.h"
 #include "Analytics/AnalyticsPrivacySettings.h"
 #include "KismetReinstanceUtilities.h"
 
@@ -1671,7 +1671,7 @@ void UEditorEngine::RebuildModelFromBrushes(UModel* Model, bool bSelectedBrushes
 	TArray<ABrush*> DynamicBrushes;
 	if (!bTreatMovableBrushesAsStatic)
 	{
-		for (auto It(Level->Actors.CreateConstIterator()); It; ++It)
+	for( auto It(Level->Actors.CreateConstIterator()); It; ++It )
 		{
 			ABrush* DynamicBrush = Cast<ABrush>(*It);
 			if (DynamicBrush && DynamicBrush->Brush && !DynamicBrush->IsStaticBrush() &&
@@ -3427,7 +3427,7 @@ void UEditorEngine::PasteSelectedActorsFromClipboard( UWorld* InWorld, const FTe
 
 				// Update the pivot location.
 				check(SingleActor);
-				SetPivot(SingleActor->GetActorLocation(), false, true);
+				SetPivot( SingleActor->GetActorLocation(), false, true );
 
 				// If grouping is active, go through the unique group actors and update the group actor location
 				if (GEditor->bGroupingActive)
@@ -4826,6 +4826,15 @@ bool UEditorEngine::SnapObjectTo( FActorOrComponent Object, const bool InAlign,
 		FVector NewLocation = Hit.Location - LocationOffset;
 		NewLocation.Z += KINDA_SMALL_NUMBER;	// Move the new desired location up by an error tolerance
 		
+		if (Object.Actor)
+		{
+			GEditor->BroadcastBeginObjectMovement(*Object.Actor);
+		}
+		else
+		{
+			GEditor->BroadcastBeginObjectMovement(*Object.Component);
+		}
+
 		Object.SetWorldLocation( NewLocation );
 		//InActor->TeleportTo( NewLocation, InActor->GetActorRotation(), false,true );
 		
@@ -4837,6 +4846,16 @@ bool UEditorEngine::SnapObjectTo( FActorOrComponent Object, const bool InAlign,
 			Object.SetWorldRotation( NewRotation );
 		}
 
+		if (Object.Actor)
+		{
+			GEditor->BroadcastEndObjectMovement(*Object.Actor);
+		}
+		else
+		{
+			GEditor->BroadcastEndObjectMovement(*Object.Component);
+		}
+
+
 		// Switch to the pie world if we have one
 		FScopedConditionalWorldSwitcher WorldSwitcher( GCurrentLevelEditingViewportClient );
 
@@ -5690,6 +5709,10 @@ bool UEditorEngine::Exec( UWorld* InWorld, const TCHAR* Stream, FOutputDevice& A
 	{
 		HandleRemoveArchtypeFlagCommand( Str, Ar );
 	}
+	else if( FParse::Command(&Str,TEXT("STARTMOVIECAPTURE")) )
+	{
+		bProcessed = HandleStartMovieCaptureCommand( Str, Ar );
+	}
 	else
 	{
 		bProcessed = FBlueprintEditorUtils::KismetDiagnosticExec(Stream, Ar);
@@ -6459,6 +6482,31 @@ bool UEditorEngine::HandleRemoveArchtypeFlagCommand( const TCHAR* Str, FOutputDe
 	return true;
 }
 
+bool UEditorEngine::HandleStartMovieCaptureCommand( const TCHAR* Cmd, FOutputDevice& Ar )
+{
+	IMovieSceneCaptureInterface* CaptureInterface = IMovieSceneCaptureModule::Get().GetFirstActiveMovieSceneCapture();
+	if (CaptureInterface)
+	{
+		CaptureInterface->StartCapturing();
+		return true;
+	}
+
+	for (const FWorldContext& Context : GEngine->GetWorldContexts())
+	{
+		if (Context.WorldType == EWorldType::PIE)
+		{
+			FSlatePlayInEditorInfo* SlatePlayInEditorSession = GEditor->SlatePlayInEditorMap.Find(Context.ContextHandle);
+			if (SlatePlayInEditorSession && SlatePlayInEditorSession->SlatePlayInEditorWindowViewport.IsValid())
+			{
+				IMovieSceneCaptureModule::Get().CreateMovieSceneCapture(SlatePlayInEditorSession->SlatePlayInEditorWindowViewport.ToSharedRef());
+				return true;
+			}
+		}
+	}
+
+	return false;
+}
+
 /**
  * @return true if the given component's StaticMesh can be merged with other StaticMeshes
  */
diff --git a/Engine/Source/Editor/UnrealEd/Private/PlayLevel.cpp b/Engine/Source/Editor/UnrealEd/Private/PlayLevel.cpp
--- a/Engine/Source/Editor/UnrealEd/Private/PlayLevel.cpp
+++ b/Engine/Source/Editor/UnrealEd/Private/PlayLevel.cpp
@@ -2898,33 +2898,42 @@ UGameInstance* UEditorEngine::CreatePIEGameInstance(int32 PIEInstance, bool bInS
 					CenterNewWindow = true;
 				}
 
-				TSharedRef<SWindow> PieWindow = SNew(SWindow)
+				TSharedPtr<SWindow> PieWindow = PlayInSettings->CustomPIEWindow.Pin();
+
+				const bool bHasCustomWindow = PieWindow.IsValid();
+				if (!bHasCustomWindow)
+				{
+					PieWindow = SNew(SWindow)
 					.Title(ViewportName)
 					.ScreenPosition(FVector2D( NewWindowPosition.X, NewWindowPosition.Y ))
 					.ClientSize(FVector2D( NewWindowWidth, NewWindowHeight ))
 					.AutoCenter(CenterNewWindow ? EAutoCenter::PreferredWorkArea : EAutoCenter::None)
 					.UseOSWindowBorder(bUseOSWndBorder)
 					.SizingRule(ESizingRule::UserSized);
+				}
 
 
 				// Setup a delegate for switching to the play world on slate input events, drawing and ticking
 				FOnSwitchWorldHack OnWorldSwitch = FOnSwitchWorldHack::CreateUObject( this, &UEditorEngine::OnSwitchWorldForSlatePieWindow );
 				PieWindow->SetOnWorldSwitchHack( OnWorldSwitch );
 
-				// Mac does not support parenting, do not keep on top
-#if PLATFORM_MAC
-				FSlateApplication::Get().AddWindow(PieWindow);
-#else
-				TSharedRef<SWindow, ESPMode::NotThreadSafe> MainWindow = FModuleManager::LoadModuleChecked<IMainFrameModule>(TEXT("MainFrame")).GetParentWindow().ToSharedRef();
-				if (PlayInSettings->PIEAlwaysOnTop)
-				{
-					FSlateApplication::Get().AddWindowAsNativeChild(PieWindow, MainWindow, true);
-				}
-				else
+				if (!bHasCustomWindow)
 				{
-					FSlateApplication::Get().AddWindow(PieWindow);
+					// Mac does not support parenting, do not keep on top
+	#if PLATFORM_MAC
+					FSlateApplication::Get().AddWindow(PieWindow.ToSharedRef());
+	#else
+					TSharedRef<SWindow, ESPMode::NotThreadSafe> MainWindow = FModuleManager::LoadModuleChecked<IMainFrameModule>(TEXT("MainFrame")).GetParentWindow().ToSharedRef();
+					if (PlayInSettings->PIEAlwaysOnTop)
+					{
+						FSlateApplication::Get().AddWindowAsNativeChild(PieWindow.ToSharedRef(), MainWindow, true);
+					}
+					else
+					{
+						FSlateApplication::Get().AddWindow(PieWindow.ToSharedRef());
+					}
+	#endif
 				}
-#endif
 
 				TSharedRef<SOverlay> ViewportOverlayWidgetRef = SNew(SOverlay);
 
@@ -2947,8 +2956,11 @@ UGameInstance* UEditorEngine::CreatePIEGameInstance(int32 PIEInstance, bool bInS
 				// Create a viewport widget for the game to render in.
 				PieWindow->SetContent( PieViewportWidget.ToSharedRef() );
 
-				// Ensure the PIE window appears does not appear behind other windows.
-				PieWindow->BringToFront();
+				if (!bHasCustomWindow)
+				{
+					// Ensure the PIE window appears does not appear behind other windows.
+					PieWindow->BringToFront();
+				}
 
 				ViewportClient->SetViewportOverlayWidget( PieWindow, ViewportOverlayWidgetRef );
 				ViewportClient->SetGameLayerManager(GameLayerManagerRef);
diff --git a/Engine/Source/Programs/CrashReporter/CrashReportClient/Private/CrashReportClient.cpp b/Engine/Source/Programs/CrashReporter/CrashReportClient/Private/CrashReportClient.cpp
--- a/Engine/Source/Programs/CrashReporter/CrashReportClient/Private/CrashReportClient.cpp
+++ b/Engine/Source/Programs/CrashReporter/CrashReportClient/Private/CrashReportClient.cpp
@@ -29,39 +29,39 @@ FCrashReportClient::FCrashReportClient(const FPlatformErrorReport& InErrorReport
 	: DiagnosticText( LOCTEXT("ProcessingReport", "Processing crash report ...") )
 	, DiagnoseReportTask(nullptr)
 	, ErrorReport( InErrorReport )
-	, Uploader( FCrashReportClientConfig::Get().GetReceiverAddress() )
-	, bBeginUploadCalled(false)
+	, ReceiverUploader(FCrashReportClientConfig::Get().GetReceiverAddress())
+	, DataRouterUploader(FCrashReportClientConfig::Get().GetDataRouterURL())
 	, bShouldWindowBeHidden(false)
 	, bSendData(false)
 {
 	if (FPrimaryCrashProperties::Get()->IsValid())
 	{
-		bool bUsePrimaryData = false;
-		if (FPrimaryCrashProperties::Get()->HasProcessedData())
+	bool bUsePrimaryData = false;
+	if (FPrimaryCrashProperties::Get()->HasProcessedData())
+	{
+		bUsePrimaryData = true;
+	}
+	else
+	{
+		if (!ErrorReport.TryReadDiagnosticsFile() && !FParse::Param( FCommandLine::Get(), TEXT( "no-local-diagnosis" ) ))
 		{
-			bUsePrimaryData = true;
+			DiagnoseReportTask = new FAsyncTask<FDiagnoseReportWorker>( this );
+			DiagnoseReportTask->StartBackgroundTask();
 		}
 		else
 		{
-			if (!ErrorReport.TryReadDiagnosticsFile() && !FParse::Param( FCommandLine::Get(), TEXT( "no-local-diagnosis" ) ))
-			{
-				DiagnoseReportTask = new FAsyncTask<FDiagnoseReportWorker>( this );
-				DiagnoseReportTask->StartBackgroundTask();
-			}
-			else
-			{
-				bUsePrimaryData = true;
-			}
+			bUsePrimaryData = true;
 		}
+	}
 
-		if (bUsePrimaryData)
-		{
-			const FString CallstackString = FPrimaryCrashProperties::Get()->CallStack.AsString();
-			const FString ReportString = FString::Printf( TEXT( "%s\n\n%s" ), *FPrimaryCrashProperties::Get()->ErrorMessage.AsString(), *CallstackString );
-			DiagnosticText = FText::FromString( ReportString );
+	if (bUsePrimaryData)
+	{
+		const FString CallstackString = FPrimaryCrashProperties::Get()->CallStack.AsString();
+		const FString ReportString = FString::Printf( TEXT( "%s\n\n%s" ), *FPrimaryCrashProperties::Get()->ErrorMessage.AsString(), *CallstackString );
+		DiagnosticText = FText::FromString( ReportString );
 
-			FormattedDiagnosticText = FCrashReportUtil::FormatDiagnosticText( FText::FromString( ReportString ) );
-		}
+		FormattedDiagnosticText = FCrashReportUtil::FormatDiagnosticText( FText::FromString( ReportString ) );
+	}
 	}
 }
 
@@ -190,22 +190,45 @@ bool FCrashReportClient::Tick(float UnusedDeltaTime)
 	
 	if( bSendData )
 	{
-		if( !bBeginUploadCalled )
+		if (!FCrashUploadBase::IsInitialized())
+		{
+			FCrashUploadBase::StaticInitialize( ErrorReport );
+		}
+
+		if( !ReceiverUploader.IsUploadCalled())
+		{
+			// Can be called only when we have all files.
+			ReceiverUploader.BeginUpload( ErrorReport );
+		}
+
+		// IsWorkDone will always return true here (since ReceiverUploader can't finish until the diagnosis has been sent), but it
+		//  has the side effect of joining the worker thread.
+		if( !ReceiverUploader.IsFinished() )
+		{
+			// More ticks, please
+			return true;
+		}
+
+		if (!DataRouterUploader.IsUploadCalled())
 		{
 			// Can be called only when we have all files.
-			Uploader.BeginUpload( ErrorReport );
-			bBeginUploadCalled = true;
+			DataRouterUploader.BeginUpload(ErrorReport);
 		}
 
-		// IsWorkDone will always return true here (since uploader can't finish until the diagnosis has been sent), but it
+		// IsWorkDone will always return true here (since DataRouterUploader can't finish until the diagnosis has been sent), but it
 		//  has the side effect of joining the worker thread.
-		if( !Uploader.IsFinished() )
+		if (!DataRouterUploader.IsFinished())
 		{
 			// More ticks, please
 			return true;
 		}
 	}
 
+	if (FCrashUploadBase::IsInitialized())
+	{
+		FCrashUploadBase::StaticShutdown();
+	}
+
 	FPlatformMisc::RequestExit(false);
 	return false;
 }
diff --git a/Engine/Source/Programs/CrashReporter/CrashReportClient/Private/CrashReportClient.h b/Engine/Source/Programs/CrashReporter/CrashReportClient/Private/CrashReportClient.h
--- a/Engine/Source/Programs/CrashReporter/CrashReportClient/Private/CrashReportClient.h
+++ b/Engine/Source/Programs/CrashReporter/CrashReportClient/Private/CrashReportClient.h
@@ -160,10 +160,10 @@ private:
 	FPlatformErrorReport ErrorReport;
 
 	/** Object that uploads report files to the server */
-	FCrashUpload Uploader;
+	FCrashUploadToReceiver ReceiverUploader;
 
-	/** Whether BeginUpload has been called. */
-	bool bBeginUploadCalled;
+	/** Object that uploads report files to the server */
+	FCrashUploadToDataRouter DataRouterUploader;
 
 	/** Whether the main window should be hidden. */
 	bool bShouldWindowBeHidden;
diff --git a/Engine/Source/Programs/CrashReporter/CrashReportClient/Private/CrashReportClientApp.cpp b/Engine/Source/Programs/CrashReporter/CrashReportClient/Private/CrashReportClientApp.cpp
--- a/Engine/Source/Programs/CrashReporter/CrashReportClient/Private/CrashReportClientApp.cpp
+++ b/Engine/Source/Programs/CrashReporter/CrashReportClient/Private/CrashReportClientApp.cpp
@@ -132,69 +132,72 @@ void RunCrashReportClient(const TCHAR* CommandLine)
 	FPlatformErrorReport::Init();
 	auto ErrorReport = LoadErrorReport();
 	
-	if (bUnattended)
+	if (ErrorReport.HasFilesToUpload())
 	{
-		// In the unattended mode we don't send any PII.
-		ErrorReport.SetUserComment( NSLOCTEXT( "CrashReportClient", "UnattendedMode", "Sent in the unattended mode" ) );
-		FCrashReportClientUnattended CrashReportClient( ErrorReport );
-
-		// loop until the app is ready to quit
-		while (!GIsRequestingExit)
+		if (bUnattended)
 		{
-			MainLoop.Tick();
+			// In the unattended mode we don't send any PII.
+			ErrorReport.SetUserComment(NSLOCTEXT("CrashReportClient", "UnattendedMode", "Sent in the unattended mode"));
+			FCrashReportClientUnattended CrashReportClient(ErrorReport);
+
+			// loop until the app is ready to quit
+			while (!GIsRequestingExit)
+			{
+				MainLoop.Tick();
+			}
 		}
-	}
-	else
-	{
+		else
+		{
 #if !CRASH_REPORT_UNATTENDED_ONLY
-		// crank up a normal Slate application using the platform's standalone renderer
-		FSlateApplication::InitializeAsStandaloneApplication(GetStandardStandaloneRenderer());
+			// crank up a normal Slate application using the platform's standalone renderer
+			FSlateApplication::InitializeAsStandaloneApplication(GetStandardStandaloneRenderer());
 
-		// Prepare the custom Slate styles
-		FCrashReportClientStyle::Initialize();
+			// Prepare the custom Slate styles
+			FCrashReportClientStyle::Initialize();
 
-		// Create the main implementation object
-		TSharedRef<FCrashReportClient> CrashReportClient = MakeShareable(new FCrashReportClient(ErrorReport));
+			// Create the main implementation object
+			TSharedRef<FCrashReportClient> CrashReportClient = MakeShareable(new FCrashReportClient(ErrorReport));
 
-		// open up the app window	
-		TSharedRef<SCrashReportClient> ClientControl = SNew(SCrashReportClient, CrashReportClient);
+			// open up the app window	
+			TSharedRef<SCrashReportClient> ClientControl = SNew(SCrashReportClient, CrashReportClient);
 
-		auto Window = FSlateApplication::Get().AddWindow(
-			SNew(SWindow)
-			.Title(NSLOCTEXT("CrashReportClient", "CrashReportClientAppName", "Unreal Engine 4 Crash Reporter"))
-			.ClientSize(InitialWindowDimensions)
-			[
-				ClientControl
-			]);
+			auto Window = FSlateApplication::Get().AddWindow(
+				SNew(SWindow)
+				.Title(NSLOCTEXT("CrashReportClient", "CrashReportClientAppName", "Unreal Engine 4 Crash Reporter"))
+				.ClientSize(InitialWindowDimensions)
+				[
+					ClientControl
+				]);
 
-		Window->SetRequestDestroyWindowOverride(FRequestDestroyWindowOverride::CreateSP(CrashReportClient, &FCrashReportClient::RequestCloseWindow));
+			Window->SetRequestDestroyWindowOverride(FRequestDestroyWindowOverride::CreateSP(CrashReportClient, &FCrashReportClient::RequestCloseWindow));
 
-		// Setting focus seems to have to happen after the Window has been added
-		FSlateApplication::Get().ClearKeyboardFocus(EFocusCause::Cleared);
+			// Setting focus seems to have to happen after the Window has been added
+			FSlateApplication::Get().ClearKeyboardFocus(EFocusCause::Cleared);
 
-		// Debugging code
-		if (RunWidgetReflector)
-		{
-			FModuleManager::LoadModuleChecked<ISlateReflectorModule>("SlateReflector").DisplayWidgetReflector();
-		}
-
-		// loop until the app is ready to quit
-		while (!GIsRequestingExit)
-		{
-			MainLoop.Tick();
+			// Debugging code
+			if (RunWidgetReflector)
+			{
+				FModuleManager::LoadModuleChecked<ISlateReflectorModule>("SlateReflector").DisplayWidgetReflector();
+			}
 
-			if (CrashReportClient->ShouldWindowBeHidden())
+			// loop until the app is ready to quit
+			while (!GIsRequestingExit)
 			{
-				Window->HideWindow();
+				MainLoop.Tick();
+
+				if (CrashReportClient->ShouldWindowBeHidden())
+				{
+					Window->HideWindow();
+				}
 			}
-		}
 
-		// Clean up the custom styles
-		FCrashReportClientStyle::Shutdown();
+			// Clean up the custom styles
+			FCrashReportClientStyle::Shutdown();
 
-		// Close down the Slate application
-		FSlateApplication::Shutdown();
+			// Close down the Slate application
+			FSlateApplication::Shutdown();
 #endif // !CRASH_REPORT_UNATTENDED_ONLY
+		}
 	}
 
 	FPrimaryCrashProperties::Shutdown();
diff --git a/Engine/Source/Programs/CrashReporter/CrashReportClient/Private/CrashReportClientConfig.cpp b/Engine/Source/Programs/CrashReporter/CrashReportClient/Private/CrashReportClientConfig.cpp
--- a/Engine/Source/Programs/CrashReporter/CrashReportClient/Private/CrashReportClientConfig.cpp
+++ b/Engine/Source/Programs/CrashReporter/CrashReportClient/Private/CrashReportClientConfig.cpp
@@ -18,6 +18,14 @@ FCrashReportClientConfig::FCrashReportClientConfig()
 		// Use the default value.
 		CrashReportReceiverIP = TEXT( "http://crashreporter.epicgames.com:57005" );
 	}
+	UE_LOG(CrashReportClientLog, Log, TEXT("CrashReportReceiverIP: %s"), *CrashReportReceiverIP);
+
+	if (!GConfig->GetString(*SectionName, TEXT("DataRouterUrl"), DataRouterUrl, GEngineIni))
+	{
+		// Use the default value.
+		DataRouterUrl = TEXT("");
+	}
+	UE_LOG(CrashReportClientLog, Log, TEXT("DataRouterUrl: %s"), *DataRouterUrl);
 
 	if (!GConfig->GetBool( TEXT( "CrashReportClient" ), TEXT( "bAllowToBeContacted" ), bAllowToBeContacted, GEngineIni ))
 	{
@@ -44,8 +52,6 @@ FCrashReportClientConfig::FCrashReportClientConfig()
 	}
 	
 	ReadFullCrashDumpConfigurations();
-
-	UE_LOG( CrashReportClientLog, Log, TEXT( "CrashReportReceiverIP: %s" ), *CrashReportReceiverIP );
 }
 
 void FCrashReportClientConfig::SetAllowToBeContacted( bool bNewValue )
diff --git a/Engine/Source/Programs/CrashReporter/CrashReportClient/Private/CrashReportClientConfig.h b/Engine/Source/Programs/CrashReporter/CrashReportClient/Private/CrashReportClientConfig.h
--- a/Engine/Source/Programs/CrashReporter/CrashReportClient/Private/CrashReportClientConfig.h
+++ b/Engine/Source/Programs/CrashReporter/CrashReportClient/Private/CrashReportClientConfig.h
@@ -48,11 +48,16 @@ struct FCrashReportClientConfig
 	/** Initialization constructor. */
 	FCrashReportClientConfig();
 
-	const FString& GetReceiverAddress() const
+	FString GetReceiverAddress() const
 	{
 		return CrashReportReceiverIP;
 	}
 
+	FString GetDataRouterURL() const
+	{
+		return DataRouterUrl;
+	}
+
 	const FString& GetDiagnosticsFilename() const
 	{
 		return DiagnosticsFilename;
@@ -91,6 +96,9 @@ protected:
 	/** IP address of crash report receiver. */
 	FString CrashReportReceiverIP;
 
+	/** URL of Data Router service */
+	FString DataRouterUrl;
+
 	/** Filename to use when saving diagnostics report, if generated locally. */
 	FString DiagnosticsFilename;
 
diff --git a/Engine/Source/Programs/CrashReporter/CrashReportClient/Private/CrashReportClientUnattended.cpp b/Engine/Source/Programs/CrashReporter/CrashReportClient/Private/CrashReportClientUnattended.cpp
--- a/Engine/Source/Programs/CrashReporter/CrashReportClient/Private/CrashReportClientUnattended.cpp
+++ b/Engine/Source/Programs/CrashReporter/CrashReportClient/Private/CrashReportClientUnattended.cpp
@@ -5,8 +5,10 @@
 #include "CrashReportUtil.h"
 #include "CrashDescription.h"
 
-FCrashReportClientUnattended::FCrashReportClientUnattended(FPlatformErrorReport& ErrorReport)
-	: Uploader( FCrashReportClientConfig::Get().GetReceiverAddress() )
+FCrashReportClientUnattended::FCrashReportClientUnattended(FPlatformErrorReport& InErrorReport)
+	: ReceiverUploader(FCrashReportClientConfig::Get().GetReceiverAddress())
+	, DataRouterUploader(FCrashReportClientConfig::Get().GetDataRouterURL())
+	, ErrorReport(InErrorReport)
 {
 	ErrorReport.TryReadDiagnosticsFile();
 
@@ -16,20 +18,46 @@ FCrashReportClientUnattended::FCrashReportClientUnattended(FPlatformErrorReport&
 	// Update properties for the crash.
 	ErrorReport.SetPrimaryCrashProperties( *FPrimaryCrashProperties::Get() );
 
-	Uploader.BeginUpload( ErrorReport );
 	StartTicker();
 }
 
 bool FCrashReportClientUnattended::Tick(float UnusedDeltaTime)
 {
-	if (Uploader.IsFinished())
+	if (!FCrashUploadBase::IsInitialized())
 	{
-		FPlatformMisc::RequestExit(false /* don't force */);
-		// No more ticks, thank you
-		return false;
+		FCrashUploadBase::StaticInitialize(ErrorReport);
 	}
 
-	return true;	
+	if (!ReceiverUploader.IsUploadCalled())
+	{
+		// Can be called only when we have all files.
+		ReceiverUploader.BeginUpload(ErrorReport);
+	}
+
+	// IsWorkDone will always return true here (since ReceiverUploader can't finish until the diagnosis has been sent), but it
+	//  has the side effect of joining the worker thread.
+	if (!ReceiverUploader.IsFinished())
+	{
+		// More ticks, please
+		return true;
+	}
+
+	if (!DataRouterUploader.IsUploadCalled())
+	{
+		// Can be called only when we have all files.
+		DataRouterUploader.BeginUpload(ErrorReport);
+	}
+
+	// IsWorkDone will always return true here (since DataRouterUploader can't finish until the diagnosis has been sent), but it
+	//  has the side effect of joining the worker thread.
+	if (!DataRouterUploader.IsFinished())
+	{
+		// More ticks, please
+		return true;
+	}
+
+	FPlatformMisc::RequestExit(false /* don't force */);
+	return false;
 }
 
 void FCrashReportClientUnattended::StartTicker()
diff --git a/Engine/Source/Programs/CrashReporter/CrashReportClient/Private/CrashReportClientUnattended.h b/Engine/Source/Programs/CrashReporter/CrashReportClient/Private/CrashReportClientUnattended.h
--- a/Engine/Source/Programs/CrashReporter/CrashReportClient/Private/CrashReportClientUnattended.h
+++ b/Engine/Source/Programs/CrashReporter/CrashReportClient/Private/CrashReportClientUnattended.h
@@ -14,7 +14,7 @@ public:
 	 * Set up uploader object
 	 * @param ErrorReport Error report to upload
 	 */
-	explicit FCrashReportClientUnattended(FPlatformErrorReport& ErrorReport);
+	explicit FCrashReportClientUnattended(FPlatformErrorReport& InErrorReport);
 
 private:
 	/**
@@ -30,5 +30,11 @@ private:
 	void StartTicker();
 
 	/** Object that uploads report files to the server */
-	FCrashUpload Uploader;
+	FCrashUploadToReceiver ReceiverUploader;
+
+	/** Object that uploads report files to the server */
+	FCrashUploadToDataRouter DataRouterUploader;
+
+	/** Platform code for accessing the report */
+	FPlatformErrorReport ErrorReport;
 };
diff --git a/Engine/Source/Programs/CrashReporter/CrashReportClient/Private/CrashUpload.cpp b/Engine/Source/Programs/CrashReporter/CrashReportClient/Private/CrashUpload.cpp
--- a/Engine/Source/Programs/CrashReporter/CrashReportClient/Private/CrashUpload.cpp
+++ b/Engine/Source/Programs/CrashReporter/CrashReportClient/Private/CrashUpload.cpp
@@ -10,202 +10,178 @@
 #include "CrashReportUtil.h"
 #include "GenericPlatformCrashContext.h"
 #include "CrashDescription.h"
+#include "EngineBuildSettings.h"
 
 #define LOCTEXT_NAMESPACE "CrashReportClient"
 
-namespace
+namespace CrashUploadDefs
 {
 	const float PingTimeoutSeconds = 5.f;
 	// Ignore files bigger than 100MB; mini-dumps are smaller than this, but heap dumps can be very large
 	const int MaxFileSizeToUpload = 100 * 1024 * 1024;
-}
 
-FCrashUpload::FCrashUpload(const FString& ServerAddress)
-	: UrlPrefix(ServerAddress / "CrashReporter")
-	, State(EUploadState::NotSet)
-	, PauseState(EUploadState::Ready)
-{
-	SendPingRequest();
+	const FString APIKey(TEXT("CrashReporter"));
+	const FString AppEnvironmentInternal(TEXT("Dev"));
+	const FString AppEnvironmentExternal(TEXT("Release"));
+	const FString UploadType(TEXT("crashreports"));
 }
 
-FCrashUpload::~FCrashUpload()
+enum class ECompressedCrashFileHeader
 {
-	UE_LOG(CrashReportClientLog, Log, TEXT("Final state = %s"), ToString(State));
-}
+	MAGIC = 0x7E1B83C1,
+};
 
-bool FCrashUpload::PingTimeout(float DeltaTime)
+struct FCompressedCrashFile : FNoncopyable
 {
-	if (EUploadState::PingingServer == State)
-	{
-		SetCurrentState(EUploadState::ServerNotAvailable);
+	int32 CurrentFileIndex;
+	FString Filename;
+	TArray<uint8> Filedata;
 
-		// PauseState will be Ready if user has not yet decided to send the report
-		if (PauseState > EUploadState::Ready)
-		{
-			AddReportToPendingFile();
-		}
+	FCompressedCrashFile(int32 InCurrentFileIndex, const FString& InFilename, const TArray<uint8>& InFiledata)
+		: CurrentFileIndex(InCurrentFileIndex)
+		, Filename(InFilename)
+		, Filedata(InFiledata)
+	{
 	}
 
-	// One-shot
-	return false;
-}
+	/** Serialization operator. */
+	friend FArchive& operator << (FArchive& Ar, FCompressedCrashFile& Data)
+	{
+		Ar << Data.CurrentFileIndex;
+		Data.Filename.SerializeAsANSICharArray(Ar, 260);
+		Ar << Data.Filedata;
+		return Ar;
+	}
+};
 
-void FCrashUpload::BeginUpload(const FPlatformErrorReport& PlatformErrorReport)
+struct FCompressedHeader
 {
-	ErrorReport = PlatformErrorReport;
-	PendingFiles = FPlatformErrorReport( ErrorReport.GetReportDirectory() ).GetFilesToUpload();
-	UE_LOG(CrashReportClientLog, Log, TEXT("Got %d pending files to upload from '%s'"), PendingFiles.Num(), *ErrorReport.GetReportDirectoryLeafName());
+	FString DirectoryName;
+	FString FileName;
+	int32 UncompressedSize;
+	int32 FileCount;
 
-	PauseState = EUploadState::Finished;
-	if (State == EUploadState::Ready)
-	{
-		BeginUploadImpl();
-	}
-	else if (State == EUploadState::ServerNotAvailable)
+	/** Serialization operator. */
+	friend FArchive& operator << (FArchive& Ar, FCompressedHeader& Data)
 	{
-		AddReportToPendingFile();
+		Data.DirectoryName.SerializeAsANSICharArray(Ar, 260);
+		Data.FileName.SerializeAsANSICharArray(Ar, 260);
+		Ar << Data.UncompressedSize;
+		Ar << Data.FileCount;
+		return Ar;
 	}
-}
+};
 
-const FText& FCrashUpload::GetStatusText() const
+struct FCompressedData
 {
-	return UploadStateText;
-}
+	TArray<uint8> Data;
+	int32 CompressedSize;
+	int32 UncompressedSize;
+	int32 FileCount;
+};
 
-bool FCrashUpload::IsFinished() const
-{
-	return State >= EUploadState::FirstCompletedState;
-}
+bool FCrashUploadBase::bInitialized = false;
+TArray<FString> FCrashUploadBase::PendingReportDirectories;
+TArray<FString> FCrashUploadBase::FailedReportDirectories;
 
-void FCrashUpload::Cancel()
+FCrashUploadBase::FCrashUploadBase()
+	: bUploadCalled(false)
+	, State(EUploadState::NotSet)
+	, PauseState(EUploadState::Ready)
+	, PendingReportDirectoryIndex(0)
 {
-	SetCurrentState(EUploadState::Cancelled);
+
 }
 
-bool FCrashUpload::SendCheckReportRequest()
+FCrashUploadBase::~FCrashUploadBase()
 {
-	FString XMLString;
-
-	auto Request = CreateHttpRequest();
-	if (State == EUploadState::CheckingReport)
-	{
-		AssignReportIdToPostDataBuffer();
-		Request->SetURL(UrlPrefix / TEXT("CheckReport"));
-		Request->SetHeader(TEXT("Content-Type"), TEXT("text/plain; charset=us-ascii"));
-
-		UE_LOG( CrashReportClientLog, Log, TEXT( "Sending HTTP request: %s" ), *Request->GetURL() );
-	}
-	else
-	{
-		// This part is Windows-specific on the server
-		ErrorReport.LoadWindowsReportXmlFile( XMLString );
-
-		// Convert the XMLString into the UTF-8.
-		FTCHARToUTF8 Converter( (const TCHAR*)*XMLString, XMLString.Len() );
-		const int32 Length = Converter.Length();
-		PostData.Reset( Length );
-		PostData.AddUninitialized( Length );
-		CopyAssignItems( (ANSICHAR*)PostData.GetData(), Converter.Get(), Length );
-
-		Request->SetURL(UrlPrefix / TEXT("CheckReportDetail"));
-		Request->SetHeader(TEXT("Content-Type"), TEXT("text/plain; charset=utf-8"));
-
-		UE_LOG( CrashReportClientLog, Log, TEXT( "Sending HTTP request: %s" ), *Request->GetURL() );
-	}
-
-	UE_LOG( CrashReportClientLog, Log, TEXT( "PostData Num: %i" ), PostData.Num() );
-	Request->SetVerb(TEXT("POST"));
-	Request->SetContent(PostData);
-
-	return Request->ProcessRequest();
+	UE_LOG(CrashReportClientLog, Log, TEXT("Final state (Receiver) = %s"), ToString(State));
 }
 
-enum class ECompressedCrashFileHeader
+void FCrashUploadBase::StaticInitialize(const FPlatformErrorReport& PlatformErrorReport)
 {
-	MAGIC = 0x7E1B83C1,
-};
+	FPendingReports PendingReports;
+	//PendingReports.Add(*PlatformErrorReport.GetReportDirectoryLeafName());
+	PendingReportDirectories = PendingReports.GetReportDirectories();
+	PendingReports.Clear();
+	PendingReports.Save();
+}
 
-struct FCompressedCrashFile : FNoncopyable
+void FCrashUploadBase::StaticShutdown()
 {
-	int32 CurrentFileIndex;
-	FString Filename;
-	TArray<uint8> Filedata;
-
-	FCompressedCrashFile( int32 InCurrentFileIndex, const FString& InFilename, const TArray<uint8>& InFiledata )
-		: CurrentFileIndex(InCurrentFileIndex)
-		, Filename(InFilename)
-		, Filedata(InFiledata)
+	// Write failed uploads back to disk
+	FPendingReports ReportsForNextTime;
+	for (const FString& FailedReport : FailedReportDirectories)
 	{
+		ReportsForNextTime.Add(FailedReport);
 	}
+	ReportsForNextTime.Save();
+}
 
-	/** Serialization operator. */
-	friend FArchive& operator << (FArchive& Ar, FCompressedCrashFile& Data)
-	{
-		Ar << Data.CurrentFileIndex;
-		Data.Filename.SerializeAsANSICharArray( Ar, 260 );
-		Ar << Data.Filedata;
-		return Ar;
-	}
-};
-
-void FCrashUpload::CompressAndSendData()
+bool FCrashUploadBase::CompressData(const TArray<FString>& InPendingFiles, FCompressedData& OutCompressedData, TArray<uint8>& OutPostData, FCompressedHeader* OptionalHeader)
 {
-	UE_LOG(CrashReportClientLog, Log, TEXT("CompressAndSendData have %d pending files"), PendingFiles.Num());
+	UE_LOG(CrashReportClientLog, Log, TEXT("CompressAndSendData have %d pending files"), InPendingFiles.Num());
 
 	// Compress all files into one archive.
-	const int32 BufferSize = 16*1024*1024;
+	const int32 BufferSize = 16 * 1024 * 1024;
 
 	TArray<uint8> UncompressedData;
-	UncompressedData.Reserve( BufferSize );
-	FMemoryWriter MemoryWriter( UncompressedData, false, true );
+	UncompressedData.Reserve(BufferSize);
+	FMemoryWriter MemoryWriter(UncompressedData, false, true);
+
+	if (OptionalHeader != nullptr)
+	{
+		// Write dummy to fill correct size
+		MemoryWriter << *OptionalHeader;
+	}
 
 	int32 CurrentFileIndex = 0;
 
 	const FString FullCrashDumpLocation = FPrimaryCrashProperties::Get()->FullCrashDumpLocation.AsString();
 
 	// Loop to keep trying files until a send succeeds or we run out of files
-	while (PendingFiles.Num() != 0)
+	for (const FString& PathOfFileToUpload : InPendingFiles)
 	{
-		const FString PathOfFileToUpload = PendingFiles.Pop();
-		const FString Filename = FPaths::GetCleanFilename( PathOfFileToUpload );
+		const FString Filename = FPaths::GetCleanFilename(PathOfFileToUpload);
 
-		const bool bValidFullDumpForCopy = Filename == FGenericCrashContext::UE4MinidumpName && 
+		const bool bValidFullDumpForCopy = Filename == FGenericCrashContext::UE4MinidumpName &&
 			(FPrimaryCrashProperties::Get()->CrashDumpMode == ECrashDumpMode::FullDump || FPrimaryCrashProperties::Get()->CrashDumpMode == ECrashDumpMode::FullDumpAlways) &&
 			FPrimaryCrashProperties::Get()->CrashVersion >= ECrashDescVersions::VER_3_CrashContext &&
 			!FullCrashDumpLocation.IsEmpty();
+
 		if (bValidFullDumpForCopy)
 		{
 			const FString DestinationPath = FullCrashDumpLocation / FGenericCrashContext::UE4MinidumpName;
-			const bool bCreated = IFileManager::Get().MakeDirectory( *FullCrashDumpLocation, true );
+			const bool bCreated = IFileManager::Get().MakeDirectory(*FullCrashDumpLocation, true);
 			if (!bCreated)
 			{
-				UE_LOG( CrashReportClientLog, Error, TEXT( "Couldn't create directory for full crash dump %s" ), *DestinationPath );
+				UE_LOG(CrashReportClientLog, Error, TEXT("Couldn't create directory for full crash dump %s"), *DestinationPath);
 			}
 			else
 			{
-				UE_LOG( CrashReportClientLog, Warning, TEXT( "Copying full crash minidump to %s" ), *DestinationPath );
-				IFileManager::Get().Copy( *DestinationPath, *PathOfFileToUpload );
+				UE_LOG(CrashReportClientLog, Warning, TEXT("Copying full crash minidump to %s"), *DestinationPath);
+				IFileManager::Get().Copy(*DestinationPath, *PathOfFileToUpload);
 			}
-						
+
 			continue;
 		}
-		
-		if (IFileManager::Get().FileSize( *PathOfFileToUpload ) > MaxFileSizeToUpload)
+
+		if (IFileManager::Get().FileSize(*PathOfFileToUpload) > CrashUploadDefs::MaxFileSizeToUpload)
 		{
 			UE_LOG(CrashReportClientLog, Warning, TEXT("Skipping large crash report file"));
 			continue;
 		}
 
-		if (!FFileHelper::LoadFileToArray(PostData, *PathOfFileToUpload))
+		if (!FFileHelper::LoadFileToArray(OutPostData, *PathOfFileToUpload))
 		{
 			UE_LOG(CrashReportClientLog, Warning, TEXT("Failed to load crash report file"));
 			continue;
 		}
 
-		const bool bSkipLogFile = !FCrashReportClientConfig::Get().GetSendLogFile() && PathOfFileToUpload.EndsWith( TEXT( ".log" ) );
+		const bool bSkipLogFile = !FCrashReportClientConfig::Get().GetSendLogFile() && PathOfFileToUpload.EndsWith(TEXT(".log"));
 		if (bSkipLogFile)
 		{
-			UE_LOG( CrashReportClientLog, Warning, TEXT( "Skipping the %s" ), *Filename );
+			UE_LOG(CrashReportClientLog, Warning, TEXT("Skipping the %s"), *Filename);
 			continue;
 		}
 
@@ -215,15 +191,15 @@ void FCrashUpload::CompressAndSendData()
 		const bool bSkipXMLFile = PathOfFileToUpload.EndsWith( TEXT( ".xml" ) );
 		if (bSkipXMLFile)
 		{
-			UE_LOG( CrashReportClientLog, Warning, TEXT( "Skipping the %s" ), *Filename );
-			continue;
+		UE_LOG( CrashReportClientLog, Warning, TEXT( "Skipping the %s" ), *Filename );
+		continue;
 		}*/
 
 		// Skip old Report.wer file.
-		const bool bSkipWERFile = PathOfFileToUpload.Contains( TEXT( "Report.wer" ) );
+		const bool bSkipWERFile = PathOfFileToUpload.Contains(TEXT("Report.wer"));
 		if (bSkipWERFile)
 		{
-			UE_LOG( CrashReportClientLog, Warning, TEXT( "Skipping the %s" ), *Filename );
+			UE_LOG(CrashReportClientLog, Warning, TEXT("Skipping the %s"), *Filename);
 			continue;
 		}
 
@@ -233,51 +209,252 @@ void FCrashUpload::CompressAndSendData()
 		const bool bSkipDiagnostics = Filename == FCrashReportClientConfig::Get().GetDiagnosticsFilename();
 		if (bSkipDiagnostics)
 		{
-			UE_LOG( CrashReportClientLog, Warning, TEXT( "Skipping the %s" ), *Filename );
-			continue;
+		UE_LOG( CrashReportClientLog, Warning, TEXT( "Skipping the %s" ), *Filename );
+		continue;
 		}*/
 
-		UE_LOG(CrashReportClientLog, Log, TEXT("CompressAndSendData compressing %d bytes ('%s')"), PostData.Num(), *PathOfFileToUpload);
+		UE_LOG(CrashReportClientLog, Log, TEXT("CompressAndSendData compressing %d bytes ('%s')"), OutPostData.Num(), *PathOfFileToUpload);
 
-		FCompressedCrashFile FileToCompress( CurrentFileIndex, Filename, PostData );
+		FCompressedCrashFile FileToCompress(CurrentFileIndex, Filename, OutPostData);
 		CurrentFileIndex++;
 
 		MemoryWriter << FileToCompress;
 	}
 
-	
+	if (OptionalHeader != nullptr)
+	{
+		FMemoryWriter MemoryHeaderWriter(UncompressedData, false, true);
+
+		OptionalHeader->UncompressedSize = UncompressedData.Num();
+		OptionalHeader->FileCount = CurrentFileIndex;
+
+		MemoryHeaderWriter << *OptionalHeader;
+	}
+
 	uint8* CompressedDataRaw = new uint8[BufferSize];
 
-	int32 CompressedSize = BufferSize;
-	int32 UncompressedSize = UncompressedData.Num();
-	const bool bResult = FCompression::CompressMemory( COMPRESS_ZLIB, CompressedDataRaw, CompressedSize, UncompressedData.GetData(), UncompressedSize );
-	if( !bResult )
+	OutCompressedData.FileCount = CurrentFileIndex;
+	OutCompressedData.CompressedSize = BufferSize;
+	OutCompressedData.UncompressedSize = UncompressedData.Num();
+	const bool bResult = FCompression::CompressMemory(COMPRESS_ZLIB, CompressedDataRaw, OutCompressedData.CompressedSize, UncompressedData.GetData(), OutCompressedData.UncompressedSize);
+	if (bResult)
+	{
+		// Copy compressed data into the array.
+		OutCompressedData.Data.Append(CompressedDataRaw, OutCompressedData.CompressedSize);
+		delete[] CompressedDataRaw;
+		CompressedDataRaw = nullptr;
+	}
+
+	return bResult;
+}
+
+const TCHAR* FCrashUploadBase::ToString(EUploadState::Type State)
+{
+	switch (State)
+	{
+	case EUploadState::PingingServer:
+		return TEXT("PingingServer");
+
+	case EUploadState::Ready:
+		return TEXT("Ready");
+
+	case EUploadState::CheckingReport:
+		return TEXT("CheckingReport");
+
+	case EUploadState::CheckingReportDetail:
+		return TEXT("CheckingReportDetail");
+
+	case EUploadState::CompressAndSendData:
+		return TEXT("SendingFiles");
+
+	case EUploadState::WaitingToPostReportComplete:
+		return TEXT("WaitingToPostReportComplete");
+
+	case EUploadState::PostingReportComplete:
+		return TEXT("PostingReportComplete");
+
+	case EUploadState::Finished:
+		return TEXT("Finished");
+
+	case EUploadState::ServerNotAvailable:
+		return TEXT("ServerNotAvailable");
+
+	case EUploadState::UploadError:
+		return TEXT("UploadError");
+
+	case EUploadState::Cancelled:
+		return TEXT("Cancelled");
+
+	default:
+		break;
+	}
+
+	return TEXT("Unknown UploadState value");
+}
+
+void FCrashUploadBase::SetCurrentState(EUploadState::Type InState)
+{
+	if (State == EUploadState::NotSet)
+	{
+		UE_LOG(CrashReportClientLog, Log, TEXT("Initial state = %s"), ToString(State));
+	}
+	else
+	{
+		UE_LOG(CrashReportClientLog, Log, TEXT("State change from %s to %s"), ToString(State), ToString(InState));
+	}
+
+	State = InState;
+
+	switch (State)
+	{
+	default:
+		break;
+
+	case EUploadState::PingingServer:
+		UploadStateText = LOCTEXT("PingingServer", "Pinging server");
+		break;
+
+	case EUploadState::Ready:
+		UploadStateText = LOCTEXT("UploaderReady", "Ready to send to server");
+		break;
+
+	case EUploadState::ServerNotAvailable:
+		UploadStateText = LOCTEXT("ServerNotAvailable", "Server not available - report will be stored for later upload");
+		break;
+	}
+}
+
+void FCrashUploadBase::AddReportToFailedList() const
+{
+	if (PendingFiles.Num() == 0)
+	{
+		// No files so don't bother
+		return;
+	}
+
+	FailedReportDirectories.AddUnique(PendingReportDirectories[PendingReportDirectoryIndex]);
+}
+
+// FCrashUploadToReceiver //////////////////////////////////////////////////////
+
+FCrashUploadToReceiver::FCrashUploadToReceiver(const FString& InReceiverAddress)
+	: UrlPrefix(InReceiverAddress.IsEmpty() ? TEXT("") : InReceiverAddress / TEXT("CrashReporter"))
+{
+	if (!UrlPrefix.IsEmpty())
+	{
+		// Sending to receiver
+		SendPingRequest();
+	}
+	else
+	{
+		SetCurrentState(EUploadState::UploadError);
+	}
+}
+
+FCrashUploadToReceiver::~FCrashUploadToReceiver()
+{
+}
+
+bool FCrashUploadToReceiver::PingTimeout(float DeltaTime)
+{
+	if (EUploadState::PingingServer == State)
+	{
+		SetCurrentState(EUploadState::ServerNotAvailable);
+
+		// PauseState will be Ready if user has not yet decided to send the report
+		if (PauseState > EUploadState::Ready)
+		{
+			AddReportToFailedList();
+		}
+	}
+
+	// One-shot
+	return false;
+}
+
+void FCrashUploadToReceiver::BeginUpload(const FPlatformErrorReport& PlatformErrorReport)
+{
+	bUploadCalled = true;
+
+	ErrorReport = PlatformErrorReport;
+	PendingFiles = FPlatformErrorReport( ErrorReport.GetReportDirectory() ).GetFilesToUpload();
+	UE_LOG(CrashReportClientLog, Log, TEXT("Got %d pending files to upload from '%s'"), PendingFiles.Num(), *ErrorReport.GetReportDirectoryLeafName());
+
+	PauseState = EUploadState::Finished;
+	if (State == EUploadState::Ready)
+	{
+		BeginUploadImpl();
+	}
+	else if (State == EUploadState::ServerNotAvailable)
+	{
+		AddReportToFailedList();
+	}
+}
+
+bool FCrashUploadToReceiver::SendCheckReportRequest()
+{
+	FString XMLString;
+
+	auto Request = CreateHttpRequest();
+	if (State == EUploadState::CheckingReport)
+	{
+		AssignReportIdToPostDataBuffer();
+		Request->SetURL(UrlPrefix / TEXT("CheckReport"));
+		Request->SetHeader(TEXT("Content-Type"), TEXT("text/plain; charset=us-ascii"));
+
+		UE_LOG( CrashReportClientLog, Log, TEXT( "Sending HTTP request: %s" ), *Request->GetURL() );
+	}
+	else
+	{
+		// This part is Windows-specific on the server
+		ErrorReport.LoadWindowsReportXmlFile( XMLString );
+
+		// Convert the XMLString into the UTF-8.
+		FTCHARToUTF8 Converter( (const TCHAR*)*XMLString, XMLString.Len() );
+		const int32 Length = Converter.Length();
+		PostData.Reset( Length );
+		PostData.AddUninitialized( Length );
+		CopyAssignItems( (ANSICHAR*)PostData.GetData(), Converter.Get(), Length );
+
+		Request->SetURL(UrlPrefix / TEXT("CheckReportDetail"));
+		Request->SetHeader(TEXT("Content-Type"), TEXT("text/plain; charset=utf-8"));
+
+		UE_LOG( CrashReportClientLog, Log, TEXT( "Sending HTTP request: %s" ), *Request->GetURL() );
+	}
+
+	UE_LOG( CrashReportClientLog, Log, TEXT( "PostData Num: %i" ), PostData.Num() );
+	Request->SetVerb(TEXT("POST"));
+	Request->SetContent(PostData);
+
+	return Request->ProcessRequest();
+}
+
+void FCrashUploadToReceiver::CompressAndSendData()
+{
+	FCompressedData CompressedData;
+	if (!CompressData(PendingFiles, CompressedData, PostData))
 	{
 		UE_LOG(CrashReportClientLog, Warning, TEXT("Couldn't compress the crash report files"));
 		SetCurrentState(EUploadState::Cancelled);
 		return;
 	}
-	
-	const FString Filename = ErrorReport.GetReportDirectoryLeafName() + TEXT(".ue4crash");
 
-	// Copy compressed data into the array.
-	TArray<uint8> CompressedData;
-	CompressedData.Append( CompressedDataRaw, CompressedSize );
-	delete [] CompressedDataRaw;
-	CompressedDataRaw = nullptr;
+	PendingFiles.Empty();
+
+	const FString Filename = ErrorReport.GetReportDirectoryLeafName() + TEXT(".ue4crash");
 
 	// Set up request for upload
 	auto Request = CreateHttpRequest();
 	Request->SetVerb(TEXT("POST"));
 	Request->SetHeader(TEXT("Content-Type"), TEXT("application/octet-stream"));
 	Request->SetURL(UrlPrefix / TEXT("UploadReportFile"));
-	Request->SetContent(CompressedData);
+	Request->SetContent(CompressedData.Data);
 	Request->SetHeader(TEXT("DirectoryName"), *ErrorReport.GetReportDirectoryLeafName());
 	Request->SetHeader(TEXT("FileName"), Filename);
-	Request->SetHeader(TEXT("FileLength"), TTypeToString<int32>::ToString(CompressedData.Num()) );
-	Request->SetHeader(TEXT("CompressedSize"), TTypeToString<int32>::ToString(CompressedSize) );
-	Request->SetHeader(TEXT("UncompressedSize"), TTypeToString<int32>::ToString(UncompressedSize) );
-	Request->SetHeader(TEXT("NumberOfFiles"), TTypeToString<int32>::ToString(CurrentFileIndex) );
+	Request->SetHeader(TEXT("FileLength"), TTypeToString<int32>::ToString(CompressedData.Data.Num()) );
+	Request->SetHeader(TEXT("CompressedSize"), TTypeToString<int32>::ToString(CompressedData.CompressedSize) );
+	Request->SetHeader(TEXT("UncompressedSize"), TTypeToString<int32>::ToString(CompressedData.UncompressedSize) );
+	Request->SetHeader(TEXT("NumberOfFiles"), TTypeToString<int32>::ToString(CompressedData.FileCount) );
 	UE_LOG( CrashReportClientLog, Log, TEXT( "Sending HTTP request: %s" ), *Request->GetURL() );
 
 	if (Request->ProcessRequest())
@@ -291,7 +468,7 @@ void FCrashUpload::CompressAndSendData()
 	}	
 }
 
-void FCrashUpload::AssignReportIdToPostDataBuffer()
+void FCrashUploadToReceiver::AssignReportIdToPostDataBuffer()
 {
 	FString ReportDirectoryName = *ErrorReport.GetReportDirectoryLeafName();
 	const int32 DirectoryNameLength = ReportDirectoryName.Len();
@@ -302,7 +479,7 @@ void FCrashUpload::AssignReportIdToPostDataBuffer()
 	}
 }
 
-void FCrashUpload::PostReportComplete()
+void FCrashUploadToReceiver::PostReportComplete()
 {
 	if (PauseState == EUploadState::PostingReportComplete)
 	{
@@ -331,7 +508,7 @@ void FCrashUpload::PostReportComplete()
 	}
 }
 
-void FCrashUpload::OnProcessRequestComplete(FHttpRequestPtr HttpRequest, FHttpResponsePtr HttpResponse, bool bSucceeded)
+void FCrashUploadToReceiver::OnProcessRequestComplete(FHttpRequestPtr HttpRequest, FHttpResponsePtr HttpResponse, bool bSucceeded)
 {
 	UE_LOG(CrashReportClientLog, Log, TEXT("OnProcessRequestComplete(), State=%s bSucceeded=%i"), ToString(State), (int32)bSucceeded );
 	switch (State)
@@ -353,42 +530,56 @@ void FCrashUpload::OnProcessRequestComplete(FHttpRequestPtr HttpRequest, FHttpRe
 
 	case EUploadState::CheckingReport:
 	case EUploadState::CheckingReportDetail:
-		if (!bSucceeded || !ParseServerResponse(HttpResponse))
 		{
-			if (!bSucceeded)
+			bool bCheckedOkay = false;
+
+			if (!bSucceeded || !ParseServerResponse(HttpResponse, bCheckedOkay))
 			{
-				UE_LOG(CrashReportClientLog, Warning, TEXT("Request to server failed"));
+				if (!bSucceeded)
+				{
+					UE_LOG(CrashReportClientLog, Warning, TEXT("Request to server failed"));
+				}
+				else
+				{
+					UE_LOG(CrashReportClientLog, Warning, TEXT("Did not get a valid server response."));
+				}
+
+				// Failed to check with the server - skip this report for now
+				AddReportToFailedList();
+				CheckPendingReportsForFilesToUpload();
 			}
-			else
+			else if (!bCheckedOkay)
 			{
+				// Server rejected the report
 				UE_LOG(CrashReportClientLog, Warning, TEXT("Did not get a valid server response."));
+				CheckPendingReportsForFilesToUpload();
 			}
-
-			// Skip this report
-			CheckPendingReportsForFilesToUpload();
-			break;
-		}
-
-		if (State == EUploadState::CheckingReport)
-		{
-			SetCurrentState(EUploadState::CheckingReportDetail);
-			if (!SendCheckReportRequest())
+			else
 			{
-				// Skip this report
-				CheckPendingReportsForFilesToUpload();
+				if (State == EUploadState::CheckingReport)
+				{
+					SetCurrentState(EUploadState::CheckingReportDetail);
+					if (!SendCheckReportRequest())
+					{
+						// Skip this report for now
+						AddReportToFailedList();
+						CheckPendingReportsForFilesToUpload();
+					}
+				}
+				else
+				{
+					SetCurrentState(EUploadState::CompressAndSendData);
+					CompressAndSendData();
+				}
 			}
 		}
-		else
-		{
-			SetCurrentState(EUploadState::CompressAndSendData);
-			CompressAndSendData();
-		}
 		break;
 
 	case EUploadState::CompressAndSendData:
 		if (!bSucceeded)
 		{
-			UE_LOG(CrashReportClientLog, Warning, TEXT("File upload failed"));
+			UE_LOG(CrashReportClientLog, Warning, TEXT("File upload failed to receiver"));
+			AddReportToFailedList();
 			SetCurrentState(EUploadState::Cancelled);
 		}
 		else
@@ -403,7 +594,7 @@ void FCrashUpload::OnProcessRequestComplete(FHttpRequestPtr HttpRequest, FHttpRe
 	}
 }
 
-void FCrashUpload::OnPingSuccess()
+void FCrashUploadToReceiver::OnPingSuccess()
 {
 	if (PauseState > EUploadState::Ready)
 	{
@@ -416,67 +607,43 @@ void FCrashUpload::OnPingSuccess()
 	}
 }
 
-void FCrashUpload::AddReportToPendingFile() const
-{
-	if (PendingFiles.Num() == 0)
-	{
-		// No files so don't bother
-		return;
-	}
-
-	FPendingReports PendingReports;
-	// Extract full path of report from the first file
-	PendingReports.Add(FPaths::GetPath(*PendingFiles[0]));
-	PendingReports.Save();
-}
-
-void FCrashUpload::CheckPendingReportsForFilesToUpload()
+void FCrashUploadToReceiver::CheckPendingReportsForFilesToUpload()
 {
 	SetCurrentState(EUploadState::CheckingReport);
-	do
+
+	for (; PendingReportDirectoryIndex < PendingReportDirectories.Num(); PendingReportDirectoryIndex++)
 	{
-		if (PendingReportDirectories.Num() == 0)
+		ErrorReport = FPlatformErrorReport(PendingReportDirectories[PendingReportDirectoryIndex]);
+		PendingFiles = ErrorReport.GetFilesToUpload();
+
+		if (PendingFiles.Num() > 0 && SendCheckReportRequest())
 		{
-			// Nothing to upload
-			UE_LOG(CrashReportClientLog, Log, TEXT("All uploads done"));
-			SetCurrentState(EUploadState::Finished);
 			return;
 		}
-
-		ErrorReport = FPlatformErrorReport(PendingReportDirectories.Pop());
-		PendingFiles = ErrorReport.GetFilesToUpload();
 	}
-	while (PendingFiles.Num() == 0 || !SendCheckReportRequest());
+
+	// Nothing left to upload
+	UE_LOG(CrashReportClientLog, Log, TEXT("All uploads done"));
+	SetCurrentState(EUploadState::Finished);
 }
 
-void FCrashUpload::BeginUploadImpl()
+void FCrashUploadToReceiver::BeginUploadImpl()
 {
-	FPendingReports PendingReports;
-	// Make sure we don't try to upload the current report twice
-	PendingReports.Forget(*ErrorReport.GetReportDirectoryLeafName());
-	PendingReportDirectories = PendingReports.GetReportDirectories();
-
-	UE_LOG(CrashReportClientLog, Log, TEXT("Proceeding to upload %d pending files from '%s'"), PendingFiles.Num(), *ErrorReport.GetReportDirectoryLeafName());
-
-	// Not sure about this: probably want to removes files one by one as they're successfully uploaded
-	PendingReports.Clear();
-	PendingReports.Save();
-
 	SetCurrentState(EUploadState::CheckingReport);
-	if (PendingFiles.Num() == 0 || !SendCheckReportRequest())
+	if (!SendCheckReportRequest())
 	{
 		CheckPendingReportsForFilesToUpload();
 	}
 }
 
-TSharedRef<IHttpRequest> FCrashUpload::CreateHttpRequest()
+TSharedRef<IHttpRequest> FCrashUploadToReceiver::CreateHttpRequest()
 {
 	auto Request = FHttpModule::Get().CreateRequest();
-	Request->OnProcessRequestComplete().BindRaw(this, &FCrashUpload::OnProcessRequestComplete);
+	Request->OnProcessRequestComplete().BindRaw(this, &FCrashUploadToReceiver::OnProcessRequestComplete);
 	return Request;
 }
 
-void FCrashUpload::SendPingRequest()
+void FCrashUploadToReceiver::SendPingRequest()
 {
 	SetCurrentState(EUploadState::PingingServer);
 
@@ -487,7 +654,7 @@ void FCrashUpload::SendPingRequest()
 
 	if (Request->ProcessRequest())
 	{
-		FTicker::GetCoreTicker().AddTicker(FTickerDelegate::CreateRaw(this, &FCrashUpload::PingTimeout), PingTimeoutSeconds);
+		FTicker::GetCoreTicker().AddTicker(FTickerDelegate::CreateRaw(this, &FCrashUploadToReceiver::PingTimeout), CrashUploadDefs::PingTimeoutSeconds);
 	}
 	else
 	{
@@ -495,99 +662,158 @@ void FCrashUpload::SendPingRequest()
 	}
 }
 
-bool FCrashUpload::ParseServerResponse(FHttpResponsePtr Response)
+bool FCrashUploadToReceiver::ParseServerResponse(FHttpResponsePtr Response, bool& OutValidReport)
 {
 	// Turn the snippet into a complete XML document, to keep the XML parser happy
 	FXmlFile ParsedResponse(FString(TEXT("<Root>")) + Response->GetContentAsString() + TEXT("</Root>"), EConstructMethod::ConstructFromBuffer);
 	UE_LOG(CrashReportClientLog, Log, TEXT("Response->GetContentAsString(): '%s'"), *Response->GetContentAsString());	
 	if (!ParsedResponse.IsValid())
 	{
-		UE_LOG(CrashReportClientLog, Log, TEXT("Invalid response!"));	
+		UE_LOG(CrashReportClientLog, Log, TEXT("Invalid response!"));
+		OutValidReport = false;
 		return false;
 	}
 	if (auto ResultNode = ParsedResponse.GetRootNode()->FindChildNode(TEXT("CrashReporterResult")))
 	{
 		UE_LOG(CrashReportClientLog, Log, TEXT("ResultNode->GetAttribute(TEXT(\"bSuccess\")) = %s"), *ResultNode->GetAttribute(TEXT("bSuccess")));	
-		return ResultNode->GetAttribute(TEXT("bSuccess")) == TEXT("true");
+		OutValidReport = ResultNode->GetAttribute(TEXT("bSuccess")) == TEXT("true");
+		return true;
 	}
-	UE_LOG(CrashReportClientLog, Log, TEXT("Could not find CrashReporterResult"));	
+	UE_LOG(CrashReportClientLog, Log, TEXT("Could not find CrashReporterResult"));
+	OutValidReport = false;
 	return false;
 }
 
-void FCrashUpload::SetCurrentState(EUploadState::Type InState)
+
+// FCrashUploadToDataRouter //////////////////////////////////////////////////////
+
+FCrashUploadToDataRouter::FCrashUploadToDataRouter(const FString& InDataRouterUrl)
+	: DataRouterUrl(InDataRouterUrl)
 {
-	if (State == EUploadState::NotSet)
+	if (!DataRouterUrl.IsEmpty())
 	{
-		UE_LOG(CrashReportClientLog, Log, TEXT("Initial state = %s"), ToString(State));
+		SetCurrentState(EUploadState::Ready);
 	}
 	else
 	{
-		UE_LOG(CrashReportClientLog, Log, TEXT("State change from %s to %s"), ToString(State), ToString(InState));
+		SetCurrentState(EUploadState::UploadError);
 	}
+}
 
-	State = InState;
-
-	switch (State)
-	{
-	default:
-		break;
+FCrashUploadToDataRouter::~FCrashUploadToDataRouter()
+{
+}
 
-	case EUploadState::PingingServer:
-		UploadStateText = LOCTEXT("PingingServer", "Pinging server");
-		break;
+void FCrashUploadToDataRouter::BeginUpload(const FPlatformErrorReport& PlatformErrorReport)
+{
+	bUploadCalled = true;
 
-	case EUploadState::Ready:
-		UploadStateText = LOCTEXT("UploaderReady", "Ready to send to server");
-		break;
+	ErrorReport = PlatformErrorReport;
+	PendingFiles = FPlatformErrorReport(ErrorReport.GetReportDirectory()).GetFilesToUpload();
+	UE_LOG(CrashReportClientLog, Log, TEXT("Got %d pending files to upload from '%s'"), PendingFiles.Num(), *ErrorReport.GetReportDirectoryLeafName());
 
-	case EUploadState::ServerNotAvailable:
-		UploadStateText = LOCTEXT("ServerNotAvailable", "Server not available - report will be stored for later upload");
-		break;
+	PauseState = EUploadState::Finished;
+	if (State == EUploadState::Ready)
+	{
+		SetCurrentState(EUploadState::CompressAndSendData);
+		CompressAndSendData();
 	}
 }
 
-const TCHAR* FCrashUpload::ToString(EUploadState::Type State)
+void FCrashUploadToDataRouter::CompressAndSendData()
 {
-	switch(State)
+	FCompressedHeader CompressedHeader;
+	CompressedHeader.DirectoryName = ErrorReport.GetReportDirectoryLeafName();
+	CompressedHeader.FileName = ErrorReport.GetReportDirectoryLeafName() + TEXT(".ue4crash");
+
+	FCompressedData CompressedData;
+	if (!CompressData(PendingFiles, CompressedData, PostData, &CompressedHeader))
 	{
-		case EUploadState::PingingServer: 
-			return TEXT("PingingServer");
+		UE_LOG(CrashReportClientLog, Warning, TEXT("Couldn't compress the crash report files"));
+		SetCurrentState(EUploadState::Cancelled);
+		return;
+	}
 
-		case EUploadState::Ready:
-			return TEXT("Ready");
+	PendingFiles.Empty();
 
-		case EUploadState::CheckingReport:
-			return TEXT("CheckingReport");
+	FString UserId = FString::Printf(TEXT("%s|%s|%s"), *FPlatformMisc::GetMachineId().ToString(EGuidFormats::Digits).ToLower(), *FPlatformMisc::GetEpicAccountId(), *FPlatformMisc::GetOperatingSystemId());
 
-		case EUploadState::CheckingReportDetail:
-			return TEXT("CheckingReportDetail");
+	FString UrlParams = FString::Printf(TEXT("?AppID=%s&AppVersion=%s&AppEnvironment=%s&UploadType=%s&UserID=%s"),
+		*FGenericPlatformHttp::UrlEncode(CrashUploadDefs::APIKey),
+		*FGenericPlatformHttp::UrlEncode(FEngineVersion::Current().ToString()),
+		*FGenericPlatformHttp::UrlEncode(FEngineBuildSettings::IsInternalBuild() ? CrashUploadDefs::AppEnvironmentInternal : CrashUploadDefs::AppEnvironmentExternal),
+		*FGenericPlatformHttp::UrlEncode(CrashUploadDefs::UploadType),
+		*FGenericPlatformHttp::UrlEncode(UserId));
 
-		case EUploadState::CompressAndSendData:
-			return TEXT("SendingFiles");
+	// Set up request for upload
+	auto Request = CreateHttpRequest();
+	Request->SetVerb(TEXT("POST"));
+	Request->SetHeader(TEXT("Content-Type"), TEXT("application/octet-stream"));
+	Request->SetURL(DataRouterUrl + UrlParams);
+	Request->SetContent(CompressedData.Data);
+	UE_LOG(CrashReportClientLog, Log, TEXT("Sending HTTP request: %s"), *Request->GetURL());
 
-		case EUploadState::WaitingToPostReportComplete:
-			return TEXT("WaitingToPostReportComplete");
+	if (Request->ProcessRequest())
+	{
+		return;
+	}
+	else
+	{
+		UE_LOG(CrashReportClientLog, Warning, TEXT("Failed to send file upload request"));
+		SetCurrentState(EUploadState::Cancelled);
+	}
+}
 
-		case EUploadState::PostingReportComplete:
-			return TEXT("PostingReportComplete");
+TSharedRef<IHttpRequest> FCrashUploadToDataRouter::CreateHttpRequest()
+{
+	auto Request = FHttpModule::Get().CreateRequest();
+	Request->OnProcessRequestComplete().BindRaw(this, &FCrashUploadToDataRouter::OnProcessRequestComplete);
+	return Request;
+}
 
-		case EUploadState::Finished:
-			return TEXT("Finished");
+void FCrashUploadToDataRouter::OnProcessRequestComplete(FHttpRequestPtr HttpRequest, FHttpResponsePtr HttpResponse, bool bSucceeded)
+{
+	UE_LOG(CrashReportClientLog, Log, TEXT("OnProcessRequestComplete(), State=%s bSucceeded=%i"), ToString(State), (int32)bSucceeded);
+	switch (State)
+	{
+	default:
+		// May get here if response is received after time-out has passed
+		break;
 
-		case EUploadState::ServerNotAvailable:
-			return TEXT("ServerNotAvailable");
+	case EUploadState::CompressAndSendData:
+		if (!bSucceeded)
+		{
+			UE_LOG(CrashReportClientLog, Warning, TEXT("File upload failed to data router"));
+			AddReportToFailedList();
+			SetCurrentState(EUploadState::Cancelled);
+		}
+		else
+		{
+			CheckPendingReportsForFilesToUpload();
+		}
+		break;
+	}
+}
 
-		case EUploadState::UploadError:
-			return TEXT("UploadError");
+void FCrashUploadToDataRouter::CheckPendingReportsForFilesToUpload()
+{
+	SetCurrentState(EUploadState::CompressAndSendData);
 
-		case EUploadState::Cancelled:
-			return TEXT("Cancelled");
+	for (; PendingReportDirectoryIndex < PendingReportDirectories.Num(); PendingReportDirectoryIndex++)
+	{
+		ErrorReport = FPlatformErrorReport(PendingReportDirectories[PendingReportDirectoryIndex]);
+		PendingFiles = ErrorReport.GetFilesToUpload();
 
-		default:
-			break;
+		if (PendingFiles.Num() > 0)
+		{
+			CompressAndSendData();
+			return;
+		}
 	}
 
-	return TEXT("Unknown UploadState value");
+	// Nothing left to upload
+	UE_LOG(CrashReportClientLog, Log, TEXT("All uploads done"));
+	SetCurrentState(EUploadState::Finished);
 }
 
 #undef LOCTEXT_NAMESPACE
diff --git a/Engine/Source/Programs/CrashReporter/CrashReportClient/Private/CrashUpload.h b/Engine/Source/Programs/CrashReporter/CrashReportClient/Private/CrashUpload.h
--- a/Engine/Source/Programs/CrashReporter/CrashReportClient/Private/CrashUpload.h
+++ b/Engine/Source/Programs/CrashReporter/CrashReportClient/Private/CrashUpload.h
@@ -5,48 +5,44 @@
 #include "Http.h"
 #include "PlatformErrorReport.h"
 
-/**
- * Handles uploading files to the crash report server
- */
-class FCrashUpload
+class FCrashUploadBase
 {
 public:
-	/**
-	 * Constructor: pings server 
-	 * @param ServerAddress Host IP of the crash report server
-	 */
-	explicit FCrashUpload(const FString& ServerAddress);
+	FCrashUploadBase();
+
+	virtual ~FCrashUploadBase();
 
 	/**
-	 * Destructor for logging
-	 */
-	~FCrashUpload();
+	* Has BeginUpload been called?
+	*/
+	bool IsUploadCalled() const { return bUploadCalled; }
 
 	/**
-	 * Commence upload when ready
-	 * @param PlatformErrorReport Error report to upload files from
-	 */
-	void BeginUpload(const FPlatformErrorReport& PlatformErrorReport);
+	* Provide progress or error information for the UI
+	*/
+	const FText& GetStatusText() const { return UploadStateText; }
 
 	/**
-	 * Provide progress or error information for the UI
-	 */
-	const FText& GetStatusText() const;
-
-
-	/**
 	 * Determine whether the upload has finished (successfully or otherwise)
 	 * @return Whether the upload has finished
 	 */
-	bool IsFinished() const;
+	bool IsFinished() const
+	{
+		return State >= EUploadState::FirstCompletedState;
+	}
 
-	/**
-	 * Cancel the upload
-	 * @note Only valid to call this before the upload has begun
-	 */
-	void Cancel();
+	void Cancel()
+	{
+		SetCurrentState(EUploadState::Cancelled);
+	}
 
-private:
+	static bool IsInitialized() { return bInitialized; }
+
+	static void StaticInitialize(const FPlatformErrorReport& PlatformErrorReport);
+
+	static void StaticShutdown();
+
+protected:
 	/** State enum to keep track of what the uploader is doing */
 	struct EUploadState
 	{
@@ -71,19 +67,84 @@
 		};
 	};
 
+	static bool CompressData(const TArray<FString>& InPendingFiles, struct FCompressedData& OutCompressedData, TArray<uint8>& OutPostData, struct FCompressedHeader* OptionalHeader = nullptr);
+
 	/**
-	 * Get a string representation of the state, for logging purposes
-	 * @param State Value to stringize
-	 * @return Literal string value
-	 */
+	* Get a string representation of the state, for logging purposes
+	* @param State Value to stringize
+	* @return Literal string value
+	*/
 	static const TCHAR* ToString(EUploadState::Type InState);
 
 	/**
-	 * Set the current state, also updating the status text where necessary
-	 * @param State State the uploader is now in
-	 */
+	* Set the current state, also updating the status text where necessary
+	* @param State State the uploader is now in
+	*/
 	void SetCurrentState(EUploadState::Type InState);
 
 	/**
+	* When failed, add the report to a file list containing reports to upload next time
+	*/
+	void AddReportToFailedList() const;
+
+protected:
+	bool bUploadCalled;
+
+	/** What this class is currently doing */
+	EUploadState::Type State;
+
+	/** Status of upload to display */
+	FText UploadStateText;
+
+	/** State to pause at until confirmation has been received to continue */
+	EUploadState::Type PauseState;
+
+	/** Full paths of files still to be uploaded */
+	TArray<FString> PendingFiles;
+
+	/** Error report being processed */
+	FPlatformErrorReport ErrorReport;
+
+	/** Buffer to keep reusing for file content and other messages */
+	TArray<uint8> PostData;
+
+	int32 PendingReportDirectoryIndex;
+
+protected:
+	static bool bInitialized;
+
+	/** Full paths of reports from previous runs still to be uploaded */
+	static TArray<FString> PendingReportDirectories;
+
+	/** Full paths of reports from this run that did not upload */
+	static TArray<FString> FailedReportDirectories;
+};
+
+
+/**
+ * Handles uploading files to the crash report server
+ */
+class FCrashUploadToReceiver : public FCrashUploadBase
+{
+public:
+	/**
+	 * Constructor: pings server 
+	 * @param ServerAddress Host IP of the crash report server
+	 */
+	explicit FCrashUploadToReceiver(const FString& InReceiverAddress);
+
+	/**
+	 * Destructor for logging
+	 */
+	virtual ~FCrashUploadToReceiver();
+
+	/**
+	 * Commence upload when ready
+	 * @param PlatformErrorReport Error report to upload files from
+	 */
+	void BeginUpload(const FPlatformErrorReport& PlatformErrorReport);
+
+private:
+	/**
 	 * Send a request to see if the server will accept this report
 	 * @return Whether request was successfully sent
@@ -127,10 +188,5 @@
 	bool PingTimeout(float DeltaTime);
 
 	/**
-	 * In case the server is not available, add the report to a file containing reports to upload next time
-	 */
-	void AddReportToPendingFile() const;
-
-	/**
 	 * If there a no pending files, look through pending reports for files to upload
 	 */
@@ -153,32 +209,60 @@
 
 	/**
 	 * Parse an XML response from the server for the success field
-	 * @param Response Response to get message to parse from
-	 * @return Whether it's both a valid message and the code is success
+	 * @param Response			Response to get message to parse from
+	 * @param OutValidReport	Answer from the server on whether to continue with this report upload
+	 * @return					Whether a valid response was received from the server
 	 */
-	static bool ParseServerResponse(FHttpResponsePtr Response);
+	static bool ParseServerResponse(FHttpResponsePtr Response, bool& OutValidReport);
 
 	/** Host, port and common prefix of all requests to the server */
 	FString UrlPrefix;
+};
 
-	/** What this class is currently doing */
-	EUploadState::Type State;
+/**
+* Handles uploading files to the data router
+*/
+class FCrashUploadToDataRouter : public FCrashUploadBase
+{
+public:
+	/**
+	 * Constructor: pings server
+	 * @param ServerAddress Host IP of the crash report server
+	 */
+	explicit FCrashUploadToDataRouter(const FString& InDataRouterUrl);
 
-	/** Status of upload to display */
-	FText UploadStateText;
+	/**
+	 * Destructor for logging
+	 */
+	virtual ~FCrashUploadToDataRouter();
 
-	/** State to pause at until confirmation has been received to continue */
-	EUploadState::Type PauseState;
+	void BeginUpload(const FPlatformErrorReport& PlatformErrorReport);
 
-	/** Full paths of files still to be uploaded */
-	TArray<FString> PendingFiles;
+	/**
+	 * Compresses all crash report files and sends one compressed file.
+	 */
+	void CompressAndSendData();
 
-	/** Error report being processed */
-	FPlatformErrorReport ErrorReport;
+	/**
+	 * Create a request object and bind this class's response handler to it
+	 */
+	TSharedRef<IHttpRequest> CreateHttpRequest();
 
-	/** Full paths of reports from previous runs still to be uploaded */
-	TArray<FString> PendingReportDirectories;
+	/**
+	 * Callback from HTTP library when a request has completed
+	 * @param HttpRequest The request object
+	 * @param HttpResponse The response from the server
+	 * @param bSucceeded Whether a response was successfully received
+	 */
+	void OnProcessRequestComplete(FHttpRequestPtr HttpRequest, FHttpResponsePtr HttpResponse, bool bSucceeded);
 
-	/** Buffer to keep reusing for file content and other messages */
-	TArray<uint8> PostData;
+private:
+	/**
+	 * If there a no pending files, look through pending reports for files to upload
+	 */
+	void CheckPendingReportsForFilesToUpload();
+
+private:
+	/** Url for data router requests */
+	FString DataRouterUrl;
 };
diff --git a/Engine/Source/Runtime/Core/Private/Tests/Async/AsyncTest.cpp b/Engine/Source/Runtime/Core/Private/Tests/Async/AsyncTest.cpp
--- a/Engine/Source/Runtime/Core/Private/Tests/Async/AsyncTest.cpp
+++ b/Engine/Source/Runtime/Core/Private/Tests/Async/AsyncTest.cpp
@@ -11,6 +11,7 @@ IMPLEMENT_SIMPLE_AUTOMATION_TEST(FAsyncGraphTest, "System.Core.Async.Async (Task
 IMPLEMENT_SIMPLE_AUTOMATION_TEST(FAsyncThreadedTaskTest, "System.Core.Async.Async (Thread)", EAutomationTestFlags::EditorContext | EAutomationTestFlags::EngineFilter)
 IMPLEMENT_SIMPLE_AUTOMATION_TEST(FAsyncThreadedPoolTest, "System.Core.Async.Async (Thread Pool)", EAutomationTestFlags::EditorContext | EAutomationTestFlags::EngineFilter)
 IMPLEMENT_SIMPLE_AUTOMATION_TEST(FAsyncVoidTaskTest, "System.Core.Async.Async (Void)", EAutomationTestFlags::EditorContext | EAutomationTestFlags::EngineFilter)
+IMPLEMENT_SIMPLE_AUTOMATION_TEST(FAsyncCompletionCallbackTest, "System.Core.Async.Async (Completion Callback)", EAutomationTestFlags::EditorContext | EAutomationTestFlags::EngineFilter)
 
 
 /** Helper methods used in the test cases. */
@@ -52,9 +53,7 @@ bool FAsyncThreadedTaskTest::RunTest(const FString& Parameters)
 }
 
 
-/**
-* Test that threaded pool tasks return correctly.
-*/
+/** Test that threaded pool tasks return correctly. */
 bool FAsyncThreadedPoolTest::RunTest(const FString& Parameters)
 {
 	auto Future = Async(EAsyncExecution::ThreadPool, AsyncTestUtils::Task);
@@ -79,3 +78,19 @@ bool FAsyncVoidTaskTest::RunTest(const FString& Parameters)
 
 	return true;
 }
+
+
+/** Test that asynchronous tasks have their completion callback called. */
+bool FAsyncCompletionCallbackTest::RunTest(const FString& Parameters)
+{
+	bool Completed = false;
+	auto Future = Async(EAsyncExecution::TaskGraph, AsyncTestUtils::Task, [&]() {
+		Completed = true;
+	});
+	int Result = Future.Get();
+
+	TestEqual(TEXT("Completion callback task must return expected value"), Result, 123);
+	TestTrue(TEXT("Completion callback task must call its callback function"), Completed);
+
+	return true;
+}
