Fix Crash Reporter analytics config delegate, as it was being called in static code before the engine had been initialized but needed to call engine functions.

diff --git a/Engine/Source/Programs/CrashReporter/CrashReportClient/Private/CrashReportAnalytics.cpp b/Engine/Source/Programs/CrashReporter/CrashReportClient/Private/CrashReportAnalytics.cpp
--- a/Engine/Source/Programs/CrashReporter/CrashReportClient/Private/CrashReportAnalytics.cpp
+++ b/Engine/Source/Programs/CrashReporter/CrashReportClient/Private/CrashReportAnalytics.cpp
@@ -12,16 +12,23 @@ bool FCrashReportAnalytics::bIsInitialized;
 TSharedPtr<IAnalyticsProvider> FCrashReportAnalytics::Analytics;
 
 /**
- * Engine analytics config log to initialize the engine analytics provider.
- * External code should bind this delegate if engine analytics are desired,
- * preferably in private code that won't be redistributed.
+ * Default config func that essentially tells the crash reporter to disable analytics.
  */
-FAnalyticsET::Config& GetCrashReportAnalyticsOverrideConfig()
+FAnalyticsET::Config DefaultAnalyticsConfigFunc()
 {
-	static FAnalyticsET::Config Config;
-	return Config;
+	return FAnalyticsET::Config();
 }
 
+/**
+* Engine analytics config to initialize the crash reporter analytics provider.
+* External code should bind this delegate if crash reporter analytics are desired,
+* preferably in private code that won't be redistributed.
+*/
+TFunction<FAnalyticsET::Config()>& GetCrashReportAnalyticsConfigFunc()
+{
+	static TFunction<FAnalyticsET::Config()> Config;
+	return Config;
+}
 
 /**
  * On-demand construction of the singleton. 
@@ -36,10 +43,11 @@ void FCrashReportAnalytics::Initialize()
 {
 	checkf(!bIsInitialized, TEXT("FCrashReportAnalytics::Initialize called more than once."));
 
-	if (!GetCrashReportAnalyticsOverrideConfig().APIServerET.IsEmpty())
+	FAnalyticsET::Config Config = GetCrashReportAnalyticsConfigFunc()();
+	if (!Config.APIServerET.IsEmpty())
 	{
 		// Connect the engine analytics provider (if there is a configuration delegate installed)
-		Analytics = FAnalyticsET::Get().CreateAnalyticsProvider(GetCrashReportAnalyticsOverrideConfig());
+		Analytics = FAnalyticsET::Get().CreateAnalyticsProvider(Config);
 		if( Analytics.IsValid() )
 		{
 			Analytics->SetUserID(FString::Printf(TEXT("%s|%s|%s"), *FPlatformMisc::GetMachineId().ToString(EGuidFormats::Digits).ToLower(), *FPlatformMisc::GetEpicAccountId(), *FPlatformMisc::GetOperatingSystemId()));
