Copying //UE4/Dev-Core to //UE4/Main

==========================
MAJOR FEATURES + CHANGES
==========================

Change 2877442 on 2016/02/23 by [REDACTED_PERSON][REDACTED_HANDLE]

	Allow 'dumpnonframe' stat command to also take an optional group name which will filter the output

Change 2879289 on 2016/02/24 by [REDACTED_PERSON][REDACTED_HANDLE]

	Removed "-" in switch string for load time tracker. Makes -NoLoadTrackClear cmd line param actually work

Change 2879323 on 2016/02/24 by [REDACTED_PERSON][REDACTED_HANDLE]

	PC builds now log engine initialization time with the load time tracker

Change 2879418 on 2016/02/24 by [REDACTED_PERSON][REDACTED_HANDLE]

	Optimization for FSignedArchiveReader
	 - Compute incoming data hash in parallel on another thread while we serialize and decrypt the signature on the chunk cache worker thread

Change 2879419 on 2016/02/24 by [REDACTED_PERSON][REDACTED_HANDLE]

	Extra stats for PakFile group
	 - Time reading from pak handles
	 - Number of open pak handles

Change 2881199 on 2016/02/25 by [REDACTED_PERSON][REDACTED_HANDLE]

	Use 32 bit CRCs to hash chunk contents rather than SHA1. Much faster to compute than SHA1, and the decrypting 4 bytes rather than 20 is also a lot faster

Change 2885687 on 2016/02/29 by [REDACTED_PERSON][REDACTED_HANDLE]

	Close disregard for GC pool when disabling it to leave it in a good state.

Change 2887358 on 2016/03/01 by [REDACTED_PERSON][REDACTED_HANDLE]

	FlushCurrentThreadCache calls forwarded on from FMalloc wrappers.

Change 2887398 on 2016/03/01 by [REDACTED_PERSON][REDACTED_HANDLE]

	Stats for MallocBinned2.

Change 2887682 on 2016/03/01 by [REDACTED_PERSON][REDACTED_HANDLE]

	Added code to catch game hangs. If game thread is stalled for more than 5 seconds we will catch that, report its callstack and crash.

Change 2889573 on 2016/03/02 by [REDACTED_PERSON][REDACTED_HANDLE]

	Renamed the new StackWalkAndDump to ThreadStackWalkAndDump to avoid name collisions and build errors.

Change 2889687 on 2016/03/02 by [REDACTED_PERSON][REDACTED_HANDLE]

	Send thread heartbeats during potentially slow tasks.

Change 2889963 on 2016/03/02 by [REDACTED_PERSON][REDACTED_HANDLE]

	Added the ability to query the name of a thread given its TLS id.

Change 2890439 on 2016/03/02 by [REDACTED_PERSON][REDACTED_HANDLE]

	MallocBinned redirection header, to allow easy switching between MallocBinned and MallocBinned2.
	Switch back to MallocBinned by default.

Change 2891799 on 2016/03/03 by [REDACTED_PERSON][REDACTED_HANDLE]

	Shrunken GMallocBinned2MemSizeToIndex array.
	Improvement of some shared templated alignment-related functions.

Change 2891812 on 2016/03/03 by [REDACTED_PERSON][REDACTED_HANDLE]

	Refactoring of pak signing to be more compatible with patching
	 - Store signatures in a block at the end of the pak file
	 - Optimize decryption of signatures by serializing them all when the pak is opened, and decrypting on the chunk cache worker thread between read requests.

Change 2891813 on 2016/03/03 by [REDACTED_PERSON][REDACTED_HANDLE]

	Don't compile UnrealPak with ICU

Change 2892386 on 2016/03/03 by [REDACTED_PERSON][REDACTED_HANDLE]

	GitHub #2070 - Fix incorrect FMallocDebug::Malloc alignment checking

Change 2894189 on 2016/03/04 by [REDACTED_PERSON][REDACTED_HANDLE]

	Memswap fixes for sizes of 3, 5, 6 and 7, and early out on 0.

Change 2894427 on 2016/03/04 by [REDACTED_PERSON][REDACTED_HANDLE]

	Don't display message boxes in unattended mode when D3D device is removed.

Change 2896912 on 2016/03/07 by [REDACTED_PERSON][REDACTED_HANDLE]

	Disable reading of anything outside the pak file when in shipping builds

Change 2897053 on 2016/03/07 by [REDACTED_PERSON][REDACTED_HANDLE]

	A few improvements to thread heartbeat thread:
	- Don't check for heartbeat in debug builds
	- Don't spawn the thread heartbeat thread for programs
	- Ignore stalls when exiting

Change 2898752 on 2016/03/08 by [REDACTED_PERSON][REDACTED_HANDLE]

	Fixing race conditions when using malloc verify

Change 2899075 on 2016/03/08 by [REDACTED_PERSON][REDACTED_HANDLE]

	UnrealPak will now accept hex values for p and q. Added the command line param to trust p and q.

Change 2899091 on 2016/03/08 by [REDACTED_PERSON][REDACTED_HANDLE]

	Fix for TQueue::Peek const-correctness.

Change 2902775 on 2016/03/10 by [REDACTED_PERSON][REDACTED_HANDLE]

	Adding code to test signature keys.

Change 2905040 on 2016/03/11 by [REDACTED_PERSON][REDACTED_HANDLE]

	TBigInt Multiply is now 30x faster

Change 2905218 on 2016/03/11 by [REDACTED_PERSON][REDACTED_HANDLE]

	New FCString::Strspn and FCString::Strcspn functions.

diff --git a/Engine/Source/Developer/HotReload/Private/HotReload.cpp b/Engine/Source/Developer/HotReload/Private/HotReload.cpp
--- a/Engine/Source/Developer/HotReload/Private/HotReload.cpp
+++ b/Engine/Source/Developer/HotReload/Private/HotReload.cpp
@@ -1358,6 +1358,8 @@ void FHotReloadModule::DoHotReloadFromIDE()
 
 		UE_LOG(LogHotReload, Log, TEXT("Starting Hot-Reload from IDE"));
 
+		HotReloadStartTime = FPlatformTime::Seconds();
+
 		FScopedSlowTask SlowTask(100.f, LOCTEXT("CompilingGameCode", "Compiling Game Code"));
 		SlowTask.MakeDialog();
 
diff --git a/Engine/Source/Editor/UnrealEd/Private/Settings/EditorProjectSettings.cpp b/Engine/Source/Editor/UnrealEd/Private/Settings/EditorProjectSettings.cpp
--- a/Engine/Source/Editor/UnrealEd/Private/Settings/EditorProjectSettings.cpp
+++ b/Engine/Source/Editor/UnrealEd/Private/Settings/EditorProjectSettings.cpp
@@ -4,9 +4,9 @@
 #include "EditorProjectSettings.h"
 
 
-EUnit ConvertDefaultInputUnits(UEditorProjectAppearanceSettings::EDefaultLocationUnit In)
+EUnit ConvertDefaultInputUnits(EDefaultLocationUnit In)
 {
-	typedef UEditorProjectAppearanceSettings::EDefaultLocationUnit EDefaultLocationUnit;
+	typedef EDefaultLocationUnit EDefaultLocationUnit;
 
 	switch(In)
 	{
diff --git a/Engine/Source/Programs/AutomationTool/AutomationUtils/ProjectParams.cs b/Engine/Source/Programs/AutomationTool/AutomationUtils/ProjectParams.cs
--- a/Engine/Source/Programs/AutomationTool/AutomationUtils/ProjectParams.cs
+++ b/Engine/Source/Programs/AutomationTool/AutomationUtils/ProjectParams.cs
@@ -498,6 +498,14 @@ namespace AutomationTool
 			this.Clean = GetOptionalParamValueIfNotSpecified(Command, Clean, this.Clean, "clean", null);
 			this.SignPak = ParseParamValueIfNotSpecified(Command, SignPak, "signpak", String.Empty);
 			this.SignedPak = !String.IsNullOrEmpty(this.SignPak) || GetParamValueIfNotSpecified(Command, SignedPak, this.SignedPak, "signedpak");
+			if (string.IsNullOrEmpty(this.SignPak))
+			{
+				this.SignPak = Path.Combine(RawProjectPath.Directory.FullName, @"Build\NoRedist\Keys.txt");
+				if (!File.Exists(this.SignPak))
+				{
+					this.SignPak = null;
+				}
+			}
 			this.Pak = this.SignedPak || GetParamValueIfNotSpecified(Command, Pak, this.Pak, "pak");
 			this.SkipPak = GetParamValueIfNotSpecified(Command, SkipPak, this.SkipPak, "skippak");
 			if (this.SkipPak)
diff --git a/Engine/Source/Programs/AutomationTool/AutomationUtils/UE4Build.cs b/Engine/Source/Programs/AutomationTool/AutomationUtils/UE4Build.cs
--- a/Engine/Source/Programs/AutomationTool/AutomationUtils/UE4Build.cs
+++ b/Engine/Source/Programs/AutomationTool/AutomationUtils/UE4Build.cs
@@ -565,11 +565,11 @@ namespace AutomationTool
 
 			LogVerbose("Updating {0} with:", VerFile);
 			LogVerbose(" #define DECRYPTION_KEY_EXPONENT {0}", PublicKeyExponent);
-			LogVerbose(" #define DECYRPTION_KEY_MODULUS {0}", Modulus);
+			LogVerbose(" #define DECRYPTION_KEY_MODULUS {0}", Modulus);
 
 			VersionFileUpdater PublicKeyInl = new VersionFileUpdater(VerFile);
 			PublicKeyInl.ReplaceLine("#define DECRYPTION_KEY_EXPONENT ", PublicKeyExponent);
-			PublicKeyInl.ReplaceLine("#define DECYRPTION_KEY_MODULUS ", Modulus);
+			PublicKeyInl.ReplaceLine("#define DECRYPTION_KEY_MODULUS ", Modulus);
 			PublicKeyInl.Commit();
 		}
 
@@ -1233,9 +1233,6 @@ namespace AutomationTool
 				UpdateVersionFiles(ActuallyUpdateVersionFiles: true, ChangelistNumberOverride: InChangelistNumberOverride);
 			}
 
-			var EncryptionKeyFilename = ParseParamValue("SignPak");
-			UpdatePublicKey(EncryptionKeyFilename);
-
 			//////////////////////////////////////
 
 			// make a set of unique platforms involved
diff --git a/Engine/Source/Programs/AutomationTool/AutomationUtils/Utils.cs b/Engine/Source/Programs/AutomationTool/AutomationUtils/Utils.cs
--- a/Engine/Source/Programs/AutomationTool/AutomationUtils/Utils.cs
+++ b/Engine/Source/Programs/AutomationTool/AutomationUtils/Utils.cs
@@ -1114,6 +1114,7 @@ namespace AutomationTool
 				{
 					throw new AutomationException("Unable to update version info in {0}", MyFile.FullName);
 				}
+				OriginalLines = Lines;
 			}
 		}
 
diff --git a/Engine/Source/Programs/AutomationTool/Scripts/CopyBuildToStagingDirectory.Automation.cs b/Engine/Source/Programs/AutomationTool/Scripts/CopyBuildToStagingDirectory.Automation.cs
--- a/Engine/Source/Programs/AutomationTool/Scripts/CopyBuildToStagingDirectory.Automation.cs
+++ b/Engine/Source/Programs/AutomationTool/Scripts/CopyBuildToStagingDirectory.Automation.cs
@@ -751,7 +751,7 @@ public partial class Project : CommandUtils
 
 		var UnrealPakResponseFile = CreatePakResponseFileFromStagingManifest(SC);
 
-		CreatePak(Params, SC, UnrealPakResponseFile, SC.ShortProjectName);
+		CreatePak(Params, SC, UnrealPakResponseFile, SC.ShortProjectName, Params.SignPak);
 	}
 
 	/// <summary>
@@ -844,7 +844,7 @@ public partial class Project : CommandUtils
 	/// <param name="SC"></param>
 	/// <param name="UnrealPakResponseFile"></param>
 	/// <param name="PakName"></param>
-	private static void CreatePak(ProjectParams Params, DeploymentContext SC, Dictionary<string, string> UnrealPakResponseFile, string PakName)
+	private static void CreatePak(ProjectParams Params, DeploymentContext SC, Dictionary<string, string> UnrealPakResponseFile, string PakName, string EncryptionKeys)
 	{
         bool bShouldGeneratePatch = Params.IsGeneratingPatch && SC.StageTargetPlatform.GetPlatformPatchesWithDiffPak(Params, SC);
 
@@ -930,7 +930,7 @@ public partial class Project : CommandUtils
 			}
 			if (!bCopiedExistingPak)
 			{
-				RunUnrealPak(UnrealPakResponseFile, OutputLocation, Params.SignPak, PakOrderFileLocation, SC.StageTargetPlatform.GetPlatformPakCommandLine(), Params.Compressed, PatchSourceContentPath);
+				RunUnrealPak(UnrealPakResponseFile, OutputLocation, EncryptionKeys, PakOrderFileLocation, SC.StageTargetPlatform.GetPlatformPakCommandLine(), Params.Compressed, PatchSourceContentPath);
 			}
 		}
 
@@ -1102,7 +1102,7 @@ public partial class Project : CommandUtils
 		System.Threading.Tasks.Parallel.ForEach(PakPairs, (PakPair) =>
 		{
 			var ChunkName = Path.GetFileNameWithoutExtension(PakPair.Item2);
-			CreatePak(Params, SC, PakPair.Item1, ChunkName);
+			CreatePak(Params, SC, PakPair.Item1, ChunkName, Params.SignPak);
 		});
 
 		String ChunkLayerFilename = CombinePaths(GetTmpPackagingPath(Params, SC), GetChunkPakLayerListName());
diff --git a/Engine/Source/Programs/CrashReporter/CrashReportCommon/CrashContext.cs b/Engine/Source/Programs/CrashReporter/CrashReportCommon/CrashContext.cs
--- a/Engine/Source/Programs/CrashReporter/CrashReportCommon/CrashContext.cs
+++ b/Engine/Source/Programs/CrashReporter/CrashReportCommon/CrashContext.cs
@@ -620,6 +620,10 @@ namespace Tools.CrashReporter.CrashReportCommon
 		[XmlElement]
 		public bool bAllowToBeContacted = false;
 
+		/// <summary> Warning info for missing data or failed processing in the Crash Report Processor. </summary>
+		[XmlElement]
+		public string ProcessorFailedMessage;
+
 		/// <summary>A simple default constructor to allow Xml serialization.</summary>
 		public FPrimaryCrashProperties()
 		{ }
diff --git a/Engine/Source/Programs/CrashReporter/CrashReportCommon/CrashReportCommon.csproj b/Engine/Source/Programs/CrashReporter/CrashReportCommon/CrashReportCommon.csproj
--- a/Engine/Source/Programs/CrashReporter/CrashReportCommon/CrashReportCommon.csproj
+++ b/Engine/Source/Programs/CrashReporter/CrashReportCommon/CrashReportCommon.csproj
@@ -51,6 +51,7 @@
   <ItemGroup>
     <Reference Include="System" />
     <Reference Include="System.Core" />
+    <Reference Include="System.Net" />
     <Reference Include="System.Web" />
     <Reference Include="System.Web.Extensions" />
     <Reference Include="System.Xml.Linq" />
@@ -70,6 +71,8 @@
       <Link>Properties\MetaData.cs</Link>
     </Compile>
     <Compile Include="ReportData.cs" />
+    <Compile Include="SlackWriter.cs" />
+    <Compile Include="Stats.cs" />
     <Compile Include="WerReport.cs" />
   </ItemGroup>
   <ItemGroup>
diff --git a/Engine/Source/Programs/CrashReporter/CrashReportCommon/LogWriter.cs b/Engine/Source/Programs/CrashReporter/CrashReportCommon/LogWriter.cs
--- a/Engine/Source/Programs/CrashReporter/CrashReportCommon/LogWriter.cs
+++ b/Engine/Source/Programs/CrashReporter/CrashReportCommon/LogWriter.cs
@@ -55,7 +55,7 @@ namespace Tools.CrashReporter.CrashReportCommon
 		{
 			if (Line != null)
 			{
-				string FullLine = DateTime.Now.ToString( "yyyy/MM/dd HH:mm:ss" ) + " [" + Thread.CurrentThread.ManagedThreadId.ToString( "D3" ) + "] : " + Line;
+				string FullLine = DateTime.UtcNow.ToString( "yyyy/MM/dd HH:mm:ss" ) + "UTC [" + Thread.CurrentThread.ManagedThreadId.ToString( "D3" ) + "] : " + Line;
 				lock( Sync )
 				{
 					LogFile.WriteLine( FullLine );
diff --git a/Engine/Source/Programs/CrashReporter/CrashReportCommon/SlackWriter.cs b/Engine/Source/Programs/CrashReporter/CrashReportCommon/SlackWriter.cs
new file mode 100644
--- /dev/null
+++ b/Engine/Source/Programs/CrashReporter/CrashReportCommon/SlackWriter.cs
@@ -0,0 +1,125 @@
+﻿using System;
+using System.Collections.Generic;
+using System.Collections.Specialized;
+using System.Linq;
+using System.Net;
+using System.Text;
+using System.Threading.Tasks;
+using System.Web;
+
+namespace Tools.CrashReporter.CrashReportCommon
+{
+	/// <summary>
+	/// Allows us to write messages to Slack
+	/// </summary>
+	public class SlackWriter
+	{
+		/// <summary>
+		/// Full URL from your Incoming Webhooks Integration on the Slack website
+		/// </summary>
+		public string WebhookUrl { get; set; }
+
+		/// <summary>
+		/// Overrides the channel to which the message is posted. The Incoming Webhooks Integration created on the Slack website will have a default channel.
+		/// </summary>
+		public string Channel { get; set; }
+
+		/// <summary>
+		/// Overrides the username with which the message is posted. The Incoming Webhooks Integration created on the Slack website will have a default username.
+		/// </summary>
+		public string Username { get; set; }
+
+		/// <summary>
+		/// Overrides the emoji for the username with which the message is posted. The Incoming Webhooks Integration created on the Slack website will have a default emoji.
+		/// </summary>
+		public string IconEmoji { get; set; }
+
+		/// <summary>
+		/// Sends a message to Slack
+		/// </summary>
+		/// <param name="MessageText"></param>
+		public void Write(string MessageText)
+		{
+			lock (Lock)
+			{
+				if (!string.IsNullOrEmpty(WebhookUrl))
+				{
+					string WebhookSnapshot = WebhookUrl;
+					string ChannelSnapshot = Channel;
+					string UsernameSnapshot = Username;
+					string IconEmojiSnapshot = IconEmoji;
+
+					WriteTask = WriteTask.ContinueWith(PreviousTask => WriteAsync(WebhookSnapshot, ChannelSnapshot, UsernameSnapshot, IconEmojiSnapshot, MessageText));
+				}
+			}
+		}
+
+		/// <summary>
+		/// Shutdown and flush the Slack writer
+		/// </summary>
+		public void Dispose()
+		{
+			lock (Lock)
+			{
+				WriteTask.Wait();
+			}
+		}
+
+		private string WriteAsync(string InWebhook, string InChannel, string InUsername, string InIconEmoji, string InMessageText)
+		{
+			try
+			{
+				using (WebClient Client = new WebClient())
+				{
+					Client.Headers[HttpRequestHeader.ContentType] = "application/json";
+					Client.Encoding = Encoding;
+
+					string Payload = GetPayload(InChannel, InUsername, InIconEmoji, InMessageText);
+
+					return Client.UploadString(new Uri(InWebhook), "POST", Payload);
+				}
+			}
+			catch (WebException Ex)
+			{
+				return Ex.Message;
+			}
+		}
+
+		private static string GetPayload(string InChannel, string InUsername, string InIconEmoji, string InMessageText)
+		{
+			StringBuilder Payload = new StringBuilder();
+			Payload.Append("{");
+
+			if (!string.IsNullOrEmpty(InChannel))
+			{
+				Payload.Append(string.Format("\"channel\": \"{0}\", ", JsonEncode(InChannel)));
+			}
+			if (!string.IsNullOrEmpty(InUsername))
+			{
+				Payload.Append(string.Format("\"username\": \"{0}\", ", JsonEncode(InUsername)));
+			}
+			if (!string.IsNullOrEmpty(InIconEmoji))
+			{
+				Payload.Append(string.Format("\"icon_emoji\": \"{0}\", ", JsonEncode(InIconEmoji)));
+			}
+			Payload.Append(string.Format("\"text\": \"{0}\"}}", JsonEncode(InMessageText)));
+
+			return Payload.ToString();
+		}
+
+		private static string JsonEncode(string InString)
+		{
+			return InString
+				.Replace(@"\", @"\\")
+				.Replace("\r\n", @"\n")
+				.Replace(Environment.NewLine, @"\n")
+				.Replace("\r", @"\n")
+				.Replace("\n", @"\n")
+				.Replace("\"", "\\\"");
+		}
+
+		private readonly Encoding Encoding = new UTF8Encoding();
+		private readonly Object Lock = new Object();
+		private Task<string> WriteTask = Task.FromResult("init");
+	}
+}
diff --git a/Engine/Source/Programs/CrashReporter/CrashReportCommon/Stats.cs b/Engine/Source/Programs/CrashReporter/CrashReportCommon/Stats.cs
new file mode 100644
--- /dev/null
+++ b/Engine/Source/Programs/CrashReporter/CrashReportCommon/Stats.cs
@@ -0,0 +1,99 @@
+﻿using System;
+using System.Collections.Generic;
+using System.Linq;
+using System.Text;
+using System.Threading.Tasks;
+using System.Web.ModelBinding;
+
+namespace Tools.CrashReporter.CrashReportCommon
+{
+	/// <summary>
+	/// Thread-safe event counter for measuring approximate rate over a fixed period between now and a fixed time ago
+	/// </summary>
+	public class EventCounter
+	{
+		private class Bin
+		{
+			public int Count { get; set; }
+		}
+
+		/// <summary>
+		/// Create an event counter
+		/// </summary>
+		/// <param name="InAvePeriod">The period over which events are counter. The time we keep them and calculate the average rate.</param>
+		/// <param name="InResolution">The number of bins in which events are kept in the given period.</param>
+		public EventCounter(TimeSpan InAvePeriod, int InResolution)
+		{
+			Period = InAvePeriod;
+			Resolution = InResolution;
+			Bins = new Queue<Bin>(Resolution);
+			BinPeriod = TimeSpan.FromTicks(Period.Ticks/Resolution);
+			ActiveBin = new Bin();
+			Bins.Enqueue(ActiveBin);
+			Count = 0;
+			Timestamp = DateTime.UtcNow;
+		}
+
+		/// <summary>
+		/// Adds events to the counter. Keeps track of internal bins and updates them keeping only those in our Period.
+		/// </summary>
+		/// <param name="InCount">The number of events to add</param>
+		public void AddEvent(int InCount = 1)
+		{
+			lock (ThisLock)
+			{
+				TotalCount += InCount;
+				DateTime EventTime = DateTime.UtcNow;
+
+				while (EventTime - Timestamp > BinPeriod)
+				{
+					ActiveBin = new Bin();
+					Bins.Enqueue(ActiveBin);
+					Timestamp += BinPeriod;
+				}
+
+				Count += InCount;
+				ActiveBin.Count += InCount;
+
+				while (Bins.Count > Resolution)
+				{
+					Count -= Bins.Dequeue().Count;
+				}
+			}
+		}
+
+		/// <summary>
+		/// Count of all events since the counter was created.
+		/// </summary>
+		public int TotalEvents { get { return TotalCount; } }
+
+		/// <summary>
+		/// Count of all events currently in the time period we're interested in (Period) preceding now.
+		/// </summary>
+		public int EventsInPeriod { get { return Count; } }
+
+		/// <summary>
+		/// Current rate of events average over the period.
+		/// </summary>
+		public double EventsPerSecond
+		{
+			get
+			{
+				lock (ThisLock)
+				{
+					return Count/(BinPeriod.TotalSeconds*Bins.Count);
+				}
+			}
+		}
+
+		private DateTime Timestamp;
+		private readonly TimeSpan Period;
+		private readonly int Resolution;
+		private readonly TimeSpan BinPeriod;
+		private readonly Queue<Bin> Bins;
+		private Bin ActiveBin;
+		private int TotalCount;
+		private int Count;
+		private readonly Object ThisLock = new Object();
+	}
+}
diff --git a/Engine/Source/Programs/CrashReporter/CrashReportProcess/App.config b/Engine/Source/Programs/CrashReporter/CrashReportProcess/App.config
--- a/Engine/Source/Programs/CrashReporter/CrashReportProcess/App.config
+++ b/Engine/Source/Programs/CrashReporter/CrashReportProcess/App.config
@@ -1,55 +1,9 @@
 <?xml version="1.0"?>
 <configuration>
 	<configSections>
-		<sectionGroup name="applicationSettings" type="System.Configuration.ApplicationSettingsGroup, System, Version=4.0.0.0, Culture=neutral, PublicKeyToken=b77a5c561934e089" >
-			<section name="Tools.CrashReporter.CrashReportProcess.Properties.Settings" type="System.Configuration.ClientSettingsSection, System, Version=4.0.0.0, Culture=neutral, PublicKeyToken=b77a5c561934e089" requirePermission="false" />
-		</sectionGroup>
  </configSections>
 	<connectionStrings />
 	<startup>
 		<supportedRuntime version="v4.0" sku=".NETFramework,Version=v4.0,Profile=Client"/>
 	</startup>
-	<applicationSettings>
-		<Tools.CrashReporter.CrashReportProcess.Properties.Settings>
-   <setting name="ProcessedReports" serializeAs="String">
-    <value>E:\Sites\CrashReporterFiles</value>
-   </setting>
-   <setting name="DaysToSunsetReport" serializeAs="String">
-    <value>14</value>
-   </setting>
-   <setting name="CrashReportWebSite" serializeAs="String">
-    <value>CrashReporterUE4</value>
-   </setting>
-   <setting name="ProcessedVideos" serializeAs="String">
-    <value>F:\CrashReporterVideos</value>
-   </setting>
-   <setting name="DepotRoot" serializeAs="String">
-    <value>F:\depot</value>
-   </setting>
-   <setting name="InternalLandingZone" serializeAs="String">
-    <value>F:\InternalCrashReports\CrashReportRepository</value>
-   </setting>
-   <setting name="ExternalLandingZone" serializeAs="String">
-    <value>\\na-web-cluster.dmz.epicgames.net\crash_reporter\Rep</value>
-   </setting>
-   <setting name="InvalidReportsDirectory" serializeAs="String">
-    <value>F:\InvalidCrashReports</value>
-   </setting>
-   <setting name="DepotIndex" serializeAs="String">
-    <value>crash-01</value>
-   </setting>
-   <setting name="SyncBinariesFromDepot" serializeAs="String">
-    <value>Stream\UE4\Main\Engine\Binaries\Win64\MinidumpDiagnostics*</value>
-   </setting>
-   <setting name="SyncConfigFromDepot" serializeAs="String">
-    <value>Stream\UE4\Main\Engine\Config\...</value>
-   </setting>
-   <setting name="SyncThirdPartyFromDepot" serializeAs="String">
-    <value>Stream\UE4\Main\Engine\Binaries\ThirdParty\OpenSSL\...</value>
-   </setting>
-   <setting name="MDDBinariesFolderInDepot" serializeAs="String">
-    <value>Stream\UE4\Main\Engine\Binaries\Win64</value>
-   </setting>
-  </Tools.CrashReporter.CrashReportProcess.Properties.Settings>
-	</applicationSettings>
 </configuration>
diff --git a/Engine/Source/Programs/CrashReporter/CrashReportProcess/Config.cs b/Engine/Source/Programs/CrashReporter/CrashReportProcess/Config.cs
new file mode 100644
--- /dev/null
+++ b/Engine/Source/Programs/CrashReporter/CrashReportProcess/Config.cs
@@ -0,0 +1,222 @@
+﻿using System;
+using System.Collections.Generic;
+using System.IO;
+using System.Linq;
+using System.Text;
+using System.Threading.Tasks;
+using System.Xml;
+using System.Xml.Serialization;
+using System.Reflection;
+
+namespace Tools.CrashReporter.CrashReportProcess
+{
+	/// <summary>
+	/// Config values for the Crash Report Process. The binary expects to find an xml representation in the binary folder at runtime.
+	/// </summary>
+	[XmlRoot]
+	public class Config
+	{
+		/// <summary>
+		/// CRP version for displaying in status reports/logs etc
+		/// </summary>
+		[XmlElement]
+		public string VersionString { get; set; }
+
+		/// <summary>
+		/// Folder where the files that are linked from the website are stored.
+		/// </summary>
+		[XmlElement]
+		public string ProcessedReports { get; set; }
+
+		/// <summary>
+		/// Time after which old, unprocessed reports in landing zones are deleted.
+		/// </summary>
+		[XmlElement]
+		public int DeleteWaitingReportsDays { get; set; }
+
+		/// <summary>
+		/// The website domain to which we output crashes.
+		/// </summary>
+		[XmlElement]
+		public string CrashReportWebSite { get; set; }
+
+		/// <summary>
+		/// Folder where the video file that are linked from the website are stored.
+		/// </summary>
+		[XmlElement]
+		public string ProcessedVideos { get; set; }
+
+		/// <summary>
+		/// Folder where the Perforce depot is located.
+		/// </summary>
+		[XmlElement]
+		public string DepotRoot { get; set; }
+
+		/// <summary>
+		/// Folder where new crash reports are queued from internal, company users.
+		/// </summary>
+		[XmlElement]
+		public string InternalLandingZone { get; set; }
+
+		/// <summary>
+		/// Folder where new crash reports are queued from external users.
+		/// </summary>
+		[XmlElement]
+		public string ExternalLandingZone { get; set; }
+
+		/// <summary>
+		/// Folder where failed reports are moved.
+		/// </summary>
+		[XmlElement]
+		public string InvalidReportsDirectory { get; set; }
+
+		/// <summary>
+		/// String passed to Minidump Diagnostics to modify its Perforce depot root.
+		/// </summary>
+		[XmlElement]
+		public string DepotIndex { get; set; }
+
+		/// <summary>
+		/// String specifying the binaries that will be synched from source control to get the latest MinidumpDiagnostics.
+		/// </summary>
+		[XmlElement]
+		public string SyncBinariesFromDepot { get; set; }
+
+		/// <summary>
+		/// String specifying the config files that will be synched from source control for MinidumpDiagnostics.
+		/// </summary>
+		[XmlElement]
+		public string SyncConfigFromDepot { get; set; }
+
+		/// <summary>
+		/// String specifying the third-party files that will be synched from source control for MinidumpDiagnostics. e.g. OpenSSL
+		/// </summary>
+		[XmlElement]
+		public string SyncThirdPartyFromDepot { get; set; }
+
+		/// <summary>
+		/// String specifying the local path of the folder containing the MinidumpDiagnostics exe
+		/// </summary>
+		[XmlElement]
+		public string MDDBinariesFolderInDepot { get; set; }
+
+		/// <summary>
+		/// The number of threads created by each main processor thread to upload crashes to the website. (relieves a bottleneck when using a single processor thread)
+		/// </summary>
+		[XmlElement]
+		public int AddReportsPerProcessor { get; set; }
+
+		/// <summary>
+		/// The number of main processor threads
+		/// </summary>
+		[XmlElement]
+		public int ProcessorThreadCount { get; set; }
+
+		/// <summary>
+		/// Incoming webhook URL for Slack integration.
+		/// </summary>
+		[XmlElement]
+		public string SlackWebhookUrl { get; set; }
+
+		/// <summary>
+		/// Slack channel to which messages are posted. Can be null or empty to use the webhook default.
+		/// </summary>
+		[XmlElement]
+		public string SlackChannel { get; set; }
+
+		/// <summary>
+		/// Slack username from which messages are posted. Can be null or empty to use the webhook default.
+		/// </summary>
+		[XmlElement]
+		public string SlackUsername { get; set; }
+
+		/// <summary>
+		/// Slack user icon name (e.g. ":someemoji:") from which messages are posted. Can be null or empty to use the webhook default.
+		/// </summary>
+		[XmlElement]
+		public string SlackEmoji { get; set; }
+
+		/// <summary>
+		/// Local folder used to setup testing folders in Debug builds. Overrides other folder params in config.
+		/// </summary>
+		[XmlElement]
+		public string DebugTestingFolder { get; set; }
+
+		/// <summary>
+		/// Period between reports to Slack (if available).
+		/// </summary>
+		[XmlElement]
+		public int MinutesBetweenQueueSizeReports { get; set; }
+
+		/// <summary>
+		/// The limit for AddCrash() calls to the website failing in a row. An alert will be created if this is exceeded.
+		/// </summary>
+		[XmlElement]
+		public int ConsecutiveFailedWebAddLimit { get; set; }
+
+		/// <summary>
+		/// Get the default config object (lazy loads it on first access)
+		/// </summary>
+		public static Config Default
+		{
+			get
+			{
+				if (DefaultSingleton == null)
+				{
+					LoadConfig();
+				}
+				return DefaultSingleton;
+			}
+		}
+
+		/// <summary>
+		/// Force reload default config object from disk
+		/// </summary>
+		public static void LoadConfig()
+		{
+			DefaultSingleton = LoadConfigPrivate();
+		}
+
+		private static Config LoadConfigPrivate()
+		{
+			Config LoadedConfig = new Config();
+
+			string ExePath = Assembly.GetEntryAssembly().Location;
+			string ExeFolder = Path.GetDirectoryName(ExePath);
+			string ConfigPath = Path.Combine(ExeFolder, "CrashReportProcess.config");
+
+			if (File.Exists(ConfigPath))
+			{
+				if (!string.IsNullOrEmpty(ConfigPath))
+				{
+					using (XmlReader Reader = XmlReader.Create(ConfigPath))
+					{
+						CrashReporterProcessServicer.WriteEvent("Loading config from " + ConfigPath);
+						XmlSerializer Xml = new XmlSerializer(typeof (Config));
+						LoadedConfig = Xml.Deserialize(Reader) as Config;
+					}
+				}
+			}
+
+#if DEBUG
+			// Debug mode redirects to local folder in DebugTestingFolder
+			LoadedConfig.ProcessedReports = Path.Combine(LoadedConfig.DebugTestingFolder, "ProcessedReports");
+			LoadedConfig.ProcessedVideos = Path.Combine(LoadedConfig.DebugTestingFolder, "ProcessedVideos");
+			LoadedConfig.DepotRoot = Path.Combine(LoadedConfig.DebugTestingFolder, "DepotRoot");
+			LoadedConfig.InternalLandingZone = Path.Combine(LoadedConfig.DebugTestingFolder, "InternalLandingZone");
+			LoadedConfig.ExternalLandingZone = Path.Combine(LoadedConfig.DebugTestingFolder, "ExternalLandingZone");
+			LoadedConfig.InvalidReportsDirectory = Path.Combine(LoadedConfig.DebugTestingFolder, "InvalidReportsDirectory");
+			LoadedConfig.VersionString += " debugbuild";
+			
+#if SLACKTESTING
+			LoadedConfig.SlackUsername = "CrashReportProcess_TESTING_IgnoreMe";
+#else
+			LoadedConfig.SlackWebhookUrl = string.Empty;	// no Slack in dbeug
+#endif
+#endif
+			return LoadedConfig;
+		}
+
+		private static Config DefaultSingleton;
+	}
+}
diff --git a/Engine/Source/Programs/CrashReporter/CrashReportProcess/Config/CrashReportProcess.config b/Engine/Source/Programs/CrashReporter/CrashReportProcess/Config/CrashReportProcess.config
new file mode 100644
--- /dev/null
+++ b/Engine/Source/Programs/CrashReporter/CrashReportProcess/Config/CrashReportProcess.config
@@ -0,0 +1,26 @@
+﻿<?xml version="1.0" encoding="utf-8"?>
+<Config xmlns:xsi="http://www.w3.org/2001/XMLSchema-instance" xmlns:xsd="http://www.w3.org/2001/XMLSchema">
+	<VersionString>MyVersionString_v0.1</VersionString>
+	<ProcessedReports>MyProcessedReportsFolderPath</ProcessedReports>
+	<DeleteWaitingReportsDays>14</DeleteWaitingReportsDays>
+	<CrashReportWebSite>MyCrashReportWebsiteName</CrashReportWebSite>
+	<ProcessedVideos>MyProcessedVideosFolderPath</ProcessedVideos>
+	<DepotRoot>MyP4DepotRootFolderPath</DepotRoot>
+	<InternalLandingZone>MyInternalCrashLandingZoneFolderPath</InternalLandingZone>
+	<ExternalLandingZone>MyExternallCrashLandingZoneFolderPath</ExternalLandingZone>
+	<InvalidReportsDirectory>MyInvalidCrashReportsFolderPath</InvalidReportsDirectory>
+	<DepotIndex>MyDepotCustomizationhPassedToMDD</DepotIndex>
+	<SyncBinariesFromDepot>MyMDDBinariesPathInP4Depot</SyncBinariesFromDepot>
+	<SyncConfigFromDepot>MyMDDConfigPathInP4Depot</SyncConfigFromDepot>
+	<SyncThirdPartyFromDepot>MyThirdPartyBinariesPathInP4Depot</SyncThirdPartyFromDepot>
+	<MDDBinariesFolderInDepot>MyMDDBinariesFolderPathInP4Depot</MDDBinariesFolderInDepot>
+	<AddReportsPerProcessor>1</AddReportsPerProcessor>
+	<ProcessorThreadCount>10</ProcessorThreadCount>
+	<SlackWebhookUrl>MySlackWebhookForPostingMessages</SlackWebhookUrl>
+	<SlackChannel>MySlackChannelForPostingMessages</SlackChannel>
+	<SlackUsername>MySlackUserForPostingMessages</SlackUsername>
+	<SlackEmoji>MySlackUserEmojiForPostingMessages</SlackEmoji>
+	<DebugTestingFolder>MyDebugTestingFolderPath</DebugTestingFolder>
+	<MinutesBetweenQueueSizeReports>60</MinutesBetweenQueueSizeReports>
+	<ConsecutiveFailedWebAddLimit>10</ConsecutiveFailedWebAddLimit>
+</Config>
\ No newline at end of file
diff --git a/Engine/Source/Programs/CrashReporter/CrashReportProcess/CrashReportProcess.csproj b/Engine/Source/Programs/CrashReporter/CrashReportProcess/CrashReportProcess.csproj
--- a/Engine/Source/Programs/CrashReporter/CrashReportProcess/CrashReportProcess.csproj
+++ b/Engine/Source/Programs/CrashReporter/CrashReportProcess/CrashReportProcess.csproj
@@ -61,6 +61,17 @@
     <ErrorReport>prompt</ErrorReport>
     <CodeAnalysisRuleSet>MinimumRecommendedRules.ruleset</CodeAnalysisRuleSet>
   </PropertyGroup>
+  <PropertyGroup Condition="'$(Configuration)|$(Platform)' == 'DebugSlackTesting|AnyCPU'">
+    <DebugSymbols>true</DebugSymbols>
+    <OutputPath>..\..\..\..\Binaries\DotNET\</OutputPath>
+    <DefineConstants>TRACE;DEBUG;SLACKTESTING</DefineConstants>
+    <AllowUnsafeBlocks>true</AllowUnsafeBlocks>
+    <DocumentationFile>..\..\..\..\Binaries\DotNET\CrashReportProcess.xml</DocumentationFile>
+    <DebugType>full</DebugType>
+    <PlatformTarget>AnyCPU</PlatformTarget>
+    <ErrorReport>prompt</ErrorReport>
+    <CodeAnalysisRuleSet>MinimumRecommendedRules.ruleset</CodeAnalysisRuleSet>
+  </PropertyGroup>
   <ItemGroup>
     <Reference Include="System" />
     <Reference Include="System.Configuration.Install" />
@@ -75,6 +86,7 @@
     <Reference Include="System.Xml" />
   </ItemGroup>
   <ItemGroup>
+    <Compile Include="Config.cs" />
     <Compile Include="CrashReportProcessService.cs">
       <SubType>Component</SubType>
     </Compile>
@@ -97,21 +109,28 @@
       <DesignTimeSharedInput>True</DesignTimeSharedInput>
       <DependentUpon>Settings.settings</DependentUpon>
     </Compile>
+    <Compile Include="ReceiverReportQueue.cs" />
     <Compile Include="ReportProcessor.cs" />
     <Compile Include="ReportQueue.cs" />
     <Compile Include="ReportWatcher.cs" />
     <Compile Include="Settings.cs" />
+    <Compile Include="StatusReporting.cs" />
+    <Compile Include="StatusReportingConstants.cs" />
   </ItemGroup>
   <ItemGroup>
     <None Include="App.config">
       <SubType>Designer</SubType>
     </None>
+    <None Include="Config\CrashReportProcess.config" />
     <None Include="Properties\Settings.settings">
       <Generator>SettingsSingleFileGenerator</Generator>
       <LastGenOutput>Settings.Designer.cs</LastGenOutput>
     </None>
   </ItemGroup>
   <ItemGroup>
+    <None Include="Config\NotForLicensees\CrashReportProcess.config">
+      <SubType>Designer</SubType>
+    </None>
   </ItemGroup>
   <ItemGroup>
     <ProjectReference Include="..\..\DotNETCommon\DotNETUtilities\DotNETUtilities.csproj">
@@ -129,8 +148,10 @@
   </ItemGroup>
   <Import Project="$(MSBuildToolsPath)\Microsoft.CSharp.targets" />
   <PropertyGroup>
-    <PostBuildEvent>
-    </PostBuildEvent>
+    <PostBuildEvent>xcopy "$(ProjectDir)Config\CrashReportProcess.config" "$(ProjectDir)$(OutDir)" /y
+IF EXIST "$(ProjectDir)Config\NotForLicensees\CrashReportProcess.config" (
+xcopy "$(ProjectDir)Config\NotForLicensees\CrashReportProcess.config" "$(ProjectDir)$(OutDir)" /y
+)</PostBuildEvent>
   </PropertyGroup>
   <!-- To modify your build process, add your task inside one of the targets below and uncomment it. 
        Other similar extension points exist, see Microsoft.Common.targets.
diff --git a/Engine/Source/Programs/CrashReporter/CrashReportProcess/CrashReportProcessService.cs b/Engine/Source/Programs/CrashReporter/CrashReportProcess/CrashReportProcessService.cs
--- a/Engine/Source/Programs/CrashReporter/CrashReportProcess/CrashReportProcessService.cs
+++ b/Engine/Source/Programs/CrashReporter/CrashReportProcess/CrashReportProcessService.cs
@@ -11,22 +11,28 @@ using System.Reflection;
 using System.ServiceProcess;
 using System.Text;
 using System.Threading.Tasks;
+using Tools.CrashReporter.CrashReportCommon;
+using Tools.CrashReporter.CrashReportProcess.Properties;
 
 namespace Tools.CrashReporter.CrashReportProcess
 {
 	/// <summary>
 	/// A class to handle processing received crash reports for displaying on the website.
 	/// </summary>
-	public partial class CrashReporterProcessServicer : ServiceBase
+	partial class CrashReporterProcessServicer : ServiceBase
 	{
 		/// <summary>A class the handle detection of new reports.</summary>
 		public ReportWatcher Watcher = null;
 
 		/// <summary>A class to lazily process reports as they are detected.</summary>
-		public ReportProcessor Processor = null;
+		public readonly List<ReportProcessor> Processors = new List<ReportProcessor>();
 
 		/// <summary>Current log file to write debug progress info to</summary>
-		static public CrashReportCommon.LogWriter Log = null;
+		public static LogWriter Log = null;
+
+		private static SlackWriter Slack = null;
+
+		public static StatusReporting StatusReporter = null;
 
 		/// <summary>Folder in which to store log files</summary>
 		static private string LogFolder = null;
@@ -83,6 +89,18 @@ namespace Tools.CrashReporter.CrashReportProcess
 			if( Message != null && Message.Length > 2 )
 			{
 				Log.Print( "[EXCEPT] " + Message );
+				StatusReporter.IncrementCount(StatusReportingConstants.ExceptionEvent);
+			}
+		}
+
+		/// <summary>
+		/// Write to Slack.
+		/// </summary>
+		static public void WriteSlack(string Message)
+		{
+			if (Message != null && Message.Length > 0)
+			{
+				Slack.Write(Message);
 			}
 		}
 
@@ -104,14 +122,32 @@ namespace Tools.CrashReporter.CrashReportProcess
 		protected override void OnStart( string[] Arguments )
 		{
 			// Create a log file for any start-up messages
-			Log = new CrashReportCommon.LogWriter("CrashReportProcess", LogFolder);
+			Log = new LogWriter("CrashReportProcess", LogFolder);
+
+			Config.LoadConfig();
+
+			Slack = new SlackWriter
+			{
+				WebhookUrl = Config.Default.SlackWebhookUrl,
+				Channel = Config.Default.SlackChannel,
+				Username = Config.Default.SlackUsername,
+				IconEmoji = Config.Default.SlackEmoji
+			};
+
+			StatusReporter = new StatusReporting();
 
 			// Add directory watchers
 			Watcher = new ReportWatcher();
 
-			Processor = new ReportProcessor( Watcher );
+			for (int ProcessorIndex = 0; ProcessorIndex < Config.Default.ProcessorThreadCount; ProcessorIndex++)
+			{
+				var Processor = new ReportProcessor(Watcher, ProcessorIndex);
+				Processors.Add(Processor);
+			}
 
-			WriteEvent( "Successfully started at " + DateTime.Now.ToString() );
+			StatusReporter.Start();
+			DateTime StartupTime = DateTime.UtcNow;
+			WriteEvent("Successfully started at " + StartupTime);
 		}
 
 		/// <summary>
@@ -120,12 +156,21 @@ namespace Tools.CrashReporter.CrashReportProcess
 		protected override void OnStop()
 		{
 			// Clean up the directory watcher and crash processor threads
-			Processor.Dispose();
-			Processor = null;
+			foreach (var Processor in Processors)
+			{
+				Processor.Dispose();
+			}
+			Processors.Clear();
 
 			Watcher.Dispose();
 			Watcher = null;
 
+			StatusReporter.Dispose();
+			StatusReporter = null;
+
+			Slack.Dispose();
+			Slack = null;
+
 			// Flush the log to disk
 			Log.Dispose();
 			Log = null;
diff --git a/Engine/Source/Programs/CrashReporter/CrashReportProcess/Properties/Settings.settings b/Engine/Source/Programs/CrashReporter/CrashReportProcess/Properties/Settings.settings
--- a/Engine/Source/Programs/CrashReporter/CrashReportProcess/Properties/Settings.settings
+++ b/Engine/Source/Programs/CrashReporter/CrashReportProcess/Properties/Settings.settings
@@ -1,45 +1,5 @@
 ﻿<?xml version='1.0' encoding='utf-8'?>
-<SettingsFile xmlns="http://schemas.microsoft.com/VisualStudio/2004/01/settings" CurrentProfile="(Default)" GeneratedClassNamespace="Tools.CrashReporter.CrashReportProcess.Properties" GeneratedClassName="Settings">
+<SettingsFile xmlns="http://schemas.microsoft.com/VisualStudio/2004/01/settings" CurrentProfile="(Default)">
   <Profiles />
-  <Settings>
-    <Setting Name="ProcessedReports" Type="System.String" Scope="Application">
-      <Value Profile="(Default)">E:\Sites\CrashReporterFiles</Value>
-    </Setting>
-    <Setting Name="DaysToSunsetReport" Type="System.Int32" Scope="Application">
-      <Value Profile="(Default)">14</Value>
-    </Setting>
-    <Setting Name="CrashReportWebSite" Type="System.String" Scope="Application">
-      <Value Profile="(Default)">CrashReporterUE4</Value>
-    </Setting>
-    <Setting Name="ProcessedVideos" Type="System.String" Scope="Application">
-      <Value Profile="(Default)">F:\CrashReporterVideos</Value>
-    </Setting>
-    <Setting Name="DepotRoot" Type="System.String" Scope="Application">
-      <Value Profile="(Default)">F:\depot</Value>
-    </Setting>
-    <Setting Name="InternalLandingZone" Type="System.String" Scope="Application">
-      <Value Profile="(Default)">F:\InternalCrashReports\CrashReportRepository</Value>
-    </Setting>
-    <Setting Name="ExternalLandingZone" Type="System.String" Scope="Application">
-      <Value Profile="(Default)">\\na-web-cluster.dmz.epicgames.net\crash_reporter\Rep</Value>
-    </Setting>
-    <Setting Name="InvalidReportsDirectory" Type="System.String" Scope="Application">
-      <Value Profile="(Default)">F:\InvalidCrashReports</Value>
-    </Setting>
-    <Setting Name="DepotIndex" Type="System.String" Scope="Application">
-      <Value Profile="(Default)">crash-01</Value>
-    </Setting>
-    <Setting Name="SyncBinariesFromDepot" Type="System.String" Scope="Application">
-      <Value Profile="(Default)">Stream\UE4\Main\Engine\Binaries\Win64\MinidumpDiagnostics*</Value>
-    </Setting>
-    <Setting Name="SyncConfigFromDepot" Type="System.String" Scope="Application">
-      <Value Profile="(Default)">Stream\UE4\Main\Engine\Config\...</Value>
-    </Setting>
-    <Setting Name="SyncThirdPartyFromDepot" Type="System.String" Scope="Application">
-      <Value Profile="(Default)">Stream\UE4\Main\Engine\Binaries\ThirdParty\OpenSSL\...</Value>
-    </Setting>
-    <Setting Name="MDDBinariesFolderInDepot" Type="System.String" Scope="Application">
-      <Value Profile="(Default)">Stream\UE4\Main\Engine\Binaries\Win64</Value>
-    </Setting>
-  </Settings>
+  <Settings />
 </SettingsFile>
\ No newline at end of file
diff --git a/Engine/Source/Programs/CrashReporter/CrashReportProcess/ReceiverReportQueue.cs b/Engine/Source/Programs/CrashReporter/CrashReportProcess/ReceiverReportQueue.cs
new file mode 100644
--- /dev/null
+++ b/Engine/Source/Programs/CrashReporter/CrashReportProcess/ReceiverReportQueue.cs
@@ -0,0 +1,537 @@
+﻿using System;
+using System.Collections.Generic;
+using System.IO;
+using System.Linq;
+using System.Runtime.InteropServices;
+using Tools.CrashReporter.CrashReportCommon;
+using Tools.DotNETCommon.XmlHandler;
+
+namespace Tools.CrashReporter.CrashReportProcess
+{
+	using ReportIdSet = HashSet<string>;
+	using System.Runtime.InteropServices;
+	using System.Globalization;
+
+	/// <summary>
+	/// A queue of pending reports in a specific folder
+	/// </summary>
+	class ReceiverReportQueue : ReportQueueBase
+	{
+		/// <summary>
+		/// Limit over which the we don't attempt to refresh the queue from disk until we've processed and shrunk the queue.
+		/// Should save cycles when busy.
+		/// </summary>
+		private const int SkipQueueRefreshMin = 50;
+
+		/// <summary>
+		/// Limit over which the we don't add to the queue
+		/// </summary>
+		private const int QueueMax = 250;
+
+		/// <summary>
+		/// Constructor taking the landing zone
+		/// </summary>
+		public ReceiverReportQueue(string InQueueName, string LandingZonePath)
+		{
+			QueueName = InQueueName;
+			LandingZone = LandingZonePath;
+		}
+
+		/// <summary>
+		/// Look for new report folders and add them to the publicly available thread-safe queue.
+		/// </summary>
+		public override int CheckForNewReports()
+		{
+			try
+			{
+				if (QueueCount >= SkipQueueRefreshMin)
+				{
+					CrashReporterProcessServicer.WriteEvent(string.Format("CheckForNewReports: skipped busy queue size {0} in {1}", QueueCount, LandingZone));
+					return LastQueueSizeOnDisk;
+				}
+
+				var NewReportPath = "";
+				var ReportsInLandingZone = new ReportIdSet();
+
+				// Check the landing zones for new reports
+				DirectoryInfo DirInfo = new DirectoryInfo(LandingZone);
+
+				// Couldn't find a landing zone, skip and try again later.
+				// Crash receiver will recreate them if needed.
+				if (!DirInfo.Exists)
+				{
+					CrashReporterProcessServicer.WriteFailure("LandingZone not found: " + LandingZone);
+					return LastQueueSizeOnDisk;
+				}
+
+				var DirectoriesInLandingZone = DirInfo.GetDirectories().OrderBy(dirinfo => dirinfo.CreationTimeUtc).ToArray();
+				LastQueueSizeOnDisk = DirectoriesInLandingZone.Length;
+
+				int EnqueuedCount = 0;
+				CrashReporterProcessServicer.WriteEvent(string.Format("CheckForNewReports: {0} reports in {1}", DirectoriesInLandingZone.Length, LandingZone));
+
+				// Add any new reports
+				for (int DirIndex = 0; DirIndex < DirectoriesInLandingZone.Length && QueueCount < QueueMax; DirIndex++)
+				{
+					var SubDirInfo = DirectoriesInLandingZone[DirIndex];
+					try
+					{
+						if (Directory.Exists(SubDirInfo.FullName))
+						{
+							NewReportPath = SubDirInfo.FullName;
+							ReportsInLandingZone.Add(NewReportPath);
+							if (!ReportsInLandingZoneLastTimeWeChecked.Contains(NewReportPath))
+							{
+								EnqueueNewReport(NewReportPath);
+								EnqueuedCount++;
+							}
+						}
+					}
+					catch (System.Exception Ex)
+					{
+						CrashReporterProcessServicer.WriteException("CheckForNewReportsInner: " + Ex.ToString());
+						ReportProcessor.MoveReportToInvalid(NewReportPath, "NEWRECORD_FAIL_" + DateTime.Now.Ticks);
+					}
+				}
+				//CrashReporterProcessServicer.WriteEvent( string.Format( "ReportsInLandingZone={0}, ReportsInLandingZoneLastTimeWeChecked={1}", ReportsInLandingZone.Count, ReportsInLandingZoneLastTimeWeChecked.Count ) );
+				ReportsInLandingZoneLastTimeWeChecked = ReportsInLandingZone;
+				CrashReporterProcessServicer.WriteEvent(string.Format("CheckForNewReports: {0} enqueued to queue size {1} from {2}", EnqueuedCount, QueueCount, LandingZone));
+				CrashReporterProcessServicer.WriteEvent(string.Format("CheckForNewReports: Enqueue rate {0:N1}/minute from {1}", EnqueueCounter.EventsPerSecond * 60, LandingZone));
+			}
+			catch (Exception Ex)
+			{
+				CrashReporterProcessServicer.WriteException("CheckForNewReportsOuter: " + Ex.ToString());
+			}
+
+			return LastQueueSizeOnDisk;
+		}
+
+		/// <summary>
+		/// Tidy up the landing zone folder
+		/// </summary>
+		public override void CleanLandingZone()
+		{
+			var DirInfo = new DirectoryInfo(LandingZone);
+			foreach (var SubDirInfo in DirInfo.GetDirectories())
+			{
+				if (SubDirInfo.LastWriteTimeUtc.AddDays(Config.Default.DeleteWaitingReportsDays) < DateTime.UtcNow)
+				{
+					SubDirInfo.Delete(true);
+				}
+			}
+		}
+
+		/// <summary> Tries to dequeue a report from the list. </summary>
+		public override bool TryDequeueReport(out FGenericCrashContext Context)
+		{
+			lock (NewReportsLock)
+			{
+				if (NewCrashContexts.Count > 0)
+				{
+					Context = NewCrashContexts.Dequeue();
+					DequeueCounter.AddEvent();
+					CrashReporterProcessServicer.StatusReporter.IncrementCount(StatusReportingConstants.ProcessingStartedEvent);
+					CrashReporterProcessServicer.WriteEvent(string.Format("- Dequeued: {0:N1}/minute BuiltFromCL={1,7} Path={2}", DequeueCounter.EventsPerSecond * 60, Context.PrimaryCrashProperties.EngineVersion, Context.CrashDirectory));
+					return true;
+				}
+			}
+
+			Context = null;
+			return false;
+		}
+
+		public override string QueueId
+		{
+			get
+			{
+				return QueueName;
+			}
+		}
+
+		public override string LandingZonePath
+		{
+			get
+			{
+				return LandingZone;
+			}
+		}
+
+		/// <summary> Looks for the CrashContext.runtime-xml, if found, will return a new instance of the FGenericCrashContext. </summary>
+		private FGenericCrashContext FindCrashContext(string NewReportPath)
+		{
+			string CrashContextPath = Path.Combine(NewReportPath, FGenericCrashContext.CrashContextRuntimeXMLName);
+			bool bExist = File.Exists(CrashContextPath);
+
+			if (bExist)
+			{
+				return FGenericCrashContext.FromFile(CrashContextPath);
+			}
+
+			return null;
+		}
+
+		/// <summary> Converts WER metadata xml file to the crash context. </summary>
+		private void ConvertMetadataToCrashContext(WERReportMetadata Metadata, string NewReportPath, ref FGenericCrashContext OutCrashContext)
+		{
+			if (OutCrashContext == null)
+			{
+				OutCrashContext = new FGenericCrashContext();
+			}
+
+			FReportData ReportData = new FReportData(Metadata, NewReportPath);
+
+			OutCrashContext.PrimaryCrashProperties.CrashVersion = (int)ECrashDescVersions.VER_1_NewCrashFormat;
+
+			//OutCrashContext.PrimaryCrashProperties.ProcessId = 0; don't overwrite valid ids, zero is default anyway
+			OutCrashContext.PrimaryCrashProperties.CrashGUID = new DirectoryInfo(NewReportPath).Name;
+			// 			OutCrashContext.PrimaryCrashProperties.IsInternalBuild
+			// 			OutCrashContext.PrimaryCrashProperties.IsPerforceBuild
+			// 			OutCrashContext.PrimaryCrashProperties.IsSourceDistribution
+			// 			OutCrashContext.PrimaryCrashProperties.SecondsSinceStart
+			OutCrashContext.PrimaryCrashProperties.GameName = ReportData.GameName;
+			// 			OutCrashContext.PrimaryCrashProperties.ExecutableName
+			// 			OutCrashContext.PrimaryCrashProperties.BuildConfiguration
+			// 			OutCrashContext.PrimaryCrashProperties.PlatformName
+			// 			OutCrashContext.PrimaryCrashProperties.PlatformNameIni
+			OutCrashContext.PrimaryCrashProperties.PlatformFullName = ReportData.Platform;
+			OutCrashContext.PrimaryCrashProperties.EngineMode = ReportData.EngineMode;
+			OutCrashContext.PrimaryCrashProperties.EngineVersion = ReportData.GetEngineVersion();
+			OutCrashContext.PrimaryCrashProperties.CommandLine = ReportData.CommandLine;
+			// 			OutCrashContext.PrimaryCrashProperties.LanguageLCID
+
+			// Create a locate and get the language.
+			int LanguageCode = 0;
+			int.TryParse(ReportData.Language, out LanguageCode);
+			try
+			{
+				if (LanguageCode > 0)
+				{
+					CultureInfo LanguageCI = new CultureInfo(LanguageCode);
+					OutCrashContext.PrimaryCrashProperties.AppDefaultLocale = LanguageCI.Name;
+				}
+			}
+			catch (Exception)
+			{
+				OutCrashContext.PrimaryCrashProperties.AppDefaultLocale = null;
+			}
+
+			if (string.IsNullOrEmpty(OutCrashContext.PrimaryCrashProperties.AppDefaultLocale))
+			{
+				// Default to en-US
+				OutCrashContext.PrimaryCrashProperties.AppDefaultLocale = "en-US";
+			}
+
+			// 			OutCrashContext.PrimaryCrashProperties.IsUE4Release
+			string WERUserName = ReportData.UserName;
+			if (!string.IsNullOrEmpty(WERUserName) || string.IsNullOrEmpty(OutCrashContext.PrimaryCrashProperties.UserName))
+			{
+				OutCrashContext.PrimaryCrashProperties.UserName = WERUserName;
+			}
+
+			OutCrashContext.PrimaryCrashProperties.BaseDir = ReportData.BaseDir;
+			// 			OutCrashContext.PrimaryCrashProperties.RootDir
+			OutCrashContext.PrimaryCrashProperties.MachineId = ReportData.MachineId;
+			OutCrashContext.PrimaryCrashProperties.EpicAccountId = ReportData.EpicAccountId;
+			// 			OutCrashContext.PrimaryCrashProperties.CallStack
+			// 			OutCrashContext.PrimaryCrashProperties.SourceContext
+			OutCrashContext.PrimaryCrashProperties.UserDescription = string.Join("\n", ReportData.UserDescription);
+
+			if (string.IsNullOrEmpty(OutCrashContext.PrimaryCrashProperties.ErrorMessage))
+			{
+				OutCrashContext.PrimaryCrashProperties.ErrorMessage = string.Join("\n", ReportData.ErrorMessage);
+			}
+
+			// 			OutCrashContext.PrimaryCrashProperties.CrashDumpMode
+			// 			OutCrashContext.PrimaryCrashProperties.Misc.NumberOfCores
+			// 			OutCrashContext.PrimaryCrashProperties.Misc.NumberOfCoresIncludingHyperthreads
+			// 			OutCrashContext.PrimaryCrashProperties.Misc.Is64bitOperatingSystem
+			// 			OutCrashContext.PrimaryCrashProperties.Misc.CPUVendor
+			// 			OutCrashContext.PrimaryCrashProperties.Misc.CPUBrand
+			// 			OutCrashContext.PrimaryCrashProperties.Misc.PrimaryGPUBrand
+			// 			OutCrashContext.PrimaryCrashProperties.Misc.OSVersionMajor
+			// 			OutCrashContext.PrimaryCrashProperties.Misc.OSVersionMinor
+			// 			OutCrashContext.PrimaryCrashProperties.Misc.AppDiskTotalNumberOfBytes
+			// 			OutCrashContext.PrimaryCrashProperties.Misc.AppDiskNumberOfFreeBytes
+			// 			OutCrashContext.PrimaryCrashProperties.MemoryStats.TotalPhysical
+			// 			OutCrashContext.PrimaryCrashProperties.MemoryStats.TotalVirtual
+			// 			OutCrashContext.PrimaryCrashProperties.MemoryStats.PageSize
+			// 			OutCrashContext.PrimaryCrashProperties.MemoryStats.TotalPhysicalGB
+			// 			OutCrashContext.PrimaryCrashProperties.MemoryStats.AvailablePhysical
+			// 			OutCrashContext.PrimaryCrashProperties.MemoryStats.AvailableVirtual
+			// 			OutCrashContext.PrimaryCrashProperties.MemoryStats.UsedPhysical
+			// 			OutCrashContext.PrimaryCrashProperties.MemoryStats.PeakUsedPhysical
+			// 			OutCrashContext.PrimaryCrashProperties.MemoryStats.UsedVirtual
+			// 			OutCrashContext.PrimaryCrashProperties.MemoryStats.PeakUsedVirtual
+			// 			OutCrashContext.PrimaryCrashProperties.MemoryStats.bIsOOM
+			// 			OutCrashContext.PrimaryCrashProperties.MemoryStats.OOMAllocationSize
+			// 			OutCrashContext.PrimaryCrashProperties.MemoryStats.OOMAllocationAlignment
+			OutCrashContext.PrimaryCrashProperties.TimeofCrash = new DateTime(ReportData.Ticks);
+			OutCrashContext.PrimaryCrashProperties.bAllowToBeContacted = ReportData.AllowToBeContacted;
+		}
+
+		/// <summary> Looks for the WER metadata xml file, if found, will return a new instance of the WERReportMetadata. </summary>
+		private WERReportMetadata FindMetadata(string NewReportPath)
+		{
+			WERReportMetadata MetaData = null;
+
+			// Just to verify that the report is still there.
+			DirectoryInfo DirInfo = new DirectoryInfo(NewReportPath);
+			if (!DirInfo.Exists)
+			{
+				CrashReporterProcessServicer.WriteFailure("FindMetadata: Directory not found " + NewReportPath);
+			}
+			else
+			{
+				// Check to see if we wish to suppress processing of this report.	
+				foreach (var Info in DirInfo.GetFiles("*.xml"))
+				{
+					var MetaDataToCheck = XmlHandler.ReadXml<WERReportMetadata>(Info.FullName);
+					if (CheckMetaData(MetaDataToCheck))
+					{
+						MetaData = MetaDataToCheck;
+						break;
+					}
+				}
+			}
+
+			return MetaData;
+		}
+
+		/// <summary> Looks for the Diagnostics.txt file and returns true if exists. </summary>
+		private bool FindDiagnostics(string NewReportPath)
+		{
+			bool bFound = false;
+
+			DirectoryInfo DirInfo = new DirectoryInfo(NewReportPath);
+			foreach (var Info in DirInfo.GetFiles(CrashReporterConstants.DiagnosticsFileName))
+			{
+				bFound = true;
+				break;
+			}
+			return bFound;
+		}
+
+		/// <summary>
+		/// Optionally don't process some reports based on the Windows Error report meta data.
+		/// </summary>
+		/// <param name="WERData">The Windows Error Report meta data.</param>
+		/// <returns>false to reject the report.</returns>
+		private bool CheckMetaData(WERReportMetadata WERData)
+		{
+			if (WERData == null)
+			{
+				return false;
+			}
+
+			// Reject any crashes with the invalid metadata.
+			if (WERData.ProblemSignatures == null || WERData.DynamicSignatures == null || WERData.OSVersionInformation == null)
+			{
+				return false;
+			}
+
+			// Reject any crashes from the minidump processor.
+			if (WERData.ProblemSignatures.Parameter0 != null && WERData.ProblemSignatures.Parameter0.ToLower() == "MinidumpDiagnostics".ToLower())
+			{
+				return false;
+			}
+
+			return true;
+		}
+
+		// From CrashUpload.cpp
+		/*
+		struct FCompressedCrashFile : FNoncopyable
+		{
+			int32 CurrentFileIndex; // 4 bytes for file index
+			FString Filename; // 4 bytes for length + 260 bytes for char data
+			TArray<uint8> Filedata; // 4 bytes for length + N bytes for data
+		}
+		*/
+
+		/// <summary> 
+		/// Decompresses a compressed crash report. 
+		/// </summary>
+		unsafe private bool DecompressReport(string CompressedReportPath, FCompressedCrashInformation Meta)
+		{
+			string ReportPath = Path.GetDirectoryName(CompressedReportPath);
+			using (FileStream FileStream = File.OpenRead(CompressedReportPath))
+			{
+				Int32 UncompressedSize = Meta.GetUncompressedSize();
+				Int32 CompressedSize = Meta.GetCompressedSize();
+
+				if (FileStream.Length != CompressedSize)
+				{
+					return false;
+				}
+
+				byte[] UncompressedBufferArray = new byte[UncompressedSize];
+				byte[] CompressedBufferArray = new byte[CompressedSize];
+
+				FileStream.Read(CompressedBufferArray, 0, CompressedSize);
+
+				fixed (byte* UncompressedBufferPtr = UncompressedBufferArray, CompressedBufferPtr = CompressedBufferArray)
+				{
+					bool bResult = NativeMethods.UE4UncompressMemoryZLIB((IntPtr)UncompressedBufferPtr, UncompressedSize, (IntPtr)CompressedBufferPtr, CompressedSize);
+					if (!bResult)
+					{
+						CrashReporterProcessServicer.WriteFailure("! ZLibFail: Path=" + ReportPath);
+						return false;
+					}
+				}
+
+				MemoryStream MemoryStream = new MemoryStream(UncompressedBufferArray, false);
+				BinaryReader BinaryData = new BinaryReader(MemoryStream);
+
+				for (int FileIndex = 0; FileIndex < Meta.GetNumberOfFiles(); FileIndex++)
+				{
+					Int32 CurrentFileIndex = BinaryData.ReadInt32();
+					if (CurrentFileIndex != FileIndex)
+					{
+						CrashReporterProcessServicer.WriteFailure("! ReadFail: Required=" + FileIndex + " Read=" + CurrentFileIndex);
+						return false;
+					}
+
+					string Filename = FBinaryReaderHelper.ReadFixedSizeString(BinaryData);
+					Int32 FiledataLength = BinaryData.ReadInt32();
+					byte[] Filedata = BinaryData.ReadBytes(FiledataLength);
+
+					File.WriteAllBytes(Path.Combine(ReportPath, Filename), Filedata);
+				}
+			}
+
+			return true;
+		}
+
+		/// <summary> Enqueues a new crash. </summary>
+		private void EnqueueNewReport(string NewReportPath)
+		{
+			string ReportName = Path.GetFileName(NewReportPath);
+
+			string CompressedReportPath = Path.Combine(NewReportPath, ReportName + ".ue4crash");
+			string MetadataPath = Path.Combine(NewReportPath, ReportName + ".xml");
+			bool bIsCompressed = File.Exists(CompressedReportPath) && File.Exists(MetadataPath);
+			if (bIsCompressed)
+			{
+				FCompressedCrashInformation CompressedCrashInformation = XmlHandler.ReadXml<FCompressedCrashInformation>(MetadataPath);
+				bool bResult = DecompressReport(CompressedReportPath, CompressedCrashInformation);
+				if (!bResult)
+				{
+					ReportProcessor.CleanReport(new DirectoryInfo(NewReportPath));
+					ReportProcessor.MoveReportToInvalid(NewReportPath, "DECOMPRESSION_FAIL_" + DateTime.Now.Ticks + "_" + ReportName);
+					return;
+				}
+				else
+				{
+					// Rename metadata file to not interfere with the WERReportMetadata.
+					File.Move(MetadataPath, Path.ChangeExtension(MetadataPath, "processed_xml"));
+				}
+			}
+
+			// Unified crash reporting
+			FGenericCrashContext GenericContext = FindCrashContext(NewReportPath);
+			FGenericCrashContext Context = GenericContext;
+
+			bool bFromWER = false;
+			if (Context == null || !Context.HasProcessedData())
+			{
+				WERReportMetadata MetaData = FindMetadata(NewReportPath);
+				if (MetaData != null)
+				{
+					FReportData ReportData = new FReportData(MetaData, NewReportPath);
+					ConvertMetadataToCrashContext(MetaData, NewReportPath, ref Context);
+					bFromWER = true;
+				}
+			}
+
+			if (Context == null)
+			{
+				CrashReporterProcessServicer.WriteFailure("! NoCntx  : Path=" + NewReportPath);
+				ReportProcessor.CleanReport(new DirectoryInfo(NewReportPath));
+			}
+			else
+			{
+				Context.CrashDirectory = NewReportPath;
+				Context.PrimaryCrashProperties.SetPlatformFullName();
+
+				// Added data from WER, save to the crash context file.
+				if (bFromWER)
+				{
+					Context.ToFile();
+				}
+
+				FEngineVersion EngineVersion = new FEngineVersion(Context.PrimaryCrashProperties.EngineVersion);
+
+				uint BuiltFromCL = EngineVersion.Changelist;
+
+				string BranchName = EngineVersion.Branch;
+				if (string.IsNullOrEmpty(BranchName))
+				{
+					CrashReporterProcessServicer.WriteEvent("% Warning NoBranch: BuiltFromCL=" + string.Format("{0,7}", BuiltFromCL) + " Path=" + NewReportPath + " EngineVersion=" + Context.PrimaryCrashProperties.EngineVersion);
+					Context.PrimaryCrashProperties.ProcessorFailedMessage = "Engine version has no branch name. EngineVersion=" + Context.PrimaryCrashProperties.EngineVersion;
+				}
+				else if (BranchName.Equals(CrashReporterConstants.LicenseeBranchName, StringComparison.InvariantCultureIgnoreCase))
+				{
+					CrashReporterProcessServicer.WriteEvent("% Warning branch is UE4-QA  : BuiltFromCL=" + string.Format("{0,7}", BuiltFromCL) + " Path=" + NewReportPath);
+					Context.PrimaryCrashProperties.ProcessorFailedMessage = "Branch was the forbidden LicenseeBranchName=" + BranchName;
+				}
+
+				// Look for the Diagnostics.txt, if we have a diagnostics file we can continue event if the CL is marked as broken.
+				bool bHasDiagnostics = FindDiagnostics(NewReportPath);
+
+				if (BuiltFromCL == 0 && (!bHasDiagnostics && !Context.HasProcessedData()))
+				{
+					// For now ignore all locally made crashes.
+					CrashReporterProcessServicer.WriteEvent("% Warning CL=0 and no useful data : BuiltFromCL=" + string.Format("{0,7}", BuiltFromCL) + " Path=" + NewReportPath);
+					Context.PrimaryCrashProperties.ProcessorFailedMessage = "Report from CL=0 has no diagnostics, callstack or error msg";
+				}
+
+
+				lock (NewReportsLock)
+				{
+					NewCrashContexts.Enqueue(Context);
+				}
+				EnqueueCounter.AddEvent();
+				CrashReporterProcessServicer.StatusReporter.IncrementCount(StatusReportingConstants.QueuedEvent);
+				CrashReporterProcessServicer.WriteEvent("+ Enqueued: BuiltFromCL=" + string.Format("{0,7}", BuiltFromCL) + " Path=" + NewReportPath);
+			}
+
+		}
+
+		private int QueueCount
+		{
+			get
+			{
+				lock (NewReportsLock)
+				{
+					return NewCrashContexts.Count;
+				}			
+			}
+		}
+
+		/// <summary> A list of freshly landed crash reports. </summary>
+		private readonly Queue<FGenericCrashContext> NewCrashContexts = new Queue<FGenericCrashContext>();
+
+		/// <summary> Object used to synchronize the access to the NewReport. </summary>
+		private readonly Object NewReportsLock = new Object();
+
+		/// <summary>Friendly queue name</summary>
+		private readonly string QueueName;
+
+		/// <summary>Incoming report path</summary>
+		private readonly string LandingZone;
+
+		/// <summary>Report IDs that were in the folder on the last pass</summary>
+		private ReportIdSet ReportsInLandingZoneLastTimeWeChecked = new ReportIdSet();
+
+		private readonly EventCounter EnqueueCounter = new EventCounter(TimeSpan.FromMinutes(90), 20);
+		private readonly EventCounter DequeueCounter = new EventCounter(TimeSpan.FromMinutes(10), 10);
+
+		private int LastQueueSizeOnDisk = 0;
+	}
+}
+
+static class NativeMethods
+{
+	[DllImport("CompressionHelper.dll", CallingConvention = CallingConvention.Cdecl)]
+	public extern static bool UE4UncompressMemoryZLIB( /*void**/IntPtr UncompressedBuffer, Int32 UncompressedSize, /*const void**/IntPtr CompressedBuffer, Int32 CompressedSize);	
+}
\ No newline at end of file
diff --git a/Engine/Source/Programs/CrashReporter/CrashReportProcess/ReportProcessor.cs b/Engine/Source/Programs/CrashReporter/CrashReportProcess/ReportProcessor.cs
--- a/Engine/Source/Programs/CrashReporter/CrashReportProcess/ReportProcessor.cs
+++ b/Engine/Source/Programs/CrashReporter/CrashReportProcess/ReportProcessor.cs
@@ -21,13 +21,13 @@ namespace Tools.CrashReporter.CrashReportProcess
 	/// <summary>
 	/// A class to handle processing of received crash reports.
 	/// </summary>
-	public sealed class ReportProcessor : IDisposable
+	sealed class ReportProcessor : IDisposable
 	{
 		/// <summary>Absolute path to Perforce (default installation assumed)</summary>
 		const string PerforceExePath = "C:\\Program Files\\Perforce\\p4.exe";
 
 		/// <summary>Give up on call stack analysis after 30 minutes</summary>
-		const int MinidumpDiagnosticsTimeoutSeconds = 30 * 60;
+		private const int MinidumpDiagnosticsTimeoutMilliseconds = 1000 * 60 * 30;
 
 		/// <summary>Give up on Perforce syncs after 2 minutes</summary>
 		const int SyncTimeoutSeconds = 2 * 60;
@@ -35,32 +35,59 @@ namespace Tools.CrashReporter.CrashReportProcess
 		/// <summary> Number of all processed reports. </summary>
 		public static int ProcessedReports = 1;
 
-		/// <summary> Number of processed reports that have been added to the website. </summary>
-		public static int WebAddedReports = 0;
-
 		/// <summary> Global timer used to measure web added reports per day. </summary>
 		public static Stopwatch Timer = Stopwatch.StartNew();
 
+		private static EventCounter WebAddCounter = new EventCounter(TimeSpan.FromMinutes(10), 10);
+
+		private static DateTime LastCleanupTime = DateTime.Now;
+
+		private static DateTime LoggingDate = DateTime.UtcNow.Date;
+
+		private static Object TickStaticLock = new Object();
+
+		private static Object MinidumpDiagnosticsLock = new Object();
+
+		private static Object FailedUploadLock = new Object();
+		private static int ConsecutiveFailedUploads = 0;
+
 		/// <summary>The thread that handles detection of new crashes.</summary>
 		public ReportWatcher Watcher = null;
 
+		private readonly int ProcessorIndex;
+
 		/// <summary>Task to handle processing.</summary>
 		Task ProcessorTask;
 
 		/// <summary>Object to allow the task to be cancelled at shut-down.</summary>
 		CancellationTokenSource CancelSource;
 
-		/// <summary> Task to handle adding a new record to the database. </summary>
-		Task AddReportTask;
+		/// <summary> Tasks to handle adding a new record to the database. </summary>
+		readonly List<Task> AddReportTasks = new List<Task>();
 
 		/// <summary>
-		/// Global initialisation of the processor.
+		/// Static init for one-time init jobs
+		/// </summary>
+		static ReportProcessor()
+		{
+#if !DEBUG
+			if (!SyncRequiredFiles())
+			{
+				CrashReporterProcessServicer.WriteFailure("ReportProcessor: failed to sync files from Perforce");
+			}
+#endif			
+		}
+
+		/// <summary>
+		/// Global initialization of the processor.
 		/// </summary>
 		/// <param name="InWatcher">The object that watches for new crash reports coming in.</param>
-		public ReportProcessor( ReportWatcher InWatcher )
+		/// <param name="InProcessorIndex">The index of this processor in the application's thread list</param>
+		public ReportProcessor( ReportWatcher InWatcher, int InProcessorIndex )
 		{
 			// Create the thread to handle processing
 			Watcher = InWatcher;
+			ProcessorIndex = InProcessorIndex;
 			CancelSource = new CancellationTokenSource();
 			ProcessNewReports();
 		}
@@ -82,17 +109,6 @@ namespace Tools.CrashReporter.CrashReportProcess
 		/// <remarks>All exceptions are caught and written to the event log.</remarks>
 		private void ProcessNewReports()
 		{
-#if !DEBUG
-			if (!SyncRequiredFiles())
-			{
-				return;
-			}
-#endif
-
-			DateTime LastCleanupTime = DateTime.Now;
-
-			int LastDay = DateTime.UtcNow.Day;
-
 			var Cancel = CancelSource.Token;
 			ProcessorTask = Task.Factory.StartNew(() =>
 			{
@@ -100,6 +116,8 @@ namespace Tools.CrashReporter.CrashReportProcess
 				{
 					try
 					{
+						bool bIdle = true;
+
 						foreach( var Queue in Watcher.ReportQueues )
 						{
 							FGenericCrashContext NewContext = null;
@@ -107,33 +125,24 @@ namespace Tools.CrashReporter.CrashReportProcess
 							{
 								ProcessReport( NewContext );
 
-								int CurrentDay = DateTime.UtcNow.Day;
-								if( CurrentDay > LastDay )
-								{
-									// Check the log and create a new one for a new day.
-									CrashReporterProcessServicer.Log.CreateNewLogFile();
-									LastDay = CurrentDay;
-								}
-
 								// The effect of this break is to prioritize ReportQueues by their order in the list, from highest to lowest
+								bIdle = false;
 								break;
 							}
 						}
+
+						if (bIdle)
+						{
+							// Don't use the CPU if we don't need.
+							Thread.Sleep(1000);							
+						}
 					}
 					catch (Exception Ex)
 					{
-						CrashReporterProcessServicer.WriteException( "ProcessNewReports: " + Ex.ToString() );
+						CrashReporterProcessServicer.WriteException(string.Format("PROC-{0} ", ProcessorIndex) + "ProcessNewReports: " + Ex.ToString());
 					}
 
-					// Remove all folders older than n days to prevent unbounded growth
-					if ((DateTime.Now - LastCleanupTime) > TimeSpan.FromMinutes(15))
-					{
-						CleanRepository();
-						LastCleanupTime = DateTime.Now;
-					}
-
-					// Don't use the CPU if we don't need.
-					Thread.Sleep( 100 );
+					TickStatic(Watcher);
 				}
 			});
 		}
@@ -206,11 +215,13 @@ namespace Tools.CrashReporter.CrashReportProcess
 		{
 			try
 			{
+				CrashReporterProcessServicer.StatusReporter.IncrementCount(StatusReportingConstants.ProcessingFailedEvent);
+
 				DirectoryInfo DirInfo = new DirectoryInfo( ReportName );
 				// Rename the report directory, so we should be able to quickly find the invalid reports.
-				Directory.CreateDirectory( Properties.Settings.Default.InvalidReportsDirectory );
+				Directory.CreateDirectory( Config.Default.InvalidReportsDirectory );
 				string CleanFilename = String.Concat( ReportNameAsFilename.Split( Path.GetInvalidFileNameChars() ) );
-				string DestinationDirectory = Path.Combine( Properties.Settings.Default.InvalidReportsDirectory, CleanFilename );
+				string DestinationDirectory = Path.Combine( Config.Default.InvalidReportsDirectory, CleanFilename );
 
 				Directory.CreateDirectory( DestinationDirectory );
 
@@ -243,17 +254,17 @@ namespace Tools.CrashReporter.CrashReportProcess
 		/// <summary> 
 		/// Delete report folders older than a certain age to avoid unbounded growth of the crash repository.
 		/// </summary>
-		/// <remarks>The folders for the deduplication process older than the property 'DaysToSunsetReport' days old are deleted.</remarks>
-		private void CleanRepository()
+		/// <remarks>The folders for the deduplication process older than the property 'DeleteWaitingReportsDays' days old are deleted.</remarks>
+		private static void CleanRepository(ReportWatcher InWatcher)
 		{
 			try
 			{
-				foreach (var Queue in Watcher.ReportQueues)
+				foreach (var Queue in InWatcher.ReportQueues)
 				{
 					Queue.CleanLandingZone();
 				}
 
-				CrashReporterProcessServicer.Log.CleanOutOldLogs(Properties.Settings.Default.DaysToSunsetReport);
+				CrashReporterProcessServicer.Log.CleanOutOldLogs(Config.Default.DeleteWaitingReportsDays);
 			}
 			catch( Exception Ex )
 			{
@@ -338,16 +349,15 @@ namespace Tools.CrashReporter.CrashReportProcess
 				// Ignore any callstack that is shorter than expected, usually the callstack is invalid.
 				if( NewCrash.CallStack.Length <= CrashReporterConstants.MinCallstackDepth )
 				{
-					CrashReporterProcessServicer.WriteFailure( "! BadStack: BuiltFromCL=" + string.Format( "{0,7}", NewContext.PrimaryCrashProperties.EngineVersion ) + " Path=" + NewContext.CrashDirectory );
-				}
-				else
-				{
-					XmlPayload = XmlHandler.ToXmlString<CrashDescription>( NewCrash );
+					CrashReporterProcessServicer.WriteEvent(string.Format("PROC-{0} ", ProcessorIndex) + "! BadStack: BuiltFromCL=" + string.Format( "{0,7}", NewContext.PrimaryCrashProperties.EngineVersion ) + " Path=" + NewContext.CrashDirectory );
+					NewContext.PrimaryCrashProperties.ProcessorFailedMessage = string.Format("Callstack was too small. {0} lines (minimum {1})", NewCrash.CallStack.Length, CrashReporterConstants.MinCallstackDepth);
 				}
+				
+				XmlPayload = XmlHandler.ToXmlString<CrashDescription>( NewCrash );
 			}
 			catch (Exception Ex)
 			{
-				CrashReporterProcessServicer.WriteException( "CreateCrash: " + Ex.ToString() );
+				CrashReporterProcessServicer.WriteException(string.Format("PROC-{0} ", ProcessorIndex) + "CreateCrash: " + Ex.ToString());
 			}
 
 			return XmlPayload;
@@ -365,14 +375,14 @@ namespace Tools.CrashReporter.CrashReportProcess
 			try
 			{
 				// Simple suppression by blanking out the URL for local testing
-				if( Properties.Settings.Default.CrashReportWebSite.Length > 0 )
+				if( Config.Default.CrashReportWebSite.Length > 0 )
 				{
 
 					bool bDebug = false;
 					string RequestString;
 					if( !bDebug )
 					{
-						RequestString = "http://" + Properties.Settings.Default.CrashReportWebSite + ":80/Crashes/AddCrash/-1";
+						RequestString = "http://" + Config.Default.CrashReportWebSite + ":80/Crashes/AddCrash/-1";
 					}
 					else
 					{ 
@@ -400,13 +410,13 @@ namespace Tools.CrashReporter.CrashReportProcess
 
 					if (NewID == -1)
 					{
-						CrashReporterProcessServicer.WriteFailure("UploadCrash: " + ErrorMessage);
+						CrashReporterProcessServicer.WriteFailure(string.Format("PROC-{0} ", ProcessorIndex) + "UploadCrash: " + ErrorMessage);
 					}
 				}
 			}
 			catch( Exception Ex )
 			{
-				CrashReporterProcessServicer.WriteException( "UploadCrash: " + Ex.ToString() );
+				CrashReporterProcessServicer.WriteException(string.Format("PROC-{0} ", ProcessorIndex) + "UploadCrash: " + Ex.ToString());
 			}
 
 			return NewID;
@@ -431,24 +441,38 @@ namespace Tools.CrashReporter.CrashReportProcess
 				string CrashDetails = CreateCrash( DirInfo, NewContext, VideoFileName != null, LogFileName != null, DumpFileName != null );
 				if (CrashDetails == "")
 				{
-					CrashReporterProcessServicer.WriteFailure( "! NoDetail: Path=" + NewContext.CrashDirectory );
+					CrashReporterProcessServicer.WriteFailure(string.Format("PROC-{0} ", ProcessorIndex) + "! CreateCrash no payload: Path=" + NewContext.CrashDirectory);
 					return false;
 				}
 
 				// Upload the crash to the database, and retrieve the new row id
 				int ReportID = UploadCrash(CrashDetails);
-				if (ReportID <= 0)
+
+				lock (FailedUploadLock)
 				{
-					CrashReporterProcessServicer.WriteFailure( "! NoUpload: Path=" + NewContext.CrashDirectory );
-					string PayloadFailedFileName = Path.Combine(DirInfo.FullName, "PayloadFailed.xml");
-					File.WriteAllText(PayloadFailedFileName, CrashDetails);
-					return false;
+					if (ReportID <= 0)
+					{
+						// Upload failed
+						ConsecutiveFailedUploads++;
+						if (ConsecutiveFailedUploads == Config.Default.ConsecutiveFailedWebAddLimit)
+						{
+							CrashReporterProcessServicer.StatusReporter.Alert("Cannot contact Crash Report website.");
+							CrashReporterProcessServicer.WriteFailure("Cannot contact Crash Report website.");
+						}
+						CrashReporterProcessServicer.WriteFailure(string.Format("PROC-{0} ", ProcessorIndex) + "! NoUpload: Path=" + NewContext.CrashDirectory);
+						string PayloadFailedFileName = Path.Combine(DirInfo.FullName, "PayloadFailed.xml");
+						File.WriteAllText(PayloadFailedFileName, CrashDetails);
+						return false;
+					}
+
+					// Upload succeeded
+					ConsecutiveFailedUploads = 0;
 				}
 
 				string IDThenUnderscore = string.Format("{0}_", ReportID);
 
 				// Use the row id to name and move the files the way the web site requires
-				string DestinationFolder = Path.Combine(Properties.Settings.Default.ProcessedReports, IDThenUnderscore);
+				string DestinationFolder = Path.Combine(Config.Default.ProcessedReports, IDThenUnderscore);
 
 				// Move report files to crash reporter file store
 				if (LogFileName != null)
@@ -495,7 +519,7 @@ namespace Tools.CrashReporter.CrashReportProcess
 				// Move the video (if it exists) to an alternate store
 				if (VideoFileName != null)
 				{
-					DestinationFolder = Path.Combine(Properties.Settings.Default.ProcessedVideos, IDThenUnderscore);
+					DestinationFolder = Path.Combine(Config.Default.ProcessedVideos, IDThenUnderscore);
 
 					VideoFileName = Path.Combine(DirInfo.FullName, VideoFileName);
 					FileInfo VideoInfo = new FileInfo(VideoFileName);
@@ -505,20 +529,24 @@ namespace Tools.CrashReporter.CrashReportProcess
 					}
 				}
 
-				CrashReporterProcessServicer.WriteEvent( "# WebAdded: ReportID   =" + string.Format( "{0,7}", ReportID ) + " Path=" + NewContext.CrashDirectory );
+				CrashReporterProcessServicer.WriteEvent(string.Format("PROC-{0} ", ProcessorIndex) + "# WebAdded: ReportID   =" + string.Format("{0,7}", ReportID) + " Path=" + NewContext.CrashDirectory);
 
 				UpdateProcessedReports();
-				WebAddedReports++;
-				double Ratio = (double)WebAddedReports / (double)ProcessedReports * 100;
+				WebAddCounter.AddEvent();
+				CrashReporterProcessServicer.StatusReporter.IncrementCount(StatusReportingConstants.ProcessingSucceededEvent);
+				double Ratio = (double)WebAddCounter.TotalEvents / (double)ProcessedReports * 100;
 
-				double AddedPerDay = (double)WebAddedReports / Timer.Elapsed.TotalDays;
+				double AddedPerDay = (double)WebAddCounter.TotalEvents / Timer.Elapsed.TotalDays;
 
-				CrashReporterProcessServicer.WriteEvent( string.Format( "Ratio={0,2} Processed={1,7} WebAdded={2,7} AddReportTime={3} AddedPerDay={4}", (int)Ratio, ProcessedReports, WebAddedReports, AddReportTime.Elapsed.TotalSeconds.ToString( "0.00" ), (int)AddedPerDay ) );
+				CrashReporterProcessServicer.WriteEvent(string.Format("PROC-{0} ", ProcessorIndex) + 
+					string.Format(
+						"Ratio={0,2} Processed={1,7} WebAdded={2,7} AddReportTime={3} AddedPerDay={4} AddedPerMinute={5:N1}", (int) Ratio,
+						ProcessedReports, WebAddCounter.TotalEvents, AddReportTime.Elapsed.TotalSeconds.ToString("0.00"), (int)AddedPerDay, WebAddCounter.EventsPerSecond * 60));
 				return true;
 			}
 			catch( Exception Ex )
 			{
-				CrashReporterProcessServicer.WriteException( "AddReport: " + DirInfo.Name + "\n\n" + Ex.ToString() );
+				CrashReporterProcessServicer.WriteException(string.Format("PROC-{0} ", ProcessorIndex) + "AddReport: " + DirInfo.Name + "\n\n" + Ex.ToString());
 			}
 
 			return false;
@@ -529,10 +557,10 @@ namespace Tools.CrashReporter.CrashReportProcess
 		/// </summary>
 		/// <returns>true if all the files synced without issue, false otherwise.</returns>
 		/// <remarks>As MinidumpDiagnostics is synced to #head, it requires the engine config files that match to run properly.</remarks>
-		private bool SyncRequiredFiles()
+		private static bool SyncRequiredFiles()
 		{
 			// Use the latest MinidumpDiagnostics from the main branch.
-			string SyncBinariesString = Path.Combine(Properties.Settings.Default.DepotRoot, Properties.Settings.Default.SyncBinariesFromDepot);
+			string SyncBinariesString = Path.Combine(Config.Default.DepotRoot, Config.Default.SyncBinariesFromDepot);
 			using (var MDDSyncProc = new LaunchProcess(PerforceExePath, null, CrashReporterProcessServicer.WriteP4, "sync", SyncBinariesString))
 			{
 				if( MDDSyncProc.WaitForExit( SyncTimeoutSeconds * 1000 ) == EWaitResult.TimedOut )
@@ -542,7 +570,7 @@ namespace Tools.CrashReporter.CrashReportProcess
 				}
 			}
 
-			string SyncConfigString = Path.Combine(Properties.Settings.Default.DepotRoot, Properties.Settings.Default.SyncConfigFromDepot);
+			string SyncConfigString = Path.Combine(Config.Default.DepotRoot, Config.Default.SyncConfigFromDepot);
 			using (var ConfigSyncProc = new LaunchProcess(PerforceExePath, null, CrashReporterProcessServicer.WriteP4, "sync", SyncConfigString))
 			{
 				if (ConfigSyncProc.WaitForExit(SyncTimeoutSeconds * 1000) == EWaitResult.TimedOut)
@@ -552,7 +580,7 @@ namespace Tools.CrashReporter.CrashReportProcess
 				}
 			}
 
-			string SyncThirdPartyString = Path.Combine(Properties.Settings.Default.DepotRoot, Properties.Settings.Default.SyncThirdPartyFromDepot); // Required by Perforce
+			string SyncThirdPartyString = Path.Combine(Config.Default.DepotRoot, Config.Default.SyncThirdPartyFromDepot); // Required by Perforce
 			using (var MiscSyncProc = new LaunchProcess(PerforceExePath, null, CrashReporterProcessServicer.WriteP4, "sync", SyncThirdPartyString))
 			{
 				if (MiscSyncProc.WaitForExit(SyncTimeoutSeconds * 1000) == EWaitResult.TimedOut)
@@ -578,14 +606,14 @@ namespace Tools.CrashReporter.CrashReportProcess
 			if( !DirInfo.Exists )
 			{
 				// Something very odd happened.
-				CrashReporterProcessServicer.WriteEvent( "ProcessReport: Directory not found: " + NewContext.CrashDirectory );
+				CrashReporterProcessServicer.WriteEvent(string.Format("PROC-{0} ", ProcessorIndex) + "ProcessReport: Directory not found: " + NewContext.CrashDirectory);
 				return;
 			}
 
 			double WaitTime = ReadProcessAddReport( DirInfo, NewContext );
 			
 			// Make sure any log messages have been written to disk
-			CrashReporterProcessServicer.WriteEvent( string.Format( "ProcessReportTime={0} WaitTime={1}", ProcessReportSW.Elapsed.TotalSeconds.ToString( "0.0" ), WaitTime.ToString( "0.00" ) ) );
+			CrashReporterProcessServicer.WriteEvent(string.Format("PROC-{0} ", ProcessorIndex) + string.Format("ProcessReportTime={0} WaitTime={1}", ProcessReportSW.Elapsed.TotalSeconds.ToString("0.0"), WaitTime.ToString("0.00")));
 		}
 
 		/// <summary> Reads callstack, source context and error message from the diagnostics file, if exists. </summary>
@@ -666,46 +694,78 @@ namespace Tools.CrashReporter.CrashReportProcess
 				}
 
 				// Check if the new context has processed data.
-				if( NewContext.HasProcessedData() )
+				if (!NewContext.HasProcessedData())
 				{
-					Stopwatch WaitSW = Stopwatch.StartNew();
-					// Wait for previous task, should not really happen.
-					if( AddReportTask != null )
-					{
-						AddReportTask.Wait();
-					}
-					double WaitTime = WaitSW.Elapsed.TotalSeconds;
+					CrashReporterProcessServicer.WriteEvent(string.Format("PROC-{0} ", ProcessorIndex) + "% Warning no callstack or error msg : BuiltFromCL=" + string.Format("{0,7}", NewContext.PrimaryCrashProperties.EngineVersion) + " Path=" + NewContext.CrashDirectory);
+					NewContext.PrimaryCrashProperties.ProcessorFailedMessage = "No callstack or error message. Diagnostics missing or failed.";
+				}
 
-					//bool bAdded = AddReport(DirInfo, NewContext, LogFileName, DumpFileName, VideoFileName );
-					
-					// Save/update crash context to the file.
-					NewContext.ToFile();
-					AddReportTask = Task.Factory.StartNew( () =>
-					{
-						bool bAdded = AddReport( DirInfo, NewContext, LogFileName, DumpFileName, VideoFileName );
-						FinalizeReport( bAdded, DirInfo, NewContext );
-					} );
+				Stopwatch WaitSW = Stopwatch.StartNew();
 
-					return WaitTime;
-				}
-				else
+				// Wait for previous task, should not really happen.
+				int AddReportTaskSlot = WaitForFreeAddReportTask();
+				CrashReporterProcessServicer.WriteEvent(string.Format("PROC-{0} ", ProcessorIndex) + string.Format("Starting AddReportTask running on slot {0} of {1} ({2} active)", AddReportTaskSlot, Config.Default.AddReportsPerProcessor, GetActiveAddReportTasks()));
+				double WaitTime = WaitSW.Elapsed.TotalSeconds;
+
+				//bool bAdded = AddReport(DirInfo, NewContext, LogFileName, DumpFileName, VideoFileName );
+					
+				// Save/update crash context to the file.
+				NewContext.ToFile();
+				AddReportTasks[AddReportTaskSlot] = Task.Factory.StartNew(() =>
 				{
-					CrashReporterProcessServicer.WriteFailure( "! Invalid : BuiltFromCL=" + string.Format( "{0,7}", NewContext.PrimaryCrashProperties.EngineVersion ) + " Path=" + NewContext.CrashDirectory );
-					FinalizeReport( false, DirInfo, NewContext );
-				}
+					bool bAdded = AddReport( DirInfo, NewContext, LogFileName, DumpFileName, VideoFileName );
+					FinalizeReport( bAdded, DirInfo, NewContext );
+				} );
+
+				return WaitTime;
+
 			}
 			catch( Exception Ex )
 			{
-				CrashReporterProcessServicer.WriteException( "ProcessReport: " + NewContext.CrashDirectory + "\n\n: " + Ex.ToString() );
+				CrashReporterProcessServicer.WriteException(string.Format("PROC-{0} ", ProcessorIndex) + "ProcessReport: " + NewContext.CrashDirectory + "\n\n: " + Ex.ToString());
 			}
 
 			return 0.0;
 		}
 
+		int WaitForFreeAddReportTask()
+		{
+			lock (AddReportTasks)
+			{
+				// Is a slot available because it has finished?
+				for (int Index = 0; Index < AddReportTasks.Count; Index++)
+				{
+					if (AddReportTasks != null && AddReportTasks[Index].IsCompleted)
+					{
+						AddReportTasks[Index] = null;
+						return Index;
+					}
+				}
+
+				// Can we make a new slot?
+				if (AddReportTasks.Count < Config.Default.AddReportsPerProcessor)
+				{
+					// Can create another
+					AddReportTasks.Add(null);
+					return AddReportTasks.Count - 1;
+				}
+
+				// Wait for a slot
+				int CompletedIndex = Task.WaitAny(AddReportTasks.ToArray());
+				AddReportTasks[CompletedIndex] = null;
+				return CompletedIndex;
+			}
+		}
+
+		int GetActiveAddReportTasks()
+		{
+			return AddReportTasks.Sum(AddReportTask => AddReportTask == null || AddReportTask.IsCompleted ? 0 : 1);
+		}
+
 		void ProcessDumpFile( string DiagnosticsPath, FGenericCrashContext NewContext )
 		{
 			// Use the latest MinidumpDiagnostics from the main branch.
-			string Win64BinariesDirectory = Path.Combine(Properties.Settings.Default.DepotRoot, Properties.Settings.Default.MDDBinariesFolderInDepot);
+			string Win64BinariesDirectory = Path.Combine(Config.Default.DepotRoot, Config.Default.MDDBinariesFolderInDepot);
 #if DEBUG
 			// Note: the debug executable must be built locally or synced from Perforce manually
 			string MinidumpDiagnosticsName = Path.Combine(Win64BinariesDirectory, "MinidumpDiagnostics-Win64-Debug.exe");
@@ -750,7 +810,7 @@ namespace Tools.CrashReporter.CrashReportProcess
 					"-SyncMicrosoftSymbols",
 					"-unattended",
 					"-Log=" + NewContext.GetAsFilename() + "-backup-.log",
-					"-DepotIndex=" + Properties.Settings.Default.DepotIndex,
+					"-DepotIndex=" + Config.Default.DepotIndex,
 					"-ini:Engine:[LogFiles]:PurgeLogsDays="+PurgeLogsDays,
 					"-LOGTIMESINCESTART"
 				}
@@ -776,14 +836,52 @@ namespace Tools.CrashReporter.CrashReportProcess
 				CrashReporterProcessServicer.WriteMDD( MinidumpDiagnosticsName + " Params: " + String.Join( " ", MinidumpDiagnosticsParams ) );
 			}
 
-			LaunchProcess ReportParser = new LaunchProcess(MinidumpDiagnosticsName, Path.GetDirectoryName(MinidumpDiagnosticsName), CaptureMessageDelegate, MinidumpDiagnosticsParams.ToArray());
-
-			if (ReportParser.WaitForExit(MinidumpDiagnosticsTimeoutSeconds * 1000) == EWaitResult.TimedOut)
+			Stopwatch WaitSW = Stopwatch.StartNew();
+			lock (MinidumpDiagnosticsLock)
 			{
-				CrashReporterProcessServicer.WriteFailure( "ProcessDumpFile: Timed out running MinidumpDiagnostics" );
+				Double WaitForLockTime = WaitSW.Elapsed.TotalSeconds;
+
+				LaunchProcess ReportParser = new LaunchProcess(MinidumpDiagnosticsName, Path.GetDirectoryName(MinidumpDiagnosticsName), CaptureMessageDelegate, MinidumpDiagnosticsParams.ToArray());
+
+				if (ReportParser.WaitForExit(MinidumpDiagnosticsTimeoutMilliseconds) == EWaitResult.TimedOut)
+				{
+					CrashReporterProcessServicer.WriteFailure(string.Format("PROC-{0} ", ProcessorIndex) + "ProcessDumpFile: Timed out running MinidumpDiagnostics");
+				}
+
+				Double TotalMDDTime = WaitSW.Elapsed.TotalSeconds;
+
+				CrashReporterProcessServicer.WriteEvent(string.Format("PROC-{0} ", ProcessorIndex) + string.Format("ProcessDumpFile: Total MDD exec time {0:N1}s (blocked for {1:N1}s)", TotalMDDTime, WaitForLockTime));
 			}
 
 			ReadDiagnosticsFile( NewContext );
 		}
+
+		private static void TickStatic(ReportWatcher InWatcher)
+		{
+			if (Monitor.TryEnter(TickStaticLock, 0))
+			{
+				try
+				{
+					// Remove all folders older than n days to prevent unbounded growth
+					if ((DateTime.Now - LastCleanupTime) > TimeSpan.FromMinutes(15))
+					{
+						CleanRepository(InWatcher);
+						LastCleanupTime = DateTime.Now;
+					}
+
+					DateTime CurrentDate = DateTime.UtcNow.Date;
+					if (CurrentDate > LoggingDate)
+					{
+						// Check the log and create a new one for a new day.
+						CrashReporterProcessServicer.Log.CreateNewLogFile();
+						LoggingDate = CurrentDate;
+					}
+				}
+				finally
+				{
+					Monitor.Exit(TickStaticLock);
+				}
+			}
+		}
 	}
 }
