Copying //UE4/Dev-Core to //UE4/Main

diff --git a/Engine/Source/Programs/CrashReporter/CrashReportProcess/ReportQueue.cs b/Engine/Source/Programs/CrashReporter/CrashReportProcess/ReportQueue.cs
--- a/Engine/Source/Programs/CrashReporter/CrashReportProcess/ReportQueue.cs
+++ b/Engine/Source/Programs/CrashReporter/CrashReportProcess/ReportQueue.cs
@@ -1,481 +1,27 @@
 // Copyright 1998-2016 Epic Games, Inc. All Rights Reserved.
 
-using System;
 using System.Collections.Generic;
-using System.Collections.Concurrent;
-using System.Diagnostics;
-using System.IO;
-using System.IO.Compression;
-using System.Linq;
-
 using Tools.CrashReporter.CrashReportCommon;
-using Tools.DotNETCommon.XmlHandler;
-
 
 namespace Tools.CrashReporter.CrashReportProcess
 {
-	using ReportIdSet = HashSet<string>;
-	using System.Runtime.InteropServices;
-	using System.Globalization;
-
-	/// <summary>
-	/// A queue of pending reports in a specific folder
-	/// </summary>
-	public class ReportQueue
+	interface IReportQueue
 	{
-		/// <summary>
-		/// Constructor taking the landing zone
-		/// </summary>
-		public ReportQueue(string LandingZonePath)
-		{
-			LandingZone = LandingZonePath;
-		}
-
-		/// <summary>
-		/// Look for new report folders and add them to the publicly available thread-safe queue
-		/// </summary>
-		public void CheckForNewReports()
-		{	
-			try
-			{
-				var NewReportPath = "";
-				var ReportsInLandingZone = new ReportIdSet();
-
-				// Check the landing zones for new reports
-				DirectoryInfo DirInfo = new DirectoryInfo(LandingZone);
-
-				// Couldn't find a landing zone, skip and try again later.
-				// Crash receiver will recreate them if needed.
-				if( !DirInfo.Exists )
-				{
-					CrashReporterProcessServicer.WriteFailure( "LandingZone not found: " + LandingZone );
-					return;
-				}
-
-				// Add any new reports
-				foreach (var SubDirInfo in DirInfo.GetDirectories())
-				{
-					try
-					{
-						if( Directory.Exists( SubDirInfo.FullName ) )
-						{
-							NewReportPath = SubDirInfo.FullName;
-							ReportsInLandingZone.Add( NewReportPath );
-							if( !ReportsInLandingZoneLastTimeWeChecked.Contains( NewReportPath ) )
-							{
-								EnqueueNewReport( NewReportPath );
-							}
-						}
-					}
-					catch( System.Exception Ex )
-					{
-						CrashReporterProcessServicer.WriteException( "CheckForNewReportsInner: " + Ex.ToString() );
-						ReportProcessor.MoveReportToInvalid( NewReportPath, "NEWRECORD_FAIL_" + DateTime.Now.Ticks );
-					}
-				}
-				//CrashReporterProcessServicer.WriteEvent( string.Format( "ReportsInLandingZone={0}, ReportsInLandingZoneLastTimeWeChecked={1}", ReportsInLandingZone.Count, ReportsInLandingZoneLastTimeWeChecked.Count ) );
-				ReportsInLandingZoneLastTimeWeChecked = ReportsInLandingZone;
-			}
-			catch( Exception Ex )
-			{
-				CrashReporterProcessServicer.WriteException( "CheckForNewReportsOuter: " + Ex.ToString() );
-			}
-		}
-
-		/// <summary>
-		/// Tidy up the landing zone folder
-		/// </summary>
-		public void CleanLandingZone()
-		{
-			var DirInfo = new DirectoryInfo(LandingZone);
-			foreach(var SubDirInfo in DirInfo.GetDirectories())
-			{
-				if (SubDirInfo.LastWriteTimeUtc.AddDays(Properties.Settings.Default.DaysToSunsetReport) < DateTime.UtcNow)
-				{
-					SubDirInfo.Delete(true);
-				}
-			}
-		}
-
-		/// <summary> Looks for the CrashContext.runtime-xml, if found, will return a new instance of the FGenericCrashContext. </summary>
-		private FGenericCrashContext FindCrashContext( string NewReportPath )
-		{
-			string CrashContextPath = Path.Combine( NewReportPath, FGenericCrashContext.CrashContextRuntimeXMLName );
-			bool bExist = File.Exists( CrashContextPath );
-
-			if( bExist )
-			{
-				return FGenericCrashContext.FromFile( CrashContextPath );
-			}
-
-			return null;
-		}
-
-		/// <summary> Converts WER metadata xml file to the crash context. </summary>
-		private void ConvertMetadataToCrashContext(WERReportMetadata Metadata, string NewReportPath, ref FGenericCrashContext OutCrashContext)
-		{
-			if (OutCrashContext == null)
-			{
-				OutCrashContext = new FGenericCrashContext();
-			}
-
-			FReportData ReportData = new FReportData( Metadata, NewReportPath );
-
-			OutCrashContext.PrimaryCrashProperties.CrashVersion = (int)ECrashDescVersions.VER_1_NewCrashFormat;
-
-			//OutCrashContext.PrimaryCrashProperties.ProcessId = 0; don't overwrite valid ids, zero is default anyway
-			OutCrashContext.PrimaryCrashProperties.CrashGUID = new DirectoryInfo( NewReportPath ).Name;  
-// 			OutCrashContext.PrimaryCrashProperties.IsInternalBuild
-// 			OutCrashContext.PrimaryCrashProperties.IsPerforceBuild
-// 			OutCrashContext.PrimaryCrashProperties.IsSourceDistribution
-// 			OutCrashContext.PrimaryCrashProperties.SecondsSinceStart
-			OutCrashContext.PrimaryCrashProperties.GameName = ReportData.GameName;
-// 			OutCrashContext.PrimaryCrashProperties.ExecutableName
-// 			OutCrashContext.PrimaryCrashProperties.BuildConfiguration
-// 			OutCrashContext.PrimaryCrashProperties.PlatformName
-// 			OutCrashContext.PrimaryCrashProperties.PlatformNameIni
-			OutCrashContext.PrimaryCrashProperties.PlatformFullName = ReportData.Platform;
-			OutCrashContext.PrimaryCrashProperties.EngineMode = ReportData.EngineMode;
-			OutCrashContext.PrimaryCrashProperties.EngineVersion = ReportData.GetEngineVersion();
-			OutCrashContext.PrimaryCrashProperties.CommandLine = ReportData.CommandLine;
-// 			OutCrashContext.PrimaryCrashProperties.LanguageLCID
-
-			// Create a locate and get the language.
-			int LanguageCode = 0;
-			int.TryParse( ReportData.Language, out LanguageCode );
-			try
-			{
-				if (LanguageCode > 0)
-				{
-					CultureInfo LanguageCI = new CultureInfo( LanguageCode );
-					OutCrashContext.PrimaryCrashProperties.AppDefaultLocale = LanguageCI.Name;
-				}
-			}
-			catch (Exception)
-			{
-				OutCrashContext.PrimaryCrashProperties.AppDefaultLocale = null;
-			}
-
-			if (string.IsNullOrEmpty(OutCrashContext.PrimaryCrashProperties.AppDefaultLocale))
-			{
-				// Default to en-US
-				OutCrashContext.PrimaryCrashProperties.AppDefaultLocale = "en-US";				
-			}
-
-// 			OutCrashContext.PrimaryCrashProperties.IsUE4Release
-			string WERUserName = ReportData.UserName;
-			if (!string.IsNullOrEmpty(WERUserName) || string.IsNullOrEmpty(OutCrashContext.PrimaryCrashProperties.UserName))
-			{
-				OutCrashContext.PrimaryCrashProperties.UserName = WERUserName;
-			}
-
-			OutCrashContext.PrimaryCrashProperties.BaseDir = ReportData.BaseDir;
-// 			OutCrashContext.PrimaryCrashProperties.RootDir
-			OutCrashContext.PrimaryCrashProperties.MachineId = ReportData.MachineId;
-			OutCrashContext.PrimaryCrashProperties.EpicAccountId = ReportData.EpicAccountId;
-// 			OutCrashContext.PrimaryCrashProperties.CallStack
-// 			OutCrashContext.PrimaryCrashProperties.SourceContext
-			OutCrashContext.PrimaryCrashProperties.UserDescription = string.Join( "\n", ReportData.UserDescription );
-
-			if (string.IsNullOrEmpty(OutCrashContext.PrimaryCrashProperties.ErrorMessage))
-			{
-				OutCrashContext.PrimaryCrashProperties.ErrorMessage = string.Join("\n", ReportData.ErrorMessage);
-			}
-
-			// 			OutCrashContext.PrimaryCrashProperties.CrashDumpMode
-// 			OutCrashContext.PrimaryCrashProperties.Misc.NumberOfCores
-// 			OutCrashContext.PrimaryCrashProperties.Misc.NumberOfCoresIncludingHyperthreads
-// 			OutCrashContext.PrimaryCrashProperties.Misc.Is64bitOperatingSystem
-// 			OutCrashContext.PrimaryCrashProperties.Misc.CPUVendor
-// 			OutCrashContext.PrimaryCrashProperties.Misc.CPUBrand
-// 			OutCrashContext.PrimaryCrashProperties.Misc.PrimaryGPUBrand
-// 			OutCrashContext.PrimaryCrashProperties.Misc.OSVersionMajor
-// 			OutCrashContext.PrimaryCrashProperties.Misc.OSVersionMinor
-// 			OutCrashContext.PrimaryCrashProperties.Misc.AppDiskTotalNumberOfBytes
-// 			OutCrashContext.PrimaryCrashProperties.Misc.AppDiskNumberOfFreeBytes
-// 			OutCrashContext.PrimaryCrashProperties.MemoryStats.TotalPhysical
-// 			OutCrashContext.PrimaryCrashProperties.MemoryStats.TotalVirtual
-// 			OutCrashContext.PrimaryCrashProperties.MemoryStats.PageSize
-// 			OutCrashContext.PrimaryCrashProperties.MemoryStats.TotalPhysicalGB
-// 			OutCrashContext.PrimaryCrashProperties.MemoryStats.AvailablePhysical
-// 			OutCrashContext.PrimaryCrashProperties.MemoryStats.AvailableVirtual
-// 			OutCrashContext.PrimaryCrashProperties.MemoryStats.UsedPhysical
-// 			OutCrashContext.PrimaryCrashProperties.MemoryStats.PeakUsedPhysical
-// 			OutCrashContext.PrimaryCrashProperties.MemoryStats.UsedVirtual
-// 			OutCrashContext.PrimaryCrashProperties.MemoryStats.PeakUsedVirtual
-// 			OutCrashContext.PrimaryCrashProperties.MemoryStats.bIsOOM
-// 			OutCrashContext.PrimaryCrashProperties.MemoryStats.OOMAllocationSize
-// 			OutCrashContext.PrimaryCrashProperties.MemoryStats.OOMAllocationAlignment
-			OutCrashContext.PrimaryCrashProperties.TimeofCrash = new DateTime( ReportData.Ticks );
-			OutCrashContext.PrimaryCrashProperties.bAllowToBeContacted = ReportData.AllowToBeContacted;
-		}
-
-		/// <summary> Looks for the WER metadata xml file, if found, will return a new instance of the WERReportMetadata. </summary>
-		private WERReportMetadata FindMetadata( string NewReportPath )
-		{
-			WERReportMetadata MetaData = null;
-
-			// Just to verify that the report is still there.
-			DirectoryInfo DirInfo = new DirectoryInfo( NewReportPath );
-			if( !DirInfo.Exists )
-			{
-				CrashReporterProcessServicer.WriteFailure( "FindMetadata: Directory not found " + NewReportPath );
-			}
-			else
-			{
-				// Check to see if we wish to suppress processing of this report.	
-				foreach( var Info in DirInfo.GetFiles( "*.xml" ) )
-				{
-					var MetaDataToCheck = XmlHandler.ReadXml<WERReportMetadata>( Info.FullName );
-					if( CheckMetaData( MetaDataToCheck ) )
-					{
-						MetaData = MetaDataToCheck;
-						break;
-					}
-				}
-			}
-			
-			return MetaData;
-		}
-
-		/// <summary> Looks for the Diagnostics.txt file and returns true if exists. </summary>
-		private bool FindDiagnostics( string NewReportPath )
-		{
-			bool bFound = false;
-
-			DirectoryInfo DirInfo = new DirectoryInfo( NewReportPath );
-			foreach( var Info in DirInfo.GetFiles( CrashReporterConstants.DiagnosticsFileName ) )
-			{
-				bFound = true;
-				break;
-			}
-			return bFound;
-		}
-
-		/// <summary>
-		/// Optionally don't process some reports based on the Windows Error report meta data.
-		/// </summary>
-		/// <param name="WERData">The Windows Error Report meta data.</param>
-		/// <returns>false to reject the report.</returns>
-		private bool CheckMetaData( WERReportMetadata WERData )
-		{
-			if( WERData == null )
-			{
-				return false;
-			}
-
-			// Reject any crashes with the invalid metadata.
-			if( WERData.ProblemSignatures == null || WERData.DynamicSignatures == null || WERData.OSVersionInformation == null )
-			{
-				return false;
-			}
-
-			// Reject any crashes from the minidump processor.
-			if( WERData.ProblemSignatures.Parameter0 != null && WERData.ProblemSignatures.Parameter0.ToLower() == "MinidumpDiagnostics".ToLower() )
-			{
-				return false;
-			}
-
-			return true;
-		}
-
-		// From CrashUpload.cpp
-		/*
-		struct FCompressedCrashFile : FNoncopyable
-		{
-			int32 CurrentFileIndex; // 4 bytes for file index
-			FString Filename; // 4 bytes for length + 260 bytes for char data
-			TArray<uint8> Filedata; // 4 bytes for length + N bytes for data
-		}
-		*/
-	
-		[DllImport( "CompressionHelper.dll", CallingConvention = CallingConvention.Cdecl )]
-		private extern static bool UE4UncompressMemoryZLIB( /*void**/IntPtr UncompressedBuffer, Int32 UncompressedSize, /*const void**/IntPtr CompressedBuffer, Int32 CompressedSize );
+		int CheckForNewReports();
+		bool TryDequeueReport(out FGenericCrashContext Context);
+		void CleanLandingZone();
 
-		/// <summary> 
-		/// Decompresses a compressed crash report. 
-		/// </summary>
-		unsafe private bool DecompressReport( string CompressedReportPath, FCompressedCrashInformation Meta )
-		{
-			string ReportPath = Path.GetDirectoryName( CompressedReportPath );
-			using (FileStream FileStream = File.OpenRead( CompressedReportPath ))
-			{
-				Int32 UncompressedSize = Meta.GetUncompressedSize();
-				Int32 CompressedSize = Meta.GetCompressedSize();
-
-				if (FileStream.Length != CompressedSize)
-				{
-					return false;
-				}
-
-				byte[] UncompressedBufferArray = new byte[UncompressedSize];
-				byte[] CompressedBufferArray = new byte[CompressedSize];
-
-				FileStream.Read( CompressedBufferArray, 0, CompressedSize );
-
-				fixed (byte* UncompressedBufferPtr = UncompressedBufferArray, CompressedBufferPtr = CompressedBufferArray)
-				{
-					bool bResult = UE4UncompressMemoryZLIB( (IntPtr)UncompressedBufferPtr, UncompressedSize, (IntPtr)CompressedBufferPtr, CompressedSize );
-					if (!bResult)
-					{
-						CrashReporterProcessServicer.WriteFailure( "! ZLibFail: Path=" + ReportPath );
-						return false;
-					}
-				}
-
-				MemoryStream MemoryStream = new MemoryStream( UncompressedBufferArray, false );
-				BinaryReader BinaryData = new BinaryReader( MemoryStream );
-
-				for (int FileIndex = 0; FileIndex < Meta.GetNumberOfFiles(); FileIndex++)
-				{
-					Int32 CurrentFileIndex = BinaryData.ReadInt32();
-					if (CurrentFileIndex != FileIndex)
-					{
-						CrashReporterProcessServicer.WriteFailure( "! ReadFail: Required=" + FileIndex + " Read=" + CurrentFileIndex );
-						return false;
-					}
-
-					string Filename = FBinaryReaderHelper.ReadFixedSizeString( BinaryData );
-					Int32 FiledataLength = BinaryData.ReadInt32();
-					byte[] Filedata = BinaryData.ReadBytes( FiledataLength );
-
-					File.WriteAllBytes( Path.Combine( ReportPath, Filename ), Filedata );
-				}
-			}
-
-			return true;
-		}
-
-		/// <summary> Enqueues a new crash. </summary>
-		void EnqueueNewReport( string NewReportPath )
-		{
-			string ReportName = Path.GetFileName( NewReportPath );
-
-			string CompressedReportPath = Path.Combine( NewReportPath, ReportName + ".ue4crash" );
-			string MetadataPath = Path.Combine( NewReportPath, ReportName + ".xml" );
-			bool bIsCompressed = File.Exists( CompressedReportPath ) && File.Exists( MetadataPath );
-			if (bIsCompressed)
-			{
-				FCompressedCrashInformation CompressedCrashInformation = XmlHandler.ReadXml<FCompressedCrashInformation>( MetadataPath );
-				bool bResult = DecompressReport( CompressedReportPath, CompressedCrashInformation );
-				if (!bResult)
-				{
-					ReportProcessor.CleanReport( new DirectoryInfo( NewReportPath ) );
-					ReportProcessor.MoveReportToInvalid( NewReportPath, "DECOMPRESSION_FAIL_" + DateTime.Now.Ticks + "_" + ReportName );
-					return;
-				}
-				else
-				{
-					// Rename metadata file to not interfere with the WERReportMetadata.
-					File.Move( MetadataPath, Path.ChangeExtension( MetadataPath, "processed_xml" ) );
-				}
-			}
-
-			// Unified crash reporting
-			FGenericCrashContext GenericContext = FindCrashContext( NewReportPath );
-			FGenericCrashContext Context = GenericContext;
-
-			bool bFromWER = false;
-			if (Context == null || !Context.HasProcessedData())
-			{
-				WERReportMetadata MetaData = FindMetadata( NewReportPath );
-				if (MetaData != null)
-				{
-					FReportData ReportData = new FReportData( MetaData, NewReportPath );
-					ConvertMetadataToCrashContext(MetaData, NewReportPath, ref Context);
-					bFromWER = true;
-				}
-			}
-
-			if (Context == null)
-			{
-				CrashReporterProcessServicer.WriteFailure( "! NoCntx  : Path=" + NewReportPath );
-				ReportProcessor.CleanReport( new DirectoryInfo( NewReportPath ) );
-			}
-			else
-			{
-				Context.CrashDirectory = NewReportPath;
-				Context.PrimaryCrashProperties.SetPlatformFullName();
-
-				// Added data from WER, save to the crash context file.
-				if (bFromWER)
-				{
-					Context.ToFile();
-				}
-				
-				FEngineVersion EngineVersion = new FEngineVersion( Context.PrimaryCrashProperties.EngineVersion ); 
-
-				uint BuiltFromCL = EngineVersion.Changelist;
-
-				string BranchName = EngineVersion.Branch;
-				if (string.IsNullOrEmpty( BranchName ))
-				{
-					CrashReporterProcessServicer.WriteFailure("% NoBranch: BuiltFromCL=" + string.Format("{0,7}", BuiltFromCL) + " Path=" + NewReportPath + " EngineVersion=" + Context.PrimaryCrashProperties.EngineVersion);
-					ReportProcessor.MoveReportToInvalid( NewReportPath, Context.GetAsFilename() );
-					return;
-				}
-
-				if (BranchName.Equals( CrashReporterConstants.LicenseeBranchName, StringComparison.InvariantCultureIgnoreCase ))
-				{
-					CrashReporterProcessServicer.WriteFailure( "% UE4-QA  : BuiltFromCL=" + string.Format( "{0,7}", BuiltFromCL ) + " Path=" + NewReportPath );
-					ReportProcessor.CleanReport( NewReportPath );
-					return;
-				}
-
-				// Look for the Diagnostics.txt, if we have a diagnostics file we can continue event if the CL is marked as broken.
-				bool bHasDiagnostics = FindDiagnostics( NewReportPath );
-
-				if (BuiltFromCL == 0 && ( !bHasDiagnostics && !Context.HasProcessedData() ))
-				{
-					// For now ignore all locally made crashes.
-					CrashReporterProcessServicer.WriteFailure( "! BROKEN0 : BuiltFromCL=" + string.Format( "{0,7}", BuiltFromCL ) + " Path=" + NewReportPath );
-					ReportProcessor.CleanReport( NewReportPath );
-					return;
-				}
-
-
-				lock (NewReportsLock)
-				{
-					NewCrashContexts.Add( Context );
-				}
-				CrashReporterProcessServicer.WriteEvent( "+ Enqueued: BuiltFromCL=" + string.Format( "{0,7}", BuiltFromCL ) + " Path=" + NewReportPath );
-			}
-			
-		}
-
-		/// <summary> Tries to dequeue a report from the list. </summary>
-		public bool TryDequeueReport( out FGenericCrashContext Context )
-		{
-			lock( NewReportsLock )
-			{
-				int LastIndex = NewCrashContexts.Count - 1;
-
-				if( LastIndex >= 0 )
-				{
-					Context = NewCrashContexts[LastIndex];
-					NewCrashContexts.RemoveAt( LastIndex );
-					CrashReporterProcessServicer.WriteEvent( "- Dequeued: BuiltFromCL=" + string.Format( "{0,7}", Context.PrimaryCrashProperties.EngineVersion ) + " Path=" + Context.CrashDirectory );
-					return true;
-				}
-				else
-				{
-					Context = null;
-					return false;
-				}
-			}
-		}
-
-
-		/// <summary> A list of freshly landed crash reports. </summary>
-		private List<FGenericCrashContext> NewCrashContexts = new List<FGenericCrashContext>();
-
-		/// <summary> Object used to synchronize the access to the NewReport. </summary>
-		private Object NewReportsLock = new Object();
+		string QueueId { get; }
+		string LandingZonePath { get; }
+	}
 
-		/// <summary>Incoming report path</summary>
-		string LandingZone;
+	abstract class ReportQueueBase : IReportQueue
+	{
+		public abstract int CheckForNewReports();
+		public abstract bool TryDequeueReport(out FGenericCrashContext Context);
+		public abstract void CleanLandingZone();
 
-		/// <summary>Report IDs that were in the folder on the last pass</summary>
-		ReportIdSet ReportsInLandingZoneLastTimeWeChecked = new ReportIdSet();
+		public abstract string QueueId { get; }
+		public abstract string LandingZonePath { get; }
 	}
 }
diff --git a/Engine/Source/Programs/CrashReporter/CrashReportProcess/ReportWatcher.cs b/Engine/Source/Programs/CrashReporter/CrashReportProcess/ReportWatcher.cs
--- a/Engine/Source/Programs/CrashReporter/CrashReportProcess/ReportWatcher.cs
+++ b/Engine/Source/Programs/CrashReporter/CrashReportProcess/ReportWatcher.cs
@@ -4,16 +4,17 @@ using System;
 using System.Collections.Generic;
 using System.Threading;
 using System.Threading.Tasks;
+using System.Text;
 
 namespace Tools.CrashReporter.CrashReportProcess
 {
 	/// <summary>
 	/// A class to monitor the crash report repository for new crash reports.
 	/// </summary>
-	public sealed class ReportWatcher : IDisposable
+	sealed class ReportWatcher : IDisposable
 	{
 		/// <summary>A queue of freshly landed crash reports ready to be processed.</summary>
-		public List<ReportQueue> ReportQueues = new List<ReportQueue>();
+		public List<IReportQueue> ReportQueues = new List<IReportQueue>();
 
 		/// <summary>Task to periodically check for the arrival of new report folders.</summary>
 		Task WatcherTask;
@@ -49,13 +50,13 @@ namespace Tools.CrashReporter.CrashReportProcess
 		void Start()
 		{
 			CrashReporterProcessServicer.WriteEvent("CrashReportProcessor watching directories:");
-			var Settings = Properties.Settings.Default;
+			var Settings = Config.Default;
 
 			if (!string.IsNullOrEmpty(Settings.InternalLandingZone))
 			{
 				if( System.IO.Directory.Exists( Settings.InternalLandingZone ) )
 				{
-					ReportQueues.Add(new ReportQueue(Settings.InternalLandingZone));
+					ReportQueues.Add(new ReceiverReportQueue("Epic Crashes", Settings.InternalLandingZone));
 					CrashReporterProcessServicer.WriteEvent(string.Format("\t{0} (internal, high priority)", Settings.InternalLandingZone));
 				}
 				else
@@ -69,7 +70,7 @@ namespace Tools.CrashReporter.CrashReportProcess
 			{
 				if( System.IO.Directory.Exists( Settings.ExternalLandingZone ) )
 				{
-					ReportQueues.Add( new ReportQueue( Settings.ExternalLandingZone ) );
+					ReportQueues.Add(new ReceiverReportQueue("External Crashes", Settings.ExternalLandingZone));
 					CrashReporterProcessServicer.WriteEvent( string.Format( "\t{0}", Settings.ExternalLandingZone ) );
 				}
 				else
@@ -79,16 +80,29 @@ namespace Tools.CrashReporter.CrashReportProcess
 			}
 #endif //!DEBUG
 
+			// Init queue entries in StatusReporter
+			foreach (var Queue in ReportQueues)
+			{
+				CrashReporterProcessServicer.StatusReporter.InitQueue(Queue.QueueId, Queue.LandingZonePath);
+			}
+
 			var Cancel = CancelSource.Token;
 			WatcherTask = Task.Factory.StartNew(async () =>
 			{
+				DateTime LastQueueSizeReport = DateTime.MinValue;
 				while (!Cancel.IsCancellationRequested)
 				{
 					// Check the landing zones for new reports
+					DateTime StartTime = DateTime.Now;
+
 					foreach (var Queue in ReportQueues)
 					{
-						Queue.CheckForNewReports();
+						int QueueSize = Queue.CheckForNewReports();
+						CrashReporterProcessServicer.StatusReporter.SetQueueSize(Queue.QueueId, QueueSize);
 					}
+
+					TimeSpan TimeTaken = DateTime.Now - StartTime;
+					CrashReporterProcessServicer.WriteEvent(string.Format("Checking Landing Zones took {0:F2} seconds", TimeTaken.TotalSeconds));
 					await Task.Delay(60000, Cancel);
 				}
 			});
diff --git a/Engine/Source/Programs/CrashReporter/CrashReportProcess/StatusReporting.cs b/Engine/Source/Programs/CrashReporter/CrashReportProcess/StatusReporting.cs
new file mode 100644
--- /dev/null
+++ b/Engine/Source/Programs/CrashReporter/CrashReportProcess/StatusReporting.cs
@@ -0,0 +1,171 @@
+﻿using System;
+using System.Collections.Generic;
+using System.Linq;
+using System.Text;
+using System.Threading;
+using System.Threading.Tasks;
+
+namespace Tools.CrashReporter.CrashReportProcess
+{
+	class StatusReporting : IDisposable
+	{
+		public void InitQueue(string QueueName, string QueueLocation)
+		{
+			QueueLocations.Add(QueueName, QueueLocation);
+			QueueSizes.Add(QueueName, 0);
+		}
+
+		public void SetQueueSize(string QueueName, int Size)
+		{
+			lock (DataLock)
+			{
+				QueueSizes[QueueName] = Size;
+				bQueueSizesAvail = true;
+			}
+		}
+
+		public void IncrementCount(string CountName)
+		{
+			lock (DataLock)
+			{
+				if (Counters.ContainsKey(CountName))
+				{
+					Counters[CountName]++;
+				}
+				else
+				{
+					Counters.Add(CountName, 1);
+				}
+			}
+		}
+
+		public void Alert(string AlertText)
+		{
+			// Report an alert condition to Slack immediately
+			CrashReporterProcessServicer.WriteSlack("@channel ALERT: " + AlertText);
+		}
+
+		public void Start()
+		{
+			CrashReporterProcessServicer.WriteSlack(string.Format("CRP started (version {0})", Config.Default.VersionString));
+
+			StringBuilder StartupMessage = new StringBuilder();
+
+			StartupMessage.AppendLine("Queues...");
+			foreach (var Queue in QueueLocations)
+			{
+				StartupMessage.AppendLine("> " + Queue.Key + " @ " + Queue.Value);
+			}
+			CrashReporterProcessServicer.WriteSlack(StartupMessage.ToString());
+
+			StatusReportingTask = Task.Factory.StartNew(() =>
+			{
+				var Cancel = CancelSource.Token;
+
+				lock (TaskMonitor)
+				{
+					while (!Cancel.IsCancellationRequested)
+					{
+						StringBuilder StatusReportMessage = new StringBuilder();
+						lock (DataLock)
+						{
+							Dictionary<string, int> CountsInPeriod = GetCountsInPeriod(Counters, CountersAtLastReport);
+
+							if (bQueueSizesAvail)
+							{
+								int ProcessingStartedInPeriod = 0;
+								CountsInPeriod.TryGetValue(StatusReportingConstants.ProcessingStartedEvent, out ProcessingStartedInPeriod);
+								if (ProcessingStartedInPeriod > 0)
+								{
+									int QueueSizeSum = QueueSizes.Values.Sum();
+									TimeSpan MeanWaitTime =
+										TimeSpan.FromMinutes(0.5*Config.Default.MinutesBetweenQueueSizeReports*
+										                     (QueueSizeSum + QueueSizeSumAtLastReport)/ProcessingStartedInPeriod);
+									QueueSizeSumAtLastReport = QueueSizeSum;
+
+									string WaitTimeString;
+									if (MeanWaitTime == TimeSpan.Zero)
+									{
+										WaitTimeString = "nil";
+									}
+									else if (MeanWaitTime < TimeSpan.FromMinutes(2))
+									{
+										WaitTimeString = MeanWaitTime.TotalSeconds + " seconds";
+									}
+									else
+									{
+										WaitTimeString = Convert.ToInt32(MeanWaitTime.TotalMinutes) + " minutes";
+									}
+									StartupMessage.AppendLine("Queue waiting time " + WaitTimeString);
+								}
+
+								StatusReportMessage.AppendLine("Queue sizes...");
+								foreach (var Queue in QueueSizes)
+								{
+									StatusReportMessage.AppendLine("> " + Queue.Key + " " + Queue.Value);
+								}
+							}
+							if (CountsInPeriod.Count > 0)
+							{
+								StatusReportMessage.AppendLine(string.Format("Events in the last {0} minutes...",
+									Config.Default.MinutesBetweenQueueSizeReports));
+								foreach (var CountInPeriod in CountsInPeriod)
+								{
+									StatusReportMessage.AppendLine("> " + CountInPeriod.Key + " " + CountInPeriod.Value);
+								}
+							}
+						}
+						CrashReporterProcessServicer.WriteSlack(StatusReportMessage.ToString());
+
+						Monitor.Wait(TaskMonitor, TimeSpan.FromMinutes(Config.Default.MinutesBetweenQueueSizeReports));
+					}
+				}
+			});
+		}
+
+		public void Dispose()
+		{
+			if (StatusReportingTask != null)
+			{
+				CancelSource.Cancel();
+				lock (TaskMonitor)
+				{
+					Monitor.PulseAll(TaskMonitor);
+				}
+				StatusReportingTask.Wait();
+				CancelSource.Dispose();
+			}
+
+			CrashReporterProcessServicer.WriteSlack("CRP stopped");
+		}
+
+		private static Dictionary<string, int> GetCountsInPeriod(Dictionary<string, int> InCounters, Dictionary<string, int> InCountersAtLastReport)
+		{
+			Dictionary<string, int> CountsInPeriod = new Dictionary<string, int>();
+
+			foreach (var Counter in InCounters)
+			{
+				int LastCount = 0;
+				if (!InCountersAtLastReport.TryGetValue(Counter.Key, out LastCount))
+				{
+					InCountersAtLastReport.Add(Counter.Key, 0);
+				}
+				CountsInPeriod.Add(Counter.Key, Counter.Value - LastCount);
+				InCountersAtLastReport[Counter.Key] = Counter.Value;
+			}
+
+			return CountsInPeriod;
+		}
+
+		private Task StatusReportingTask;
+		private readonly CancellationTokenSource CancelSource = new CancellationTokenSource();
+		private readonly Dictionary<string, string> QueueLocations = new Dictionary<string, string>();
+		private readonly Dictionary<string, int> QueueSizes = new Dictionary<string, int>();
+		private readonly Dictionary<string, int> Counters = new Dictionary<string, int>();
+		private int QueueSizeSumAtLastReport = 0;
+		private readonly Dictionary<string, int> CountersAtLastReport = new Dictionary<string, int>();
+		private readonly Object TaskMonitor = new Object();
+		private readonly Object DataLock = new Object();
+		private bool bQueueSizesAvail = false;
+	}
+}
diff --git a/Engine/Source/Programs/CrashReporter/CrashReportProcess/StatusReportingConstants.cs b/Engine/Source/Programs/CrashReporter/CrashReportProcess/StatusReportingConstants.cs
new file mode 100644
--- /dev/null
+++ b/Engine/Source/Programs/CrashReporter/CrashReportProcess/StatusReportingConstants.cs
@@ -0,0 +1,17 @@
+﻿using System;
+using System.Collections.Generic;
+using System.Linq;
+using System.Text;
+using System.Threading.Tasks;
+
+namespace Tools.CrashReporter.CrashReportProcess
+{
+	static class StatusReportingConstants
+	{
+		public const string ProcessingFailedEvent = "Processing failed";
+		public const string ProcessingSucceededEvent = "Processing succeeded";
+		public const string ExceptionEvent = "Handled exceptions";
+		public const string ProcessingStartedEvent = "Started processing";
+		public const string QueuedEvent = "Queued for processing";
+	}
+}
diff --git a/Engine/Source/Programs/CrashReporter/CrashReportWebSite/Controllers/CSVController.cs b/Engine/Source/Programs/CrashReporter/CrashReportWebSite/Controllers/CSVController.cs
--- a/Engine/Source/Programs/CrashReporter/CrashReportWebSite/Controllers/CSVController.cs
+++ b/Engine/Source/Programs/CrashReporter/CrashReportWebSite/Controllers/CSVController.cs
@@ -307,11 +307,29 @@ namespace Tools.CrashReporter.CrashReportWebSite.Controllers
 		{
 			using( FAutoScopedLogTimer LogTimer = new FAutoScopedLogTimer( this.GetType().ToString(), bCreateNewLog: true ) )
 			{
-				FormHelper FormData = new FormHelper( Request, Form, "JustReport" );
-				CSV_ViewModel Results = GetResults( FormData );
-				Results.GenerationTime = LogTimer.GetElapsedSeconds().ToString( "F2" );
-				return View( "Index", Results );
+			    var model = new CSV_ViewModel();
+			    model.DateTimeFrom = DateTime.Now.AddDays(-7);
+			    model.DateTimeTo = DateTime.Now;
+			    model.DateFrom = (long) (model.DateTimeFrom - CrashesViewModel.Epoch).TotalMilliseconds;
+			    model.DateTo = (long) (model.DateTimeTo - CrashesViewModel.Epoch).TotalMilliseconds;
+                return View("Index", model);
 			}
 		}
+
+        /// <summary>
+        /// 
+        /// </summary>
+        /// <param name="Form"></param>
+        /// <returns></returns>
+	    public ActionResult GenerateCSV(FormCollection Form)
+        {
+            using (FAutoScopedLogTimer LogTimer = new FAutoScopedLogTimer(this.GetType().ToString(), bCreateNewLog: true))
+            {
+                FormHelper FormData = new FormHelper(Request, Form, "JustReport");
+                CSV_ViewModel Results = GetResults(FormData);
+                Results.GenerationTime = LogTimer.GetElapsedSeconds().ToString("F2");
+                return View("CSV", Results);
+            }
+        }
 	}
 }
\ No newline at end of file
diff --git a/Engine/Source/Programs/CrashReporter/CrashReportWebSite/Controllers/ReportsController.cs b/Engine/Source/Programs/CrashReporter/CrashReportWebSite/Controllers/ReportsController.cs
--- a/Engine/Source/Programs/CrashReporter/CrashReportWebSite/Controllers/ReportsController.cs
+++ b/Engine/Source/Programs/CrashReporter/CrashReportWebSite/Controllers/ReportsController.cs
@@ -7,11 +7,17 @@ using System.IO;
 using System.Linq;
 using System.Linq.Expressions;
 using System.Net;
+using System.Net.Mail;
 using System.Text;
+using System.Web;
+using System.Web.Helpers;
 using System.Web.Mvc;
+using System.Web.Routing;
+using System.Web.UI;
 using System.Xml;
 using Tools.CrashReporter.CrashReportWebSite.Models;
 using Tools.DotNETCommon;
+using Hangfire;
 
 namespace Tools.CrashReporter.CrashReportWebSite.Controllers
 {
@@ -25,7 +31,7 @@ namespace Tools.CrashReporter.CrashReportWebSite.Controllers
 	//	this engine version and 
 	//	user
 
-	/// <summary>
+    /// <summary>
 	/// The controller to handle graphing of crashes per user group over time.
 	/// </summary>
 	public class ReportsController : Controller
@@ -54,7 +60,7 @@ namespace Tools.CrashReporter.CrashReportWebSite.Controllers
 						break;
 					}
 				}
-
+                
 				ReportsViewModel Results = GetResults( FormData, BuggIDToBeAddedToJira );
 				Results.GenerationTime = LogTimer.GetElapsedSeconds().ToString( "F2" );
 				return View( "Index", Results );
@@ -90,6 +96,35 @@ namespace Tools.CrashReporter.CrashReportWebSite.Controllers
             }
         }
 
+        /// <summary>
+        /// 
+        /// </summary>
+        /// <param name="FormData"></param>
+        /// <returns></returns>
+        public PartialViewResult SubscribeToEmail(string branchName, string emailAddress)
+        {
+            var jobId = string.Format("{0}:{1}", branchName, emailAddress);
+
+            RecurringJob.AddOrUpdate<ReportsController>(jobId, x => x.SendReportsEmail(emailAddress, branchName), Cron.Weekly);
+
+            return PartialView("SignUpResponse", new EmailSubscriptionResponseModel (){ Email = emailAddress, Branch = branchName });
+        }
+        
+        /// <summary>
+        /// 
+        /// </summary>
+        /// <param name="jobId"></param>
+        /// <returns></returns>
+        public ActionResult UnsubscribeFromEmail(string jobId)
+        {
+            RecurringJob.RemoveIfExists(jobId);
+
+            ViewData["BranchName"] = jobId.Split(':')[0];
+            ViewData["EmailAddress"] = jobId.Split(':')[1];
+
+            return View("Unsubscribe");
+        }
+
         #region Private Methods
 
         private void AddBuggJiraMapping(Bugg NewBugg, ref HashSet<string> FoundJiras, ref Dictionary<string, List<Bugg>> JiraIDtoBugg)
@@ -114,348 +149,452 @@ namespace Tools.CrashReporter.CrashReportWebSite.Controllers
 		/// <returns>A view to display the filtered Buggs.</returns>
 		private ReportsViewModel GetResults( FormHelper FormData, int BuggIDToBeAddedToJira )
 		{
-			BuggRepository BuggsRepo = new BuggRepository();
-			CrashRepository CrashRepo = new CrashRepository();
-
-			// It would be great to have a CSV export of this as well with buggs ID being the key I can then use to join them :)
-			// 
-			// Enumerate JIRA projects if needed.
-			// https://jira.ol.epicgames.net//rest/api/2/project
-			var JC = JiraConnection.Get();
-			var JiraComponents = JC.GetNameToComponents();
-			var JiraVersions = JC.GetNameToVersions();
-
-			using( FAutoScopedLogTimer LogTimer = new FAutoScopedLogTimer( this.GetType().ToString() ) )
-			{
-				string AnonumousGroup = "Anonymous";
-				//List<String> Users = FRepository.Get().GetUserNamesFromGroupName( AnonumousGroup );
-				int AnonymousGroupID = FRepository.Get( BuggsRepo ).FindOrAddGroup( AnonumousGroup );
-				List<int> AnonumousIDs = FRepository.Get( BuggsRepo ).GetUserIdsFromUserGroup( AnonumousGroup ).ToList();
-				int AnonymousID = AnonumousIDs.First();
-
-				var CrashesQuery = CrashRepo
-					.QueryByDate( CrashRepo.ListAll(), FormData.DateFrom, FormData.DateTo )
-					// Only crashes and asserts
-					.Where( Crash => Crash.CrashType == 1 || Crash.CrashType == 2 )
-					// Only anonymous user
-                    .Where(Crash => Crash.UserNameId.Value == AnonymousID);
-
-                // Filter by BranchName
-                if (!string.IsNullOrEmpty(FormData.BranchName))
-                {
-                    CrashesQuery =
-                    (
-                        from CrashDetail in CrashesQuery
-                        where CrashDetail.Branch.Equals(FormData.BranchName)
-                        select CrashDetail
-                    );
-                }
-
-                var Crashes = CrashesQuery.Select(Crash => new
-					{
-						ID = Crash.Id,
-						TimeOfCrash = Crash.TimeOfCrash.Value,
-						//UserID = Crash.UserNameId.Value, 
-						BuildVersion = Crash.BuildVersion,
-						JIRA = Crash.Jira,
-						Platform = Crash.PlatformName,
-						FixCL = Crash.FixedChangeList,
-						BuiltFromCL = Crash.BuiltFromCL,
-						Pattern = Crash.Pattern,
-						MachineID = Crash.MachineId,
-						Branch = Crash.Branch,
-						Description = Crash.Description,
-						RawCallStack = Crash.RawCallStack,
-					} )
-					.ToList();
-				int NumCrashes = Crashes.Count;
-
-				/*
-				// Build patterns for crashes where patters is null.
-				var CrashesWithoutPattern = FRepository.Get().Crashes
-					.FilterByDate( FRepository.Get().Crashes.ListAll(), FormData.DateFrom, FormData.DateTo.AddDays( 1 ) )
-					// Only crashes and asserts
-					.Where( Crash => Crash.Pattern == null || Crash.Pattern == "" )
-					.Select( Crash => Crash )
-					.ToList();
-
-				foreach( var Crash in CrashesWithoutPattern )
-				{
-					Crash.BuildPattern();
-				}
-				*/
-
-				// Total # of ALL (Anonymous) crashes in timeframe
-				int TotalAnonymousCrashes = NumCrashes;
-
-				// Total # of UNIQUE (Anonymous) crashes in timeframe
-				HashSet<string> UniquePatterns = new HashSet<string>();
-				HashSet<string> UniqueMachines = new HashSet<string>();
-				Dictionary<string, int> PatternToCount = new Dictionary<string, int>();
-
-				//List<int> CrashesWithoutPattern = new List<int>();
-				//List<DateTime> CrashesWithoutPatternDT = new List<DateTime>();
-
-				foreach( var Crash in Crashes )
-				{
-					if( string.IsNullOrEmpty( Crash.Pattern ) )
-					{
-						//CrashesWithoutPattern.Add( Crash.ID );
-						//CrashesWithoutPatternDT.Add( Crash.TimeOfCrash );
-						continue;
-					}
-
-					UniquePatterns.Add( Crash.Pattern );
-					UniqueMachines.Add( Crash.MachineID );
-
-					bool bAdd = !PatternToCount.ContainsKey( Crash.Pattern );
-					if( bAdd )
-					{
-						PatternToCount.Add( Crash.Pattern, 1 );
-					}
-					else
-					{
-						PatternToCount[Crash.Pattern]++;
-					}
-				}
-				var PatternToCountOrdered = PatternToCount.OrderByDescending( X => X.Value ).ToList();
-				const int NumTopRecords = 200;
-				var PatternAndCount = PatternToCountOrdered.Take( NumTopRecords ).ToDictionary( X => X.Key, Y => Y.Value );
-
-				int TotalUniqueAnonymousCrashes = UniquePatterns.Count;
-
-				// Total # of AFFECTED USERS (Anonymous) in timeframe
-				int TotalAffectedUsers = UniqueMachines.Count;
-
-				var RealBuggs = BuggsRepo.Context.Buggs.Where( Bugg => PatternAndCount.Keys.Contains( Bugg.Pattern ) ).ToList();
-
-				// Build search string.
-				HashSet<string> FoundJiras = new HashSet<string>();
-				Dictionary<string, List<Bugg>> JiraIDtoBugg = new Dictionary<string, List<Bugg>>();
-
-				List<Bugg> Buggs = new List<Bugg>( NumTopRecords );
-				foreach( var Top in PatternAndCount )
-				{
-					Bugg NewBugg = RealBuggs.Where( X => X.Pattern == Top.Key ).FirstOrDefault();
-					if( NewBugg != null )
-					{
-						using( FAutoScopedLogTimer TopTimer = new FAutoScopedLogTimer( string.Format( "{0}:{1}", Buggs.Count + 1, NewBugg.Id ) ) )
-						{
-							var CrashesForBugg = Crashes.Where( Crash => Crash.Pattern == Top.Key ).ToList();
-
-							// Convert anonymous to full type;
-							var FullCrashesForBugg = new List<Crash>( CrashesForBugg.Count );
-							foreach( var Anon in CrashesForBugg )
-							{
-								FullCrashesForBugg.Add( new Crash()
-								{
-									ID = Anon.ID,
-									TimeOfCrash = Anon.TimeOfCrash,
-									BuildVersion = Anon.BuildVersion,
-									Jira = Anon.JIRA,
-									Platform = Anon.Platform,
-									FixCL = Anon.FixCL,
-									BuiltFromCL = Anon.BuiltFromCL,
-									Pattern = Anon.Pattern,
-									MachineId = Anon.MachineID,
-									Branch = Anon.Branch,
-									Description = Anon.Description,
-									RawCallStack = Anon.RawCallStack,
-								} );
-							}
-
-							NewBugg.PrepareBuggForJira( FullCrashesForBugg );
-
-							// Verify valid JiraID, this may be still a TTP 
-							if( !string.IsNullOrEmpty( NewBugg.Jira ) )
-							{
-								int TTPID = 0;
-								int.TryParse( NewBugg.Jira, out TTPID );
-
-								if( TTPID == 0 )
-								{
-									AddBuggJiraMapping(NewBugg, ref FoundJiras, ref JiraIDtoBugg);						
-								}
-							}
-
-							Buggs.Add( NewBugg );
-						}
-					}
-					else
-					{
-						FLogger.Global.WriteEvent( "Bugg for pattern " + Top.Key + " not found" );
-					}
-				}
-
-				if( BuggIDToBeAddedToJira > 0 )
-				{
-					var Bugg = Buggs.Where( X => X.Id == BuggIDToBeAddedToJira ).FirstOrDefault();
-					if( Bugg != null )
-					{
-						Bugg.CopyToJira();
-						AddBuggJiraMapping( Bugg, ref FoundJiras, ref JiraIDtoBugg );
-					}
-				}
+            var branchName = FormData.BranchName ?? "";
+            var startDate = FormData.DateFrom;
+            var endDate = FormData.DateTo;
 
-				if( JC.CanBeUsed() )
-				{
-					var BuggsCopy = new List<Bugg>( Buggs );
-
-					HashSet<string> InvalidJiras = new HashSet<string>();
+            return GetResults(branchName, startDate, endDate, BuggIDToBeAddedToJira);
+		}
 
-					// Simple verification of JIRA
-					foreach (var Value in FoundJiras)
-					{
-						if( Value.Length < 3 || !Value.Contains('-') )
-						{
-							InvalidJiras.Add(Value);
-						}
-					}
+	    /// <summary>
+        /// Retrieve all Buggs matching the search criteria.
+	    /// </summary>
+	    /// <returns></returns>
+        private ReportsViewModel GetResults(string branchName, DateTime startDate, DateTime endDate, int BuggIDToBeAddedToJira)
+	    {
+            BuggRepository BuggsRepo = new BuggRepository();
+            CrashRepository CrashRepo = new CrashRepository();
+
+	        // It would be great to have a CSV export of this as well with buggs ID being the key I can then use to join them :)
+	        // 
+	        // Enumerate JIRA projects if needed.
+	        // https://jira.ol.epicgames.net//rest/api/2/project
+	        var JC = JiraConnection.Get();
+	        var JiraComponents = JC.GetNameToComponents();
+	        var JiraVersions = JC.GetNameToVersions();
+
+	        using (FAutoScopedLogTimer LogTimer = new FAutoScopedLogTimer(this.GetType().ToString()))
+	        {
+	            string AnonumousGroup = "Anonymous";
+	            //List<String> Users = FRepository.Get().GetUserNamesFromGroupName( AnonumousGroup );
+	            int AnonymousGroupID = FRepository.Get(BuggsRepo).FindOrAddGroup(AnonumousGroup);
+	            List<int> AnonumousIDs = FRepository.Get(BuggsRepo).GetUserIdsFromUserGroup(AnonumousGroup).ToList();
+	            int AnonymousID = AnonumousIDs.First();
+
+	            var CrashesQuery = CrashRepo
+	                .QueryByDate(CrashRepo.ListAll(), startDate, endDate)
+	                // Only crashes and asserts
+	                .Where(Crash => Crash.CrashType == 1 || Crash.CrashType == 2)
+	                // Only anonymous user
+	                .Where(Crash => Crash.UserNameId.Value == AnonymousID);
+
+	            // Filter by BranchName
+	            if (!string.IsNullOrEmpty(branchName))
+	            {
+	                CrashesQuery =
+	                    (
+	                        from CrashDetail in CrashesQuery
+	                        where CrashDetail.Branch.Equals(branchName)
+	                        select CrashDetail
+	                        );
+	            }
+
+	            var Crashes = CrashesQuery.Select(Crash => new
+	            {
+	                ID = Crash.Id,
+	                TimeOfCrash = Crash.TimeOfCrash.Value,
+	                //UserID = Crash.UserNameId.Value, 
+	                BuildVersion = Crash.BuildVersion,
+	                JIRA = Crash.Jira,
+	                Platform = Crash.PlatformName,
+	                FixCL = Crash.FixedChangeList,
+	                BuiltFromCL = Crash.BuiltFromCL,
+	                Pattern = Crash.Pattern,
+	                MachineID = Crash.MachineId,
+	                Branch = Crash.Branch,
+	                Description = Crash.Description,
+	                RawCallStack = Crash.RawCallStack,
+	            })
+	                .ToList();
+	            int NumCrashes = Crashes.Count;
+
+	            /*
+                // Build patterns for crashes where patters is null.
+                var CrashesWithoutPattern = FRepository.Get().Crashes
+                    .FilterByDate( FRepository.Get().Crashes.ListAll(), startDate, endDate.AddDays( 1 ) )
+                    // Only crashes and asserts
+                    .Where( Crash => Crash.Pattern == null || Crash.Pattern == "" )
+                    .Select( Crash => Crash )
+                    .ToList();
+
+                foreach( var Crash in CrashesWithoutPattern )
+                {
+                    Crash.BuildPattern();
+                }
+                */
+
+	            // Total # of ALL (Anonymous) crashes in timeframe
+	            int TotalAnonymousCrashes = NumCrashes;
+
+	            // Total # of UNIQUE (Anonymous) crashes in timeframe
+	            HashSet<string> UniquePatterns = new HashSet<string>();
+	            HashSet<string> UniqueMachines = new HashSet<string>();
+	            Dictionary<string, int> PatternToCount = new Dictionary<string, int>();
+
+	            //List<int> CrashesWithoutPattern = new List<int>();
+	            //List<DateTime> CrashesWithoutPatternDT = new List<DateTime>();
+
+	            foreach (var Crash in Crashes)
+	            {
+	                if (string.IsNullOrEmpty(Crash.Pattern))
+	                {
+	                    //CrashesWithoutPattern.Add( Crash.ID );
+	                    //CrashesWithoutPatternDT.Add( Crash.TimeOfCrash );
+	                    continue;
+	                }
+
+	                UniquePatterns.Add(Crash.Pattern);
+	                UniqueMachines.Add(Crash.MachineID);
+
+	                bool bAdd = !PatternToCount.ContainsKey(Crash.Pattern);
+	                if (bAdd)
+	                {
+	                    PatternToCount.Add(Crash.Pattern, 1);
+	                }
+	                else
+	                {
+	                    PatternToCount[Crash.Pattern]++;
+	                }
+	            }
+	            var PatternToCountOrdered = PatternToCount.OrderByDescending(X => X.Value).ToList();
+	            const int NumTopRecords = 200;
+	            var PatternAndCount = PatternToCountOrdered.Take(NumTopRecords).ToDictionary(X => X.Key, Y => Y.Value);
+
+	            int TotalUniqueAnonymousCrashes = UniquePatterns.Count;
+
+	            // Total # of AFFECTED USERS (Anonymous) in timeframe
+	            int TotalAffectedUsers = UniqueMachines.Count;
+
+	            var RealBuggs = BuggsRepo.Context.Buggs.Where(Bugg => PatternAndCount.Keys.Contains(Bugg.Pattern)).ToList();
+
+	            // Build search string.
+	            HashSet<string> FoundJiras = new HashSet<string>();
+	            Dictionary<string, List<Bugg>> JiraIDtoBugg = new Dictionary<string, List<Bugg>>();
+
+	            List<Bugg> Buggs = new List<Bugg>(NumTopRecords);
+	            foreach (var Top in PatternAndCount)
+	            {
+	                Bugg NewBugg = RealBuggs.Where(X => X.Pattern == Top.Key).FirstOrDefault();
+	                if (NewBugg != null)
+	                {
+	                    using (
+	                        FAutoScopedLogTimer TopTimer =
+	                            new FAutoScopedLogTimer(string.Format("{0}:{1}", Buggs.Count + 1, NewBugg.Id)))
+	                    {
+	                        var CrashesForBugg = Crashes.Where(Crash => Crash.Pattern == Top.Key).ToList();
+
+	                        // Convert anonymous to full type;
+	                        var FullCrashesForBugg = new List<Crash>(CrashesForBugg.Count);
+	                        foreach (var Anon in CrashesForBugg)
+	                        {
+	                            FullCrashesForBugg.Add(new Crash()
+	                            {
+	                                ID = Anon.ID,
+	                                TimeOfCrash = Anon.TimeOfCrash,
+	                                BuildVersion = Anon.BuildVersion,
+	                                Jira = Anon.JIRA,
+	                                Platform = Anon.Platform,
+	                                FixCL = Anon.FixCL,
+	                                BuiltFromCL = Anon.BuiltFromCL,
+	                                Pattern = Anon.Pattern,
+	                                MachineId = Anon.MachineID,
+	                                Branch = Anon.Branch,
+	                                Description = Anon.Description,
+	                                RawCallStack = Anon.RawCallStack,
+	                            });
+	                        }
+
+	                        NewBugg.PrepareBuggForJira(FullCrashesForBugg);
+
+	                        // Verify valid JiraID, this may be still a TTP 
+	                        if (!string.IsNullOrEmpty(NewBugg.Jira))
+	                        {
+	                            int TTPID = 0;
+	                            int.TryParse(NewBugg.Jira, out TTPID);
+
+	                            if (TTPID == 0)
+	                            {
+	                                AddBuggJiraMapping(NewBugg, ref FoundJiras, ref JiraIDtoBugg);
+	                            }
+	                        }
+
+	                        Buggs.Add(NewBugg);
+	                    }
+	                }
+	                else
+	                {
+	                    FLogger.Global.WriteEvent("Bugg for pattern " + Top.Key + " not found");
+	                }
+	            }
+
+	            if (BuggIDToBeAddedToJira > 0)
+	            {
+	                var Bugg = Buggs.Where(X => X.Id == BuggIDToBeAddedToJira).FirstOrDefault();
+	                if (Bugg != null)
+	                {
+	                    Bugg.CopyToJira();
+	                    AddBuggJiraMapping(Bugg, ref FoundJiras, ref JiraIDtoBugg);
+	                }
+	            }
+
+	            if (JC.CanBeUsed())
+	            {
+	                var BuggsCopy = new List<Bugg>(Buggs);
+
+	                HashSet<string> InvalidJiras = new HashSet<string>();
+
+	                // Simple verification of JIRA
+	                foreach (var Value in FoundJiras)
+	                {
+	                    if (Value.Length < 3 || !Value.Contains('-'))
+	                    {
+	                        InvalidJiras.Add(Value);
+	                    }
+	                }
+
+	                foreach (var InvalidJira in InvalidJiras)
+	                {
+	                    FoundJiras.Remove(InvalidJira);
+	                }
+
+	                // Grab the data form JIRA.
+	                string JiraSearchQuery = string.Join(" OR ", FoundJiras);
+
+	                using (FAutoScopedLogTimer JiraResultsTimer = new FAutoScopedLogTimer("JiraResults"))
+	                {
+	                    bool bInvalid = false;
+	                    var JiraResults = new Dictionary<string, Dictionary<string, object>>();
+	                    try
+	                    {
+	                        if (!string.IsNullOrWhiteSpace(JiraSearchQuery))
+	                        {
+	                            JiraResults = JC.SearchJiraTickets(
+	                                JiraSearchQuery,
+	                                new string[]
+	                                {
+	                                    "key", // string
+	                                    "summary", // string
+	                                    "components", // System.Collections.ArrayList, Dictionary<string,object>, name
+	                                    "resolution",
+	                                    // System.Collections.Generic.Dictionary`2[System.String,System.Object], name
+	                                    "fixVersions", // System.Collections.ArrayList, Dictionary<string,object>, name
+	                                    "customfield_11200" // string
+	                                });
+	                        }
+	                    }
+	                    catch (System.Exception)
+	                    {
+	                        bInvalid = true;
+	                    }
+
+	                    // Invalid records have been found, find the broken using the slow path.
+	                    if (bInvalid)
+	                    {
+	                        foreach (var Query in FoundJiras)
+	                        {
+	                            try
+	                            {
+	                                var TempResult = JC.SearchJiraTickets(
+	                                    Query,
+	                                    new string[]
+	                                    {
+	                                        "key", // string
+	                                        "summary", // string
+	                                        "components", // System.Collections.ArrayList, Dictionary<string,object>, name
+	                                        "resolution",
+	                                        // System.Collections.Generic.Dictionary`2[System.String,System.Object], name
+	                                        "fixVersions", // System.Collections.ArrayList, Dictionary<string,object>, name
+	                                        "customfield_11200" // string
+	                                    });
+
+	                                foreach (var Temp in TempResult)
+	                                {
+	                                    JiraResults.Add(Temp.Key, Temp.Value);
+	                                }
+	                            }
+	                            catch (System.Exception)
+	                            {
+
+	                            }
+	                        }
+	                    }
+
+	                    // Jira Key, Summary, Components, Resolution, Fix version, Fix changelist
+	                    foreach (var Jira in JiraResults)
+	                    {
+	                        string JiraID = Jira.Key;
+
+	                        string Summary = (string) Jira.Value["summary"];
+
+	                        string ComponentsText = "";
+	                        System.Collections.ArrayList Components =
+	                            (System.Collections.ArrayList) Jira.Value["components"];
+	                        foreach (Dictionary<string, object> Component in Components)
+	                        {
+	                            ComponentsText += (string) Component["name"];
+	                            ComponentsText += " ";
+	                        }
+
+	                        Dictionary<string, object> ResolutionFields =
+	                            (Dictionary<string, object>) Jira.Value["resolution"];
+	                        string Resolution = ResolutionFields != null ? (string) ResolutionFields["name"] : "";
+
+	                        string FixVersionsText = "";
+	                        System.Collections.ArrayList FixVersions =
+	                            (System.Collections.ArrayList) Jira.Value["fixVersions"];
+	                        foreach (Dictionary<string, object> FixVersion in FixVersions)
+	                        {
+	                            FixVersionsText += (string) FixVersion["name"];
+	                            FixVersionsText += " ";
+	                        }
+
+	                        int FixCL = Jira.Value["customfield_11200"] != null
+	                            ? (int) (decimal) Jira.Value["customfield_11200"]
+	                            : 0;
+
+	                        List<Bugg> BuggsForJira;
+	                        JiraIDtoBugg.TryGetValue(JiraID, out BuggsForJira);
+
+	                        //var BuggsForJira = JiraIDtoBugg[JiraID];
+
+	                        if (BuggsForJira != null)
+	                        {
+	                            foreach (Bugg Bugg in BuggsForJira)
+	                            {
+	                                Bugg.JiraSummary = Summary;
+	                                Bugg.JiraComponentsText = ComponentsText;
+	                                Bugg.JiraResolution = Resolution;
+	                                Bugg.JiraFixVersionsText = FixVersionsText;
+	                                if (FixCL != 0)
+	                                {
+	                                    Bugg.JiraFixCL = FixCL.ToString();
+	                                }
+
+	                                BuggsCopy.Remove(Bugg);
+	                            }
+	                        }
+	                    }
+	                }
+
+	                // If there are buggs, we need to update the summary to indicate an error.
+	                // Usually caused when bugg's project has changed.
+	                foreach (var Bugg in BuggsCopy.Where(b => !string.IsNullOrEmpty(b.Jira)))
+	                {
+	                    Bugg.JiraSummary = "JIRA MISMATCH";
+	                    Bugg.JiraComponentsText = "JIRA MISMATCH";
+	                    Bugg.JiraResolution = "JIRA MISMATCH";
+	                    Bugg.JiraFixVersionsText = "JIRA MISMATCH";
+	                    Bugg.JiraFixCL = "JIRA MISMATCH";
+	                }
+	            }
+
+	            Buggs = Buggs.OrderByDescending(b => b.CrashesInTimeFrameGroup).ToList();
+
+	            return new ReportsViewModel
+	            {
+	                Buggs = Buggs,
+	                /*Crashes = Crashes,*/
+	                DateFrom = (long) (startDate - CrashesViewModel.Epoch).TotalMilliseconds,
+	                DateTo = (long) (endDate - CrashesViewModel.Epoch).TotalMilliseconds,
+	                TotalAffectedUsers = TotalAffectedUsers,
+	                TotalAnonymousCrashes = TotalAnonymousCrashes,
+	                TotalUniqueAnonymousCrashes = TotalUniqueAnonymousCrashes
+	            };
+	        }
+	    }
 
-					foreach (var InvalidJira in InvalidJiras)
-					{
-						FoundJiras.Remove( InvalidJira );
-					}
+        /// <summary>
+        /// 
+        /// </summary>
+        /// <param name="emailAddress"></param>
+        /// <param name="branchName"></param>
+	    public void SendReportsEmail(string emailAddress, string branchName)
+	    {
+            var emailBody = string.Format("{0}" + Environment.NewLine + "{1}{2}%3A{3}",
+                GetReportsEmailContents(branchName),
+                "http://crashreporter.epicgames.net/Reports/UnsubscribeFromEmail?jobId=",
+                branchName,
+                emailAddress);
+            var sMail = new SmtpClient();
 
-					// Grab the data form JIRA.
-					string JiraSearchQuery = string.Join( " OR ", FoundJiras );
+            var message = new MailMessage("[REDACTED_EMAIL]", emailAddress, "Weekly Crash Report : " + branchName, emailBody) { IsBodyHtml = true };
+            sMail.Send(message);
+	    }
 
-					using( FAutoScopedLogTimer JiraResultsTimer = new FAutoScopedLogTimer( "JiraResults" ) )
-					{
-						bool bInvalid = false;
-						var JiraResults = new Dictionary<string, Dictionary<string, object>>();
-						try
-						{
-						    if (!string.IsNullOrWhiteSpace(JiraSearchQuery))
-						    {
-                                JiraResults = JC.SearchJiraTickets(
-                                    JiraSearchQuery,
-                                    new string[] 
-                                    { 
-                                        "key",				// string
-								        "summary",			// string
-								        "components",		// System.Collections.ArrayList, Dictionary<string,object>, name
-								        "resolution",		// System.Collections.Generic.Dictionary`2[System.String,System.Object], name
-								        "fixVersions",		// System.Collections.ArrayList, Dictionary<string,object>, name
-								        "customfield_11200" // string
-							        });
-                            }
-						}
-						catch (System.Exception)
-						{
-							bInvalid = true;
-						}
-
-						// Invalid records have been found, find the broken using the slow path.
-						if( bInvalid )
-						{
-							foreach (var Query in FoundJiras)
-							{
-								try
-								{
-									var TempResult = JC.SearchJiraTickets(
-									Query,
-									new string[] 
-									{ 
-										"key",				// string
-										"summary",			// string
-										"components",		// System.Collections.ArrayList, Dictionary<string,object>, name
-										"resolution",		// System.Collections.Generic.Dictionary`2[System.String,System.Object], name
-										"fixVersions",		// System.Collections.ArrayList, Dictionary<string,object>, name
-										"customfield_11200" // string
-									} );
-
-									foreach(var Temp in TempResult)
-									{
-										JiraResults.Add( Temp.Key, Temp.Value );
-									}
-								}
-								catch (System.Exception)
-								{
-
-								}
-							}
-						}
-
-						// Jira Key, Summary, Components, Resolution, Fix version, Fix changelist
-						foreach( var Jira in JiraResults )
-						{
-							string JiraID = Jira.Key;
-
-							string Summary = (string)Jira.Value["summary"];
-
-							string ComponentsText = "";
-							System.Collections.ArrayList Components = (System.Collections.ArrayList)Jira.Value["components"];
-							foreach( Dictionary<string, object> Component in Components )
-							{
-								ComponentsText += (string)Component["name"];
-								ComponentsText += " ";
-							}
-
-							Dictionary<string, object> ResolutionFields = (Dictionary<string, object>)Jira.Value["resolution"];
-							string Resolution = ResolutionFields != null ? (string)ResolutionFields["name"] : "";
-
-							string FixVersionsText = "";
-							System.Collections.ArrayList FixVersions = (System.Collections.ArrayList)Jira.Value["fixVersions"];
-							foreach( Dictionary<string, object> FixVersion in FixVersions )
-							{
-								FixVersionsText += (string)FixVersion["name"];
-								FixVersionsText += " ";
-							}
-
-							int FixCL = Jira.Value["customfield_11200"] != null ? (int)(decimal)Jira.Value["customfield_11200"] : 0;
-
-							List<Bugg> BuggsForJira;
-							JiraIDtoBugg.TryGetValue( JiraID, out BuggsForJira );
-
-							//var BuggsForJira = JiraIDtoBugg[JiraID];
-
-							if( BuggsForJira != null )
-							{
-								foreach( Bugg Bugg in BuggsForJira )
-								{
-									Bugg.JiraSummary = Summary;
-									Bugg.JiraComponentsText = ComponentsText;
-									Bugg.JiraResolution = Resolution;
-									Bugg.JiraFixVersionsText = FixVersionsText;
-									if( FixCL != 0 )
-									{
-										Bugg.JiraFixCL = FixCL.ToString();
-									}
-
-									BuggsCopy.Remove( Bugg );
-								}
-							}
-						}
-					}
+        /// <summary>
+        /// 
+        /// </summary>
+        /// <returns></returns>
+	    private string GetReportsEmailContents(string branchName)
+	    {
+            var reportViewModel = GetResults(branchName, DateTime.Now.AddDays(-7), DateTime.Now, 0);
+            return RenderViewToString("Reports", "~/Views/Reports/ViewReports.ascx", reportViewModel);
+	    }
 
-					// If there are buggs, we need to update the summary to indicate an error.
-					// Usually caused when bugg's project has changed.
-					foreach( var Bugg in BuggsCopy.Where( b => !string.IsNullOrEmpty( b.Jira ) ) )
-					{
-						Bugg.JiraSummary = "JIRA MISMATCH";
-						Bugg.JiraComponentsText = "JIRA MISMATCH";
-						Bugg.JiraResolution = "JIRA MISMATCH";
-						Bugg.JiraFixVersionsText = "JIRA MISMATCH";
-						Bugg.JiraFixCL = "JIRA MISMATCH";
-					}
-				}
+        /// <summary>
+        /// This function uses a temporary controller object to return a view as a string object rather than passing it to the browser.
+        /// </summary>
+        /// <param name="controllerName">the name of the controller containing the view</param>
+        /// <param name="viewName">the name of the view</param>
+        /// <param name="viewData">Any datamodels associated with this view</param>
+        /// <returns></returns>
+        private static string RenderPartialToString(string viewName, object model, ControllerContext ControllerContext)
+        {
+            if (string.IsNullOrEmpty(viewName))
+                viewName = ControllerContext.RouteData.GetRequiredString("action");
+            var viewData = new ViewDataDictionary();
+            var tempData = new TempDataDictionary();
+            viewData.Model = model;
 
-				Buggs = Buggs.OrderByDescending( b => b.CrashesInTimeFrameGroup ).ToList();
+            using (var sw = new StringWriter())
+            {
+                var viewResult = ViewEngines.Engines.FindPartialView(ControllerContext, viewName);
+                var viewContext = new ViewContext(ControllerContext, viewResult.View, viewData, tempData, sw);
+                viewResult.View.Render(viewContext, sw);
 
-				return new ReportsViewModel
-				{
-					Buggs = Buggs,
-					/*Crashes = Crashes,*/
-					DateFrom = (long)( FormData.DateFrom - CrashesViewModel.Epoch ).TotalMilliseconds,
-					DateTo = (long)( FormData.DateTo - CrashesViewModel.Epoch ).TotalMilliseconds,
-					TotalAffectedUsers = TotalAffectedUsers,
-					TotalAnonymousCrashes = TotalAnonymousCrashes,
-					TotalUniqueAnonymousCrashes = TotalUniqueAnonymousCrashes
-				};
-			}
-		}
+                return sw.GetStringBuilder().ToString();
+            }
+        }
+        
+        /// <summary>
+        /// 
+        /// </summary>
+        /// <param name="controllerName"></param>
+        /// <param name="viewName"></param>
+        /// <param name="viewData"></param>
+        /// <returns></returns>
+        private static string RenderViewToString(string controllerName, string viewName, object viewData)
+        {
+            using (var writer = new StringWriter())
+            {
+                var routeData = new RouteData();
+                routeData.Values.Add("controller", controllerName);
+                var fakeControllerContext = new ControllerContext(new HttpContextWrapper(new HttpContext(new HttpRequest(null, "http://google.com", null), new HttpResponse(null))), routeData, new FakeController());
+                var razorViewResult = ViewEngines.Engines.FindPartialView(fakeControllerContext, viewName);
+
+                var viewContext = new ViewContext(fakeControllerContext, razorViewResult.View, new ViewDataDictionary(viewData), new TempDataDictionary(), writer);
+                razorViewResult.View.Render(viewContext, writer);
+                return writer.ToString();
+            }
+        }
 
         #endregion
     }
+
+    class FakeController : ControllerBase
+    {
+        protected override void ExecuteCore()
+        {
+        }
+    };
 }
\ No newline at end of file
diff --git a/Engine/Source/Programs/CrashReporter/CrashReportWebSite/CrashReportWebSite.csproj b/Engine/Source/Programs/CrashReporter/CrashReportWebSite/CrashReportWebSite.csproj
--- a/Engine/Source/Programs/CrashReporter/CrashReportWebSite/CrashReportWebSite.csproj
+++ b/Engine/Source/Programs/CrashReporter/CrashReportWebSite/CrashReportWebSite.csproj
@@ -7,7 +7,7 @@
     <ProductVersion>9.0.30729</ProductVersion>
     <SchemaVersion>2.0</SchemaVersion>
     <ProjectGuid>{7A6D44E8-D8F2-4540-8A4A-0EF1A87D65AB}</ProjectGuid>
-    <ProjectTypeGuids>{E53F8FEA-EAE0-44A6-8774-FFD645390401};{349c5851-65df-11da-9384-00065b846f21};{fae04ec0-301f-11d3-bf4b-00c04f79efbc}</ProjectTypeGuids>
+    <ProjectTypeGuids>{349c5851-65df-11da-9384-00065b846f21};{fae04ec0-301f-11d3-bf4b-00c04f79efbc}</ProjectTypeGuids>
     <OutputType>Library</OutputType>
     <AppDesignerFolder>Properties</AppDesignerFolder>
     <RootNamespace>Tools.CrashReporter.CrashReportWebSite</RootNamespace>
@@ -46,10 +46,73 @@
     <Prefer32Bit>false</Prefer32Bit>
   </PropertyGroup>
   <ItemGroup>
+    <Reference Include="Common.Logging, Version=3.3.1.0, Culture=neutral, PublicKeyToken=af08829b84f0328e, processorArchitecture=MSIL">
+      <HintPath>packages\Common.Logging.3.3.1\lib\net40\Common.Logging.dll</HintPath>
+      <Private>True</Private>
+    </Reference>
+    <Reference Include="Common.Logging.Core, Version=3.3.1.0, Culture=neutral, PublicKeyToken=af08829b84f0328e, processorArchitecture=MSIL">
+      <HintPath>packages\Common.Logging.Core.3.3.1\lib\net40\Common.Logging.Core.dll</HintPath>
+      <Private>True</Private>
+    </Reference>
+    <Reference Include="Hangfire.Core, Version=1.5.3.0, Culture=neutral, processorArchitecture=MSIL">
+      <HintPath>packages\Hangfire.Core.1.5.3\lib\net45\Hangfire.Core.dll</HintPath>
+      <Private>True</Private>
+    </Reference>
+    <Reference Include="Hangfire.SqlServer, Version=1.5.3.0, Culture=neutral, processorArchitecture=MSIL">
+      <HintPath>packages\Hangfire.SqlServer.1.5.3\lib\net45\Hangfire.SqlServer.dll</HintPath>
+      <Private>True</Private>
+    </Reference>
+    <Reference Include="Hangfire.Unity, Version=1.0.0.0, Culture=neutral, processorArchitecture=MSIL">
+      <HintPath>packages\Hangfire.Unity.1.1.0\lib\net45\Hangfire.Unity.dll</HintPath>
+      <Private>True</Private>
+    </Reference>
+    <Reference Include="Microsoft.Owin, Version=3.0.0.0, Culture=neutral, PublicKeyToken=31bf3856ad364e35, processorArchitecture=MSIL">
+      <HintPath>packages\Microsoft.Owin.3.0.0\lib\net45\Microsoft.Owin.dll</HintPath>
+      <Private>True</Private>
+    </Reference>
+    <Reference Include="Microsoft.Owin.Host.SystemWeb, Version=3.0.0.0, Culture=neutral, PublicKeyToken=31bf3856ad364e35, processorArchitecture=MSIL">
+      <HintPath>packages\Microsoft.Owin.Host.SystemWeb.3.0.0\lib\net45\Microsoft.Owin.Host.SystemWeb.dll</HintPath>
+      <Private>True</Private>
+    </Reference>
+    <Reference Include="Microsoft.Practices.ServiceLocation, Version=1.3.0.0, Culture=neutral, PublicKeyToken=31bf3856ad364e35, processorArchitecture=MSIL">
+      <HintPath>packages\CommonServiceLocator.1.3\lib\portable-net4+sl5+netcore45+wpa81+wp8\Microsoft.Practices.ServiceLocation.dll</HintPath>
+      <Private>True</Private>
+    </Reference>
+    <Reference Include="Microsoft.Practices.Unity, Version=4.0.0.0, Culture=neutral, PublicKeyToken=6d32ff45e0ccc69f, processorArchitecture=MSIL">
+      <HintPath>packages\Unity.4.0.1\lib\net45\Microsoft.Practices.Unity.dll</HintPath>
+      <Private>True</Private>
+    </Reference>
+    <Reference Include="Microsoft.Practices.Unity.Configuration, Version=4.0.0.0, Culture=neutral, PublicKeyToken=6d32ff45e0ccc69f, processorArchitecture=MSIL">
+      <HintPath>packages\Unity.4.0.1\lib\net45\Microsoft.Practices.Unity.Configuration.dll</HintPath>
+      <Private>True</Private>
+    </Reference>
+    <Reference Include="Microsoft.Practices.Unity.Mvc">
+      <HintPath>packages\Unity.Mvc.4.0.1\lib\net45\Microsoft.Practices.Unity.Mvc.dll</HintPath>
+      <Private>True</Private>
+    </Reference>
+    <Reference Include="Microsoft.Practices.Unity.RegistrationByConvention, Version=4.0.0.0, Culture=neutral, PublicKeyToken=6d32ff45e0ccc69f, processorArchitecture=MSIL">
+      <HintPath>packages\Unity.4.0.1\lib\net45\Microsoft.Practices.Unity.RegistrationByConvention.dll</HintPath>
+      <Private>True</Private>
+    </Reference>
+    <Reference Include="Microsoft.Web.Infrastructure, Version=1.0.0.0, Culture=neutral, PublicKeyToken=31bf3856ad364e35, processorArchitecture=MSIL">
+      <HintPath>packages\Microsoft.Web.Infrastructure.1.0.0.0\lib\net40\Microsoft.Web.Infrastructure.dll</HintPath>
+      <Private>True</Private>
+    </Reference>
+    <Reference Include="Newtonsoft.Json, Version=4.5.0.0, Culture=neutral, PublicKeyToken=30ad4fe6b2a6aeed, processorArchitecture=MSIL">
+      <HintPath>packages\Newtonsoft.Json.5.0.1\lib\net45\Newtonsoft.Json.dll</HintPath>
+      <Private>True</Private>
+    </Reference>
+    <Reference Include="Owin, Version=1.0.0.0, Culture=neutral, PublicKeyToken=f0ebd12fd5e55cc5, processorArchitecture=MSIL">
+      <HintPath>packages\Owin.1.0\lib\net40\Owin.dll</HintPath>
+      <Private>True</Private>
+    </Reference>
     <Reference Include="System" />
     <Reference Include="System.Data" />
     <Reference Include="System.Data.DataSetExtensions" />
-    <Reference Include="System.Data.Linq" />
+    <Reference Include="System.Data.Entity" />
+    <Reference Include="System.Data.Linq">
+      <Private>True</Private>
+    </Reference>
     <Reference Include="System.Drawing" />
     <Reference Include="System.ComponentModel.DataAnnotations">
       <RequiredTargetFramework>3.5</RequiredTargetFramework>
@@ -60,30 +123,54 @@
     <Reference Include="System.Web" />
     <Reference Include="System.Web.Abstractions" />
     <Reference Include="System.Web.Extensions" />
-    <Reference Include="System.Web.Helpers, Version=2.0.0.0, Culture=neutral, PublicKeyToken=31bf3856ad364e35, processorArchitecture=MSIL">
+    <Reference Include="System.Web.Helpers, Version=3.0.0.0, Culture=neutral, PublicKeyToken=31bf3856ad364e35, processorArchitecture=MSIL">
       <SpecificVersion>False</SpecificVersion>
-      <HintPath>C:\Program Files (x86)\Microsoft ASP.NET\ASP.NET Web Pages\v1.0\Assemblies\System.Web.Helpers.dll</HintPath>
+      <HintPath>packages\Microsoft.AspNet.WebPages.3.2.3\lib\net45\System.Web.Helpers.dll</HintPath>
+      <Private>True</Private>
     </Reference>
-    <Reference Include="System.Web.Mvc, Version=4.0.0.0, Culture=neutral, PublicKeyToken=31bf3856ad364e35, processorArchitecture=MSIL">
+    <Reference Include="System.Web.Mvc, Version=5.2.3.0, Culture=neutral, PublicKeyToken=31bf3856ad364e35, processorArchitecture=MSIL">
       <SpecificVersion>False</SpecificVersion>
-      <HintPath>C:\Program Files (x86)\Microsoft ASP.NET\ASP.NET MVC 3\Assemblies\System.Web.Mvc.dll</HintPath>
+      <HintPath>packages\Microsoft.AspNet.Mvc.5.2.3\lib\net45\System.Web.Mvc.dll</HintPath>
+      <Private>True</Private>
+    </Reference>
+    <Reference Include="System.Web.Razor, Version=3.0.0.0, Culture=neutral, PublicKeyToken=31bf3856ad364e35, processorArchitecture=MSIL">
+      <HintPath>packages\Microsoft.AspNet.Razor.3.2.3\lib\net45\System.Web.Razor.dll</HintPath>
+      <Private>True</Private>
+    </Reference>
+    <Reference Include="System.Web.Routing">
+      <Private>True</Private>
     </Reference>
-    <Reference Include="System.Web.Routing" />
     <Reference Include="System.Web.WebPages, Version=2.0.0.0, Culture=neutral, PublicKeyToken=31bf3856ad364e35, processorArchitecture=MSIL">
       <SpecificVersion>False</SpecificVersion>
-      <HintPath>C:\Program Files (x86)\Microsoft ASP.NET\ASP.NET Web Pages\v1.0\Assemblies\System.Web.WebPages.dll</HintPath>
+      <HintPath>packages\Microsoft.AspNet.WebPages.3.2.3\lib\net45\System.Web.WebPages.dll</HintPath>
+    </Reference>
+    <Reference Include="System.Web.WebPages.Deployment, Version=3.0.0.0, Culture=neutral, PublicKeyToken=31bf3856ad364e35, processorArchitecture=MSIL">
+      <HintPath>packages\Microsoft.AspNet.WebPages.3.2.3\lib\net45\System.Web.WebPages.Deployment.dll</HintPath>
+      <Private>True</Private>
+    </Reference>
+    <Reference Include="System.Web.WebPages.Razor, Version=3.0.0.0, Culture=neutral, PublicKeyToken=31bf3856ad364e35, processorArchitecture=MSIL">
+      <HintPath>packages\Microsoft.AspNet.WebPages.3.2.3\lib\net45\System.Web.WebPages.Razor.dll</HintPath>
+      <Private>True</Private>
     </Reference>
     <Reference Include="System.Xml" />
     <Reference Include="System.Configuration" />
-    <Reference Include="System.Web.Services" />
+    <Reference Include="System.Web.Services">
+      <Private>True</Private>
+    </Reference>
     <Reference Include="System.EnterpriseServices" />
     <Reference Include="System.Web.Mobile" />
     <Reference Include="System.Xml.Linq" />
+    <Reference Include="WebActivatorEx">
+      <HintPath>packages\WebActivatorEx.2.0\lib\net40\WebActivatorEx.dll</HintPath>
+      <Private>True</Private>
+    </Reference>
   </ItemGroup>
   <ItemGroup>
     <Compile Include="..\..\DotNETCommon\MetaData.cs">
       <Link>Properties\MetaData.cs</Link>
     </Compile>
+    <Compile Include="App_Start\UnityConfig.cs" />
+    <Compile Include="App_Start\UnityMvcActivator.cs" />
     <Compile Include="Controllers\CSVController.cs" />
     <Compile Include="Controllers\DashboardController.cs" />
     <Compile Include="Controllers\ErrorController.cs" />
@@ -105,6 +192,7 @@
       <DependentUpon>CrashReport.dbml</DependentUpon>
     </Compile>
     <Compile Include="Models\CrashViewModel.cs" />
+    <Compile Include="Models\EmailSubscriptionResponseModel.cs" />
     <Compile Include="Models\FormHelper.cs" />
     <Compile Include="Models\PerformanceTimers.cs" />
     <Compile Include="Models\ReportsViewModel.cs" />
@@ -115,6 +203,7 @@
       <DesignTimeSharedInput>True</DesignTimeSharedInput>
       <DependentUpon>Settings.settings</DependentUpon>
     </Compile>
+    <Compile Include="Startup.cs" />
     <Compile Include="ThirdParty\iQueryableSearch.cs" />
     <Compile Include="Controllers\BuggsController.cs" />
     <Compile Include="Controllers\CrashesController.cs" />
@@ -155,6 +244,7 @@
     <Content Include="Content\Images\SortAscending.png" />
     <Content Include="Content\Images\SortDescending.png" />
     <Content Include="Content\Images\SortPlaceHolder.png" />
+    <Content Include="Content\ReportsStyle.css" />
     <Content Include="Content\themes\base\images\ui-bg_flat_0_aaaaaa_40x100.png" />
     <Content Include="Content\themes\base\images\ui-bg_flat_75_ffffff_40x100.png" />
     <Content Include="Content\themes\base\images\ui-bg_glass_55_fbf9ee_1x400.png" />
@@ -183,6 +273,7 @@
     <Content Include="Content\themes\base\jquery.ui.tabs.css" />
     <Content Include="Content\themes\base\jquery.ui.theme.css" />
     <Content Include="Global.asax" />
+    <Content Include="Properties\PublishProfiles\CRTest.pubxml.user" />
     <Content Include="Scripts\jquery-1.5.1-vsdoc.js" />
     <Content Include="Scripts\jquery-1.5.1.js" />
     <Content Include="Scripts\jquery-1.5.1.min.js" />
@@ -203,6 +294,7 @@
     <Content Include="Views\Buggs\Index.aspx" />
     <Content Include="Views\Buggs\Show.aspx" />
     <Content Include="Views\Buggs\ViewBuggs.ascx" />
+    <Content Include="Views\CSV\CSV.aspx" />
     <Content Include="Views\CSV\Index.aspx" />
     <Content Include="Views\CSV\ViewCSV.ascx" />
     <Content Include="Views\Dashboard\Index.aspx" />
@@ -211,6 +303,8 @@
     <Content Include="Views\Crashes\ViewCrash.ascx" />
     <Content Include="Views\Home\Index.aspx" />
     <Content Include="Views\Reports\Index.aspx" />
+    <None Include="Views\Reports\SignUpResponse.cshtml" />
+    <Content Include="Views\Reports\Unsubscribe.aspx" />
     <Content Include="Views\Reports\ViewReports.ascx" />
     <Content Include="Views\Users\Index.aspx" />
     <Content Include="Web.config">
@@ -218,6 +312,7 @@
     </Content>
     <Content Include="Web.Debug.config">
       <DependentUpon>Web.config</DependentUpon>
+      <SubType>Designer</SubType>
     </Content>
     <Content Include="Web.Release.config">
       <DependentUpon>Web.config</DependentUpon>
@@ -244,12 +339,17 @@
       <LastGenOutput>CrashReport.designer.cs</LastGenOutput>
       <SubType>Designer</SubType>
     </None>
+    <Content Include="packages.config" />
     <None Include="Properties\PublishProfiles\CrashReportDev.pubxml" />
+    <None Include="Properties\PublishProfiles\CRTest.pubxml" />
     <None Include="Properties\PublishProfiles\DefaultUE4.pubxml" />
     <None Include="Properties\Settings.settings">
       <Generator>SettingsSingleFileGenerator</Generator>
       <LastGenOutput>Settings.Designer.cs</LastGenOutput>
     </None>
+    <None Include="Web.Development.config">
+      <DependentUpon>Web.config</DependentUpon>
+    </None>
   </ItemGroup>
   <ItemGroup>
     <None Include="Models\CrashReport.dbml.layout">
@@ -266,6 +366,9 @@
       <Name>CrashReportCommon</Name>
     </ProjectReference>
   </ItemGroup>
+  <ItemGroup>
+    <Folder Include="Classes\" />
+  </ItemGroup>
   <PropertyGroup>
     <VisualStudioVersion Condition="'$(VisualStudioVersion)' == ''">10.0</VisualStudioVersion>
     <VSToolsPath Condition="'$(VSToolsPath)' == ''">$(MSBuildExtensionsPath32)\Microsoft\VisualStudio\v$(VisualStudioVersion)</VSToolsPath>
diff --git a/Engine/Source/Programs/CrashReporter/CrashReportWebSite/Models/CrashReport.cs b/Engine/Source/Programs/CrashReporter/CrashReportWebSite/Models/CrashReport.cs
--- a/Engine/Source/Programs/CrashReporter/CrashReportWebSite/Models/CrashReport.cs
+++ b/Engine/Source/Programs/CrashReporter/CrashReportWebSite/Models/CrashReport.cs
@@ -219,8 +219,7 @@ namespace Tools.CrashReporter.CrashReportWebSite.Models
 			{
 				this.BuildVersion = this.AffectedVersions.Last();	// Latest Version Affected
 			}
-
-
+            
 			foreach( var AffectedBuild in this.AffectedVersions )
 			{
 				var Subs = AffectedBuild.Split( ".".ToCharArray(), StringSplitOptions.RemoveEmptyEntries );
@@ -245,7 +244,7 @@ namespace Tools.CrashReporter.CrashReportWebSite.Models
 			this.LatestOSAffected = LatestOSAffected;			// Latest Environment Affected
 
 			// ToJiraSummary
-			var FunctionCalls = CrashesForBugg.First().GetCallStack().GetFunctionCallsForJira();
+			var FunctionCalls = new CallStackContainer(CrashesForBugg.First()).GetFunctionCallsForJira();
 			if( FunctionCalls.Count > 0 )
 			{
 				this.ToJiraSummary = FunctionCalls[0];
@@ -417,7 +416,7 @@ namespace Tools.CrashReporter.CrashReportWebSite.Models
 				var CrashList =
 				(
 						from BuggCrash in CrashRepo.Context.Buggs_Crashes
-						where BuggCrash.BuggId == Id && BuggCrash.Crash.TimeOfCrash > DateTime.UtcNow.AddMonths(-6)
+                        where BuggCrash.BuggId == Id && BuggCrash.Crash.TimeOfCrash > DateTime.UtcNow.AddMonths(-6)
 						select BuggCrash.Crash
 				).OrderByDescending( c => c.TimeOfCrash ).ToList();
 				return CrashList;
diff --git a/Engine/Source/Programs/CrashReporter/CrashReportWebSite/Models/CrashRepository.cs b/Engine/Source/Programs/CrashReporter/CrashReportWebSite/Models/CrashRepository.cs
--- a/Engine/Source/Programs/CrashReporter/CrashReportWebSite/Models/CrashRepository.cs
+++ b/Engine/Source/Programs/CrashReporter/CrashReportWebSite/Models/CrashRepository.cs
@@ -345,7 +345,7 @@ namespace Tools.CrashReporter.CrashReportWebSite.Models
 							Results = NewResult;
 						}
 					}
-				}				
+				}
 
 				// Get UserGroup ResultCounts
 				Dictionary<string, int> GroupCounts = GetCountsByGroupFromCrashes( Results );
@@ -459,7 +459,10 @@ namespace Tools.CrashReporter.CrashReportWebSite.Models
 						UserIdToGroupName.Add( UserIds.UserId, UserGroupName );
 					}
 
+                    //Get list of user Ids foreach crash range
 					List<int> UserIdCrashes = EnumerableCrashes.Select( Crash => Crash.UserNameId.Value ).ToList();
+
+                    //count the number of crashes in each group using the mapping of user id to crash.
 					foreach( int UserId in UserIdCrashes )
 					{
 						string UserGroupName = UserIdToGroupName[UserId];
@@ -642,28 +645,15 @@ namespace Tools.CrashReporter.CrashReportWebSite.Models
 			NewCrash.CommandLine = NewCrashInfo.CommandLine;
 			NewCrash.EngineMode = NewCrashInfo.EngineMode;
 			NewCrash.MachineId = NewCrashInfo.MachineGuid;
+            
+            //if there's a valid username assign the associated UserNameId else use "anonymous"
+            NewCrash.UserNameId = FRepository.Get(this).FindOrAddUser(!string.IsNullOrEmpty(NewCrashInfo.UserName) ? NewCrashInfo.UserName : UserNameAnonymous);
 
-			// Valid MachineID and UserName, updated crash from non-UE4 release
-			if(!string.IsNullOrEmpty(NewCrashInfo.UserName))
-			{
-				NewCrash.UserNameId = FRepository.Get( this ).FindOrAddUser( NewCrashInfo.UserName );
-			}
-			// Valid MachineID and EpicAccountId, updated crash from UE4 release
-			else if(!string.IsNullOrEmpty( NewCrashInfo.EpicAccountId ))
-			{
-				NewCrash.EpicAccountId = NewCrashInfo.EpicAccountId;
-				NewCrash.UserNameId = FRepository.Get( this ).FindOrAddUser( UserNameAnonymous );
-			}
-			// Crash from an older version.
-			else
-			{
-				// MachineGuid for older crashes is obsolete, so ignore it.
-				//NewCrash.ComputerName = NewCrashInfo.MachineGuid;
-				NewCrash.UserNameId = FRepository.Get( this ).FindOrAddUser
-				(
-					!string.IsNullOrEmpty( NewCrashInfo.UserName ) ? NewCrashInfo.UserName : UserNameAnonymous
-				);
-			}
+            //If there's a valid EpicAccountId assign that.
+            if (!string.IsNullOrEmpty(NewCrashInfo.EpicAccountId))
+            {
+                NewCrash.EpicAccountId = NewCrashInfo.EpicAccountId;
+            }
 			
 			NewCrash.Description = "";
 			if( NewCrashInfo.UserDescription != null )
diff --git a/Engine/Source/Programs/CrashReporter/CrashReportWebSite/Models/EmailSubscriptionResponseModel.cs b/Engine/Source/Programs/CrashReporter/CrashReportWebSite/Models/EmailSubscriptionResponseModel.cs
new file mode 100644
--- /dev/null
+++ b/Engine/Source/Programs/CrashReporter/CrashReportWebSite/Models/EmailSubscriptionResponseModel.cs
@@ -0,0 +1,24 @@
+﻿using System;
+using System.Collections.Generic;
+using System.Linq;
+using System.Web;
+
+namespace Tools.CrashReporter.CrashReportWebSite.Models
+{
+    /// <summary>
+    /// The View Model for email reports sign up response
+    /// This is used for the partial view used in the partial page update response.
+    /// </summary>
+    public class EmailSubscriptionResponseModel
+    {
+        /// <summary>
+        /// The email address to which the weekly report will be sent.
+        /// </summary>
+        public string Email { get; set; }
+
+        /// <summary>
+        /// The branch for which the report should be selected.
+        /// </summary>
+        public string Branch { get; set; }
+    };
+}
\ No newline at end of file
diff --git a/Engine/Source/Programs/CrashReporter/CrashReportWebSite/Views/CSV/Index.aspx b/Engine/Source/Programs/CrashReporter/CrashReportWebSite/Views/CSV/Index.aspx
--- a/Engine/Source/Programs/CrashReporter/CrashReportWebSite/Views/CSV/Index.aspx
+++ b/Engine/Source/Programs/CrashReporter/CrashReportWebSite/Views/CSV/Index.aspx
@@ -38,10 +38,10 @@
 </asp:Content>
 
 <asp:Content ID="AboveMainContent" ContentPlaceHolderID="AboveMainContent" runat="server" >
-<div style="clear:both;"><small style="color: lightgray;">Generated in <%=Model.GenerationTime%> second(s)</small><br /></div>
+<%--<div style="clear:both;"><small style="color: lightgray;">Generated in <%=Model.GenerationTime%> second(s)</small><br /></div>--%>
 
 <div id='SearchForm' style="clear:both;">
-<%using( Html.BeginForm( "", "CSV", FormMethod.Get, new { id = "CSVForm" } ) )
+<%using( Html.BeginForm( "GenerateCSV", "CSV", FormMethod.Get, new { id = "CSVForm" } ) )
 { %>
 	<script>$.datepicker.setDefaults($.datepicker.regional['']);</script>
 
@@ -63,13 +63,4 @@
 	<%} %>
 
 </div>
-</asp:Content>
-
-<asp:Content ID="CSVContent" ContentPlaceHolderID="CSVContent" runat="server">
-	<div>
-		<div>
-			<% Html.RenderPartial("/Views/CSV/ViewCSV.ascx"); %>
-		</div>
-	</div>
-
-</asp:Content>
+</asp:Content>
\ No newline at end of file
diff --git a/Engine/Source/Programs/CrashReporter/CrashReportWebSite/Views/Reports/Index.aspx b/Engine/Source/Programs/CrashReporter/CrashReportWebSite/Views/Reports/Index.aspx
--- a/Engine/Source/Programs/CrashReporter/CrashReportWebSite/Views/Reports/Index.aspx
+++ b/Engine/Source/Programs/CrashReporter/CrashReportWebSite/Views/Reports/Index.aspx
@@ -5,6 +5,24 @@
 
 <asp:Content ID="StyleSheet" ContentPlaceHolderID="CssContent" runat="server">
 	<link href="../../Content/Site.css" rel="stylesheet" type="text/css" />
+    <link href="../../Content/ReportsStyle.css" rel="stylesheet" type="text/css" />
+    
+    <style>
+        table, table tr, table th, table td
+        {
+	        border: solid;
+	        border-width: 1px;
+
+	        color: white;
+	        background-color: black;
+
+	        font-family: 'Segoe UI', Tahoma, Geneva, Verdana, sans-serif;
+	        font-size: 12px;
+
+	        padding: 2px;
+	        margin: 2px;
+        }
+    </style>
 </asp:Content>
 
 <asp:Content ID="TitleContent" ContentPlaceHolderID="TitleContent" runat="server">
@@ -34,19 +52,46 @@
 			}
 		});
 
+		function SubscribeToEmail() {
+		    var url = "/Reports/SubscribeToEmail";
+		    var formElement = $("#ReportsEmailForm");
+		    var branchName = formElement.find("#BranchFilter option:selected").val();
+		    var emailAddress = $("#subscribeEmail").val();
+
+		    $.ajax({
+		        type: "POST",
+		        url: url,
+		        data: { branchName: branchName, emailAddress: emailAddress },
+		        success: function (response) {
+		                $("#SignupResponsePartial").html(response);
+		            }
+		        }
+		    );
+		}
 	</script>
 </asp:Content>
 
 <asp:Content ID="AboveMainContent" ContentPlaceHolderID="AboveMainContent" runat="server" >
 <div style="clear:both;"><small style="color: lightgray;">Generated in <%=Model.GenerationTime%> second(s)</small><br /></div>
-
 <div id='SearchForm' style="clear:both;">
+    <div id="ReportsEmailForm" style="float:right;">
+        <span class="SearchTextTitle">Sign up to weekly reports e-mail</span>
+        <p>
+            <span class="SearchTextTitle">Select the report branch</span>
+        </p>
+        <p id="BranchFilter"><%=Html.DropDownListFor( m=>m.BranchName, Model.BranchNames )%></p>
+        <p>
+            <span class="SearchTextTitle">Enter email : </span>
+            <input id="subscribeEmail" type="text" class="date" autocomplete="OFF" style="width:200px" />
+        </p>
+        <input type="button" value="Email" class='SearchButton' onclick="SubscribeToEmail()"/>
+    </div>
 <%using( Html.BeginForm( "", "Reports", FormMethod.Get, new { id = "ReportsForm" } ) )
 { %>
 	<script>$.datepicker.setDefaults($.datepicker.regional['']);</script>
-
+    
 	<p class="SearchTextTitle">Filter by Date</p>
-
+    
 	<p>
 		<span class="SearchTextTitle">From:</span>
 		<input id="dateFromVisible" type="text" class="date" autocomplete="OFF" style="width:80px" />
@@ -54,18 +99,19 @@
 		<span class="SearchTextTitle">To:</span>
 		<input id="dateToVisible" type="text" class="date" autocomplete="OFF" style="width:80px" />
 		<input id="dateTo" name="dateTo" type="hidden" value="<%=Model.DateTo %>"autocomplete="OFF"  />
-        
 	</p>
     <p class="SearchTextTitle">Filter by Branch</p>
     <p id="BranchFilter"><%=Html.DropDownListFor( m=>m.BranchName, Model.BranchNames )%></p>
-
-	<p>
-		<input type="submit" value="Search" class='SearchButton' />
+    
+    <p>
+		<input type="submit" value="Search" class='SearchButton' id="SearchButton"/>
 	</p>
-	
 <%} %>
 
 </div>
+    <div style="clear:both;" id="SignupResponsePartial">
+        
+    </div>
 </asp:Content>
 
 <asp:Content ID="ReportsContent" ContentPlaceHolderID="ReportsContent" runat="server">
@@ -74,5 +120,4 @@
 			<% Html.RenderPartial("/Views/Reports/ViewReports.ascx"); %>
 		</div>
 	</div>
-
 </asp:Content>
diff --git a/Engine/Source/Programs/CrashReporter/CrashReportWebSite/Views/Web.config b/Engine/Source/Programs/CrashReporter/CrashReportWebSite/Views/Web.config
--- a/Engine/Source/Programs/CrashReporter/CrashReportWebSite/Views/Web.config
+++ b/Engine/Source/Programs/CrashReporter/CrashReportWebSite/Views/Web.config
@@ -2,9 +2,9 @@
 
 <configuration>
   <configSections>
-    <sectionGroup name="system.web.webPages.razor" type="System.Web.WebPages.Razor.Configuration.RazorWebSectionGroup, System.Web.WebPages.Razor, Version=1.0.0.0, Culture=neutral, PublicKeyToken=31BF3856AD364E35">
-      <section name="host" type="System.Web.WebPages.Razor.Configuration.HostSection, System.Web.WebPages.Razor, Version=1.0.0.0, Culture=neutral, PublicKeyToken=31BF3856AD364E35" requirePermission="false" />
-      <section name="pages" type="System.Web.WebPages.Razor.Configuration.RazorPagesSection, System.Web.WebPages.Razor, Version=1.0.0.0, Culture=neutral, PublicKeyToken=31BF3856AD364E35" requirePermission="false" />
+    <sectionGroup name="system.web.webPages.razor" type="System.Web.WebPages.Razor.Configuration.RazorWebSectionGroup, System.Web.WebPages.Razor, Version=3.0.0.0, Culture=neutral, PublicKeyToken=31BF3856AD364E35">
+      <section name="host" type="System.Web.WebPages.Razor.Configuration.HostSection, System.Web.WebPages.Razor, Version=3.0.0.0, Culture=neutral, PublicKeyToken=31BF3856AD364E35" requirePermission="false" />
+      <section name="pages" type="System.Web.WebPages.Razor.Configuration.RazorPagesSection, System.Web.WebPages.Razor, Version=3.0.0.0, Culture=neutral, PublicKeyToken=31BF3856AD364E35" requirePermission="false" />
     </sectionGroup>
   </configSections>
 
@@ -16,6 +16,7 @@
         <add namespace="System.Web.Mvc.Ajax" />
         <add namespace="System.Web.Mvc.Html" />
         <add namespace="System.Web.Routing" />
+        <add namespace="System.Data.Linq" />
       </namespaces>
     </pages>
   </system.web.webPages.razor>
diff --git a/Engine/Source/Programs/CrashReporter/CrashReportWebSite/Web.Debug.config b/Engine/Source/Programs/CrashReporter/CrashReportWebSite/Web.Debug.config
--- a/Engine/Source/Programs/CrashReporter/CrashReportWebSite/Web.Debug.config
+++ b/Engine/Source/Programs/CrashReporter/CrashReportWebSite/Web.Debug.config
@@ -14,6 +14,12 @@
         xdt:Transform="SetAttributes" xdt:Locator="Match(name)"/>
     </connectionStrings>
   -->
+  <connectionStrings  xdt:Transform="Replace">
+    <add name="CrashReportConnectionString"
+      connectionString="Data Source=NCLRNDEV-02.ncl.epicgames.net;Database=CrashReporterDev;User ID=rntdashboard;Password=Yellow"
+      providerName="System.Data.SqlClient"/>
+    <add name="Hangfire" connectionString="Data Source=NCLRNDEV-02.ncl.epicgames.net;Database=Hangfire;User ID=rntdashboard;Password=Yellow" providerName="System.Data.SqlClient" />
+  </connectionStrings>
   <system.web>
     <!--
       In the example below, the "Replace" transform will replace the entire 
diff --git a/Engine/Source/Programs/CrashReporter/CrashReportWebSite/Web.Release.config b/Engine/Source/Programs/CrashReporter/CrashReportWebSite/Web.Release.config
--- a/Engine/Source/Programs/CrashReporter/CrashReportWebSite/Web.Release.config
+++ b/Engine/Source/Programs/CrashReporter/CrashReportWebSite/Web.Release.config
@@ -14,6 +14,12 @@
         xdt:Transform="SetAttributes" xdt:Locator="Match(name)"/>
     </connectionStrings>
   -->
+  <connectionStrings  xdt:Transform="Replace">
+    <add name="CrashReportConnectionString"
+      connectionString="Data Source=db-09;Initial Catalog=CrashReport;Server=db-09;Database=CrashReport;Integrated Security=true"
+      providerName="System.Data.SqlClient"/>
+    <add name="Hangfire" connectionString="Data Source=db-09;Initial Catalog=CrashReporterSchduledTask;Server=db-09;Database=CrashReport;Integrated Security=true" providerName="System.Data.SqlClient" />
+  </connectionStrings>
   <system.web>
     <compilation xdt:Transform="RemoveAttributes(debug)" />
     <!--
@@ -28,4 +34,12 @@
       </customErrors>
     -->
   </system.web>
+  <system.net>
+    <mailSettings xdt:Transform="Replace">
+      <smtp deliveryMethod="Network">
+        <network host="smtp.epicgames.net"
+                defaultCredentials="false"/>
+      </smtp>
+    </mailSettings>
+  </system.net>
 </configuration>
\ No newline at end of file
diff --git a/Engine/Source/Programs/CrashReporter/CrashReportWebSite/Web.config b/Engine/Source/Programs/CrashReporter/CrashReportWebSite/Web.config
--- a/Engine/Source/Programs/CrashReporter/CrashReportWebSite/Web.config
+++ b/Engine/Source/Programs/CrashReporter/CrashReportWebSite/Web.config
@@ -20,6 +20,7 @@
     <add name="CrashReportConnectionString"
       connectionString="Data Source=db-09;Initial Catalog=CrashReport;Server=db-09;Database=CrashReport;Integrated Security=true"
       providerName="System.Data.SqlClient"/>
+    <add name="Hangfire" connectionString="Data Source=db-09;Initial Catalog=CrashReporterSchduledTask;Server=db-09;Database=CrashReport;Integrated Security=true" providerName="System.Data.SqlClient" />
   </connectionStrings>
   <!--
     For a description of web.config changes see http://go.microsoft.com/fwlink/?LinkId=235367.
@@ -30,21 +31,7 @@
       </system.Web>
   -->
   <system.web>
-    <!-- 
-            Set compilation debug="true" to insert debugging 
-            symbols into the compiled page. Because this 
-            affects performance, set this value to true only 
-            during development.
-    -->
     <compilation debug="true" targetFramework="4.5">
-      <assemblies>
-        <add assembly="System.Web.Mvc, Version=3.0.0.1, Culture=neutral, PublicKeyToken=31BF3856AD364E35"/>
-        <add assembly="System.Web.Abstractions, Version=4.0.0.0, Culture=neutral, PublicKeyToken=31BF3856AD364E35"/>
-        <add assembly="System.Web.Routing, Version=4.0.0.0, Culture=neutral, PublicKeyToken=31BF3856AD364E35"/>
-        <add assembly="System.Data.Linq, Version=4.0.0.0, Culture=neutral, PublicKeyToken=B77A5C561934E089"/>
-        <add assembly="System.Web.WebPages, Version=1.0.0.0, Culture=neutral, PublicKeyToken=31BF3856AD364E35"/>
-        <add assembly="System.Web.Helpers, Version=1.0.0.0, Culture=neutral, PublicKeyToken=31BF3856AD364E35"/>
-      </assemblies>
     </compilation>
     <!--
             The <customErrors> section enables configuration 
@@ -65,8 +52,6 @@
         <add namespace="System.Web.Mvc.Ajax"/>
         <add namespace="System.Web.Mvc.Html"/>
         <add namespace="System.Web.Routing"/>
-        <add namespace="System.Linq"/>
-        <add namespace="System.Collections.Generic"/>
         <add namespace="Tools.CrashReporter.CrashReportWebSite.Views.Helpers"/>
       </namespaces>
     </pages>
@@ -107,10 +92,30 @@
   </applicationSettings>
   <runtime>
     <assemblyBinding xmlns="urn:schemas-microsoft-com:asm.v1">
+      <dependentAssembly>
+        <assemblyIdentity name="System.Web.Helpers" publicKeyToken="31bf3856ad364e35"/>
+        <bindingRedirect oldVersion="1.0.0.0-3.0.0.0" newVersion="3.0.0.0"/>
+      </dependentAssembly>
+      <dependentAssembly>
+        <assemblyIdentity name="System.Data.Linq" publicKeyToken="b77a5c561934e089"/>
+        <bindingRedirect oldVersion="1.0.0.0-4.0.0.0" newVersion="4.0.0.0"/>
+      </dependentAssembly>
+      <dependentAssembly>
+        <assemblyIdentity name="System.Web.WebPages" publicKeyToken="31bf3856ad364e35"/>
+        <bindingRedirect oldVersion="1.0.0.0-3.0.0.0" newVersion="3.0.0.0"/>
+      </dependentAssembly>
       <dependentAssembly>
         <assemblyIdentity name="System.Web.Mvc" publicKeyToken="31bf3856ad364e35"/>
-        <bindingRedirect oldVersion="3.0.0.0-3.0.0.1" newVersion="3.0.0.1"/>
+        <bindingRedirect oldVersion="1.0.0.0-5.2.3.0" newVersion="5.2.3.0"/>
       </dependentAssembly>
     </assemblyBinding>
   </runtime>
+  <system.net>
+    <mailSettings>
+      <smtp deliveryMethod="Network">
+        <network host="smtp.epicgames.net"
+                defaultCredentials="false"/>
+      </smtp>
+    </mailSettings>
+  </system.net>
 </configuration>
\ No newline at end of file
