Copying //UE4/Dev-Core to //UE4/Main

diff --git a/Engine/Source/Programs/CrashReporter/CrashReportWebSite/packages/Common.Logging.3.3.1/lib/net35/Common.Logging.xml b/Engine/Source/Programs/CrashReporter/CrashReportWebSite/packages/Common.Logging.3.3.1/lib/net35/Common.Logging.xml
new file mode 100644
--- /dev/null
+++ b/Engine/Source/Programs/CrashReporter/CrashReportWebSite/packages/Common.Logging.3.3.1/lib/net35/Common.Logging.xml
@@ -0,0 +1,2780 @@
+<?xml version="1.0"?>
+<doc>
+    <assembly>
+        <name>Common.Logging</name>
+    </assembly>
+    <members>
+        <member name="T:CoverageExcludeAttribute">
+            <summary>
+            Indicates classes or members to be ignored by NCover
+            </summary>
+            <remarks>
+            Note, the key is chosen, because TestDriven.NET uses it as //ea argument to "Test With... Coverage"
+            </remarks>
+            <author>Erich Eichinger</author>
+        </member>
+        <member name="T:Common.Logging.ConfigurationException">
+            <summary>
+            The exception that is thrown when a configuration system error has occurred with Common.Logging
+            </summary>
+            <author>Mark Pollack</author>
+        </member>
+        <member name="M:Common.Logging.ConfigurationException.#ctor">
+            <summary>Creates a new instance of the ObjectsException class.</summary>
+        </member>
+        <member name="M:Common.Logging.ConfigurationException.#ctor(System.String)">
+            <summary>
+            Creates a new instance of the ConfigurationException class. with the specified message.
+            </summary>
+            <param name="message">
+            A message about the exception.
+            </param>
+        </member>
+        <member name="M:Common.Logging.ConfigurationException.#ctor(System.String,System.Exception)">
+            <summary>
+            Creates a new instance of the ConfigurationException class with the specified message
+            and root cause.
+            </summary>
+            <param name="message">
+            A message about the exception.
+            </param>
+            <param name="rootCause">
+            The root exception that is being wrapped.
+            </param>
+        </member>
+        <member name="M:Common.Logging.ConfigurationException.#ctor(System.Runtime.Serialization.SerializationInfo,System.Runtime.Serialization.StreamingContext)">
+            <summary>
+            Creates a new instance of the ConfigurationException class.
+            </summary>
+            <param name="info">
+            The <see cref="T:System.Runtime.Serialization.SerializationInfo"/>
+            that holds the serialized object data about the exception being thrown.
+            </param>
+            <param name="context">
+            The <see cref="T:System.Runtime.Serialization.StreamingContext"/>
+            that contains contextual information about the source or destination.
+            </param>
+        </member>
+        <member name="T:Common.Logging.Configuration.ArgUtils">
+            <summary>
+            Various utility methods for using during factory and logger instance configuration
+            </summary>
+            <author>Erich Eichinger</author>
+        </member>
+        <member name="M:Common.Logging.Configuration.ArgUtils.#cctor">
+            <summary>
+            Initialize all members before any of this class' methods can be accessed (avoids beforeFieldInit)
+            </summary>
+        </member>
+        <member name="M:Common.Logging.Configuration.ArgUtils.RegisterTypeParser``1(Common.Logging.Configuration.ArgUtils.ParseHandler{``0})">
+            <summary>
+            Adds the parser to the list of known type parsers.
+            </summary>
+            <remarks>
+            .NET intrinsic types are pre-registerd: short, int, long, float, double, decimal, bool
+            </remarks>
+        </member>
+        <member name="M:Common.Logging.Configuration.ArgUtils.GetValue(Common.Logging.Configuration.NameValueCollection,System.String)">
+            <summary>
+            Retrieves the named value from the specified <see cref="T:Common.Logging.Configuration.NameValueCollection"/>.
+            </summary>
+            <param name="values">may be null</param>
+            <param name="key">the value's key</param>
+            <returns>if <paramref key="values"/> is not null, the value returned by values[key]. <c>null</c> otherwise.</returns>
+        </member>
+        <member name="M:Common.Logging.Configuration.ArgUtils.GetValue(Common.Logging.Configuration.NameValueCollection,System.String,System.String)">
+            <summary>
+            Retrieves the named value from the specified <see cref="T:Common.Logging.Configuration.NameValueCollection"/>.
+            </summary>
+            <param name="values">may be null</param>
+            <param name="key">the value's key</param>
+            <param name="defaultValue">the default value, if not found</param>
+            <returns>if <paramref key="values"/> is not null, the value returned by values[key]. <c>null</c> otherwise.</returns>
+        </member>
+        <member name="M:Common.Logging.Configuration.ArgUtils.Coalesce(System.String[])">
+            <summary>
+            Returns the first nonnull, nonempty value among its arguments.
+            </summary>
+            <remarks>
+            Returns <c>null</c>, if the initial list was null or empty.
+            </remarks>
+            <seealso cref="M:Common.Logging.Configuration.ArgUtils.Coalesce``1(System.Predicate{``0},``0[])"/>
+        </member>
+        <member name="M:Common.Logging.Configuration.ArgUtils.Coalesce``1(System.Predicate{``0},``0[])">
+            <summary>
+            Returns the first nonnull, nonempty value among its arguments.
+            </summary>
+            <remarks>
+            Also 
+            </remarks>
+        </member>
+        <member name="M:Common.Logging.Configuration.ArgUtils.TryParseEnum``1(``0,System.String)">
+            <summary>
+            Tries parsing <paramref key="stringValue"/> into an enum of the type of <paramref key="defaultValue"/>.
+            </summary>
+            <param name="defaultValue">the default value to return if parsing fails</param>
+            <param name="stringValue">the string value to parse</param>
+            <returns>the successfully parsed value, <paramref key="defaultValue"/> otherwise.</returns>
+        </member>
+        <member name="M:Common.Logging.Configuration.ArgUtils.TryParse``1(``0,System.String)">
+            <summary>
+            Tries parsing <paramref key="stringValue"/> into the specified return type.
+            </summary>
+            <param name="defaultValue">the default value to return if parsing fails</param>
+            <param name="stringValue">the string value to parse</param>
+            <returns>the successfully parsed value, <paramref key="defaultValue"/> otherwise.</returns>
+        </member>
+        <member name="M:Common.Logging.Configuration.ArgUtils.AssertNotNull``1(System.String,``0)">
+            <summary>
+            Throws a <see cref="T:System.ArgumentNullException"/> if <paramref key="val"/> is <c>null</c>.
+            </summary>
+        </member>
+        <member name="M:Common.Logging.Configuration.ArgUtils.AssertNotNull``1(System.String,``0,System.String,System.Object[])">
+            <summary>
+            Throws a <see cref="T:System.ArgumentNullException"/> if <paramref key="val"/> is <c>null</c>.
+            </summary>
+        </member>
+        <member name="M:Common.Logging.Configuration.ArgUtils.AssertIsAssignable``1(System.String,System.Type)">
+            <summary>
+            Throws a <see cref="T:System.ArgumentOutOfRangeException"/> if an object of type <paramref key="valType"/> is not
+            assignable to type <typeparam key="T"></typeparam>.
+            </summary>
+        </member>
+        <member name="M:Common.Logging.Configuration.ArgUtils.AssertIsAssignable``1(System.String,System.Type,System.String,System.Object[])">
+            <summary>
+            Throws a <see cref="T:System.ArgumentOutOfRangeException"/> if an object of type <paramref key="valType"/> is not
+            assignable to type <typeparam key="T"></typeparam>.
+            </summary>
+        </member>
+        <member name="M:Common.Logging.Configuration.ArgUtils.Guard(Common.Logging.Configuration.ArgUtils.Action,System.String,System.Object[])">
+            <summary>
+            Ensures any exception thrown by the given <paramref key="action"/> is wrapped with an
+            <see cref="T:Common.Logging.ConfigurationException"/>. 
+            </summary>
+            <remarks>
+            If <paramref key="action"/> already throws a ConfigurationException, it will not be wrapped.
+            </remarks>
+            <param name="action">the action to execute</param>
+            <param name="messageFormat">the message to be set on the thrown <see cref="T:Common.Logging.ConfigurationException"/></param>
+            <param name="args">args to be passed to <see cref="M:System.String.Format(System.String,System.Object[])"/> to format the message</param>
+        </member>
+        <member name="M:Common.Logging.Configuration.ArgUtils.Guard``1(Common.Logging.Configuration.ArgUtils.Function{``0},System.String,System.Object[])">
+            <summary>
+            Ensures any exception thrown by the given <paramref key="function"/> is wrapped with an
+            <see cref="T:Common.Logging.ConfigurationException"/>. 
+            </summary>
+            <remarks>
+            If <paramref key="function"/> already throws a ConfigurationException, it will not be wrapped.
+            </remarks>
+            <param name="function">the action to execute</param>
+            <param name="messageFormat">the message to be set on the thrown <see cref="T:Common.Logging.ConfigurationException"/></param>
+            <param name="args">args to be passed to <see cref="M:System.String.Format(System.String,System.Object[])"/> to format the message</param>
+        </member>
+        <member name="T:Common.Logging.Configuration.ArgUtils.ParseHandler`1">
+            <summary>
+            A delegate converting a string representation into the target type
+            </summary>
+        </member>
+        <member name="T:Common.Logging.Configuration.ArgUtils.Action">
+            <summary>
+            An anonymous action delegate with no arguments and no return value.
+            </summary>
+            <seealso cref="M:Common.Logging.Configuration.ArgUtils.Guard(Common.Logging.Configuration.ArgUtils.Action,System.String,System.Object[])"/>
+        </member>
+        <member name="T:Common.Logging.Configuration.ArgUtils.Function`1">
+            <summary>
+            An anonymous action delegate with no arguments and no return value.
+            </summary>
+            <seealso cref="M:Common.Logging.Configuration.ArgUtils.Guard``1(Common.Logging.Configuration.ArgUtils.Function{``0},System.String,System.Object[])"/>
+        </member>
+        <member name="T:Common.Logging.Configuration.DefaultConfigurationReader">
+            <summary>
+            Implementation of <see cref="T:Common.Logging.IConfigurationReader"/> that uses the standard .NET 
+            configuration APIs, ConfigurationSettings in 1.x and ConfigurationManager in 2.0
+            </summary>
+            <author>Mark Pollack</author>
+        </member>
+        <member name="M:Common.Logging.Configuration.DefaultConfigurationReader.GetSection(System.String)">
+            <summary>
+            Parses the configuration section and returns the resulting object.
+            Using the <c>System.Configuration.ConfigurationManager</c>
+            </summary>
+            <param name="sectionName">Name of the configuration section.</param>
+            <returns>
+            Object created by a corresponding <c>IConfigurationSectionHandler"</c>
+            </returns>
+            <remarks>
+            	<p>
+            Primary purpose of this method is to allow us to parse and
+            load configuration sections using the same API regardless
+            of the .NET framework version.
+            </p>
+            </remarks>
+        </member>
+        <member name="T:Common.Logging.Configuration.LogSetting">
+            <summary>
+            Container used to hold configuration information from config file.
+            </summary>
+            <author>Gilles Bayon</author>
+        </member>
+        <member name="M:Common.Logging.Configuration.LogSetting.#ctor(System.Type,Common.Logging.Configuration.NameValueCollection)">
+            <summary>
+            
+            </summary>
+            <param name="factoryAdapterType">
+            The <see cref="T:Common.Logging.ILoggerFactoryAdapter"/> type 
+            that will be used for creating <see cref="T:Common.Logging.ILog"/>
+            </param>
+            <param name="properties">
+            Additional user supplied properties that are passed to the 
+            <paramref key="factoryAdapterType"/>'s constructor.
+            </param>
+        </member>
+        <member name="P:Common.Logging.Configuration.LogSetting.FactoryAdapterType">
+            <summary>
+            The <see cref="T:Common.Logging.ILoggerFactoryAdapter"/> type that will be used for creating <see cref="T:Common.Logging.ILog"/>
+            instances.
+            </summary>
+        </member>
+        <member name="P:Common.Logging.Configuration.LogSetting.Properties">
+            <summary>
+            Additional user supplied properties that are passed to the <see cref="P:Common.Logging.Configuration.LogSetting.FactoryAdapterType"/>'s constructor.
+            </summary>
+        </member>
+        <member name="T:Common.Logging.Configuration.NameValueCollection">
+            <summary>
+            Substitute NameValueCollection in System.Collections.Specialized.
+            </summary>
+        </member>
+        <member name="M:Common.Logging.Configuration.NameValueCollection.#ctor">
+            <summary>
+            Creates a new instance of <seealso cref="T:Common.Logging.Configuration.NameValueCollection">NameValueCollection</seealso>.
+            </summary>
+        </member>
+        <member name="M:Common.Logging.Configuration.NameValueCollection.GetValues(System.String)">
+            <summary>
+            Gets the values (only a single one) for the specified key (configuration name)
+            </summary>
+            <param name="key">The key.</param>
+            <returns>an array with one value, or null if no value exist</returns>
+        </member>
+        <member name="P:Common.Logging.Configuration.NameValueCollection.Item(System.String)">
+            <summary>
+            Gets or sets the value with the specified key.
+            </summary>
+            <value>
+            The value corrsponding to the key, or null if no value exist
+            </value>
+            <param name="key">The key.</param>
+            <returns>value store for the key</returns>
+        </member>
+        <member name="T:Common.Logging.Configuration.NameValueCollectionHelper">
+            <summary>
+            Helper class for working with NameValueCollection
+            </summary>
+        </member>
+        <member name="M:Common.Logging.Configuration.NameValueCollectionHelper.ToCommonLoggingCollection(System.Collections.Specialized.NameValueCollection)">
+            <summary>
+            Convert a <see cref="T:System.Collections.Specialized.NameValueCollection"/> into the corresponding 
+            common logging equivalent <see cref="T:Common.Logging.Configuration.NameValueCollection"/>
+            </summary>
+            <param name="properties">The properties.</param>
+            <returns></returns>
+        </member>
+        <member name="T:Common.Logging.Factory.AbstractCachingLoggerFactoryAdapter">
+            <summary>
+            An implementation of <see cref="T:Common.Logging.ILoggerFactoryAdapter"/> that caches loggers handed out by this factory.
+            </summary>
+            <remarks>
+            Implementors just need to override <see cref="M:Common.Logging.Factory.AbstractCachingLoggerFactoryAdapter.CreateLogger(System.String)"/>.
+            </remarks>
+            <author>Erich Eichinger</author>
+        </member>
+        <member name="M:Common.Logging.Factory.AbstractCachingLoggerFactoryAdapter.#ctor">
+            <summary>
+            Creates a new instance, the logger cache being case-sensitive.
+            </summary>
+        </member>
+        <member name="M:Common.Logging.Factory.AbstractCachingLoggerFactoryAdapter.#ctor(System.Boolean)">
+            <summary>
+            Creates a new instance, the logger cache being <paramref key="caseSensitiveLoggerCache"/>.
+            </summary>
+            <param name="caseSensitiveLoggerCache"></param>
+        </member>
+        <member name="M:Common.Logging.Factory.AbstractCachingLoggerFactoryAdapter.ClearLoggerCache">
+            <summary>
+            Purges all loggers from cache
+            </summary>
+        </member>
+        <member name="M:Common.Logging.Factory.AbstractCachingLoggerFactoryAdapter.CreateLogger(System.String)">
+            <summary>
+            Create the specified named logger instance
+            </summary>
+            <remarks>
+            Derived factories need to implement this method to create the
+            actual logger instance.
+            </remarks>
+        </member>
+        <member name="M:Common.Logging.Factory.AbstractCachingLoggerFactoryAdapter.GetLogger(System.Type)">
+            <summary>
+            Get a ILog instance by <see cref="T:System.Type"/>.
+            </summary>
+            <param name="type">Usually the <see cref="T:System.Type"/> of the current class.</param>
+            <returns>
+            An ILog instance either obtained from the internal cache or created by a call to <see cref="M:Common.Logging.Factory.AbstractCachingLoggerFactoryAdapter.CreateLogger(System.String)"/>.
+            </returns>
+        </member>
+        <member name="M:Common.Logging.Factory.AbstractCachingLoggerFactoryAdapter.GetLogger(System.String)">
+            <summary>
+            Get a ILog instance by key.
+            </summary>
+            <param name="key">Usually a <see cref="T:System.Type"/>'s Name or FullName property.</param>
+            <returns>
+            An ILog instance either obtained from the internal cache or created by a call to <see cref="M:Common.Logging.Factory.AbstractCachingLoggerFactoryAdapter.CreateLogger(System.String)"/>.
+            </returns>
+        </member>
+        <member name="M:Common.Logging.Factory.AbstractCachingLoggerFactoryAdapter.GetLoggerInternal(System.String)">
+            <summary>
+            Get or create a ILog instance by key.
+            </summary>
+            <param name="key">Usually a <see cref="T:System.Type"/>'s Name or FullName property.</param>
+            <returns>
+            An ILog instance either obtained from the internal cache or created by a call to <see cref="M:Common.Logging.Factory.AbstractCachingLoggerFactoryAdapter.CreateLogger(System.String)"/>.
+            </returns>
+        </member>
+        <member name="T:Common.Logging.Factory.AbstractLogger">
+            <summary>
+            Provides base implementation suitable for almost all logger adapters
+            </summary>
+            <author>Erich Eichinger</author>
+        </member>
+        <member name="F:Common.Logging.Factory.AbstractLogger.Write">
+            <summary>
+            Holds the method for writing a message to the log system.
+            </summary>
+        </member>
+        <member name="M:Common.Logging.Factory.AbstractLogger.#ctor">
+            <summary>
+            Creates a new logger instance using <see cref="M:Common.Logging.Factory.AbstractLogger.WriteInternal(Common.Logging.LogLevel,System.Object,System.Exception)"/> for 
+            writing log events to the underlying log system.
+            </summary>
+            <seealso cref="M:Common.Logging.Factory.AbstractLogger.GetWriteHandler"/>
+        </member>
+        <member name="M:Common.Logging.Factory.AbstractLogger.GetWriteHandler">
+            <summary>
+            Override this method to use a different method than <see cref="M:Common.Logging.Factory.AbstractLogger.WriteInternal(Common.Logging.LogLevel,System.Object,System.Exception)"/> 
+            for writing log events to the underlying log system.
+            </summary>
+            <remarks>
+            Usually you don't need to override thise method. The default implementation returns
+            <c>null</c> to indicate that the default handler <see cref="M:Common.Logging.Factory.AbstractLogger.WriteInternal(Common.Logging.LogLevel,System.Object,System.Exception)"/> should be 
+            used.
+            </remarks>
+        </member>
+        <member name="M:Common.Logging.Factory.AbstractLogger.WriteInternal(Common.Logging.LogLevel,System.Object,System.Exception)">
+            <summary>
+            Actually sends the message to the underlying log system.
+            </summary>
+            <param name="level">the level of this log event.</param>
+            <param name="message">the message to log</param>
+            <param name="exception">the exception to log (may be null)</param>
+        </member>
+        <member name="M:Common.Logging.Factory.AbstractLogger.Trace(System.Object)">
+            <summary>
+            Log a message object with the <see cref="F:Common.Logging.LogLevel.Trace"/> level.
+            </summary>
+            <param name="message">The message object to log.</param>
+        </member>
+        <member name="M:Common.Logging.Factory.AbstractLogger.Trace(System.Object,System.Exception)">
+            <summary>
+            Log a message object with the <see cref="F:Common.Logging.LogLevel.Trace"/> level including
+            the stack trace of the <see cref="T:System.Exception"/> passed
+            as a parameter.
+            </summary>
+            <param name="message">The message object to log.</param>
+            <param name="exception">The exception to log, including its stack trace.</param>
+        </member>
+        <member name="M:Common.Logging.Factory.AbstractLogger.TraceFormat(System.IFormatProvider,System.String,System.Object[])">
+            <summary>
+            Log a message with the <see cref="F:Common.Logging.LogLevel.Trace"/> level.
+            </summary>
+            <param name="formatProvider">An <see cref="T:System.IFormatProvider"/> that supplies culture-specific formatting information.</param>
+            <param name="format">The format of the message object to log.<see cref="M:System.String.Format(System.String,System.Object[])"/> </param>
+            <param name="args"></param>
+        </member>
+        <member name="M:Common.Logging.Factory.AbstractLogger.TraceFormat(System.IFormatProvider,System.String,System.Exception,System.Object[])">
+            <summary>
+            Log a message with the <see cref="F:Common.Logging.LogLevel.Trace"/> level.
+            </summary>
+            <param name="formatProvider">An <see cref="T:System.IFormatProvider"/> that supplies culture-specific formatting information.</param>
+            <param name="format">The format of the message object to log.<see cref="M:System.String.Format(System.String,System.Object[])"/> </param>
+            <param name="exception">The exception to log.</param>
+            <param name="args"></param>
+        </member>
+        <member name="M:Common.Logging.Factory.AbstractLogger.TraceFormat(System.String,System.Object[])">
+            <summary>
+            Log a message with the <see cref="F:Common.Logging.LogLevel.Trace"/> level.
+            </summary>
+            <param name="format">The format of the message object to log.<see cref="M:System.String.Format(System.String,System.Object[])"/> </param>
+            <param name="args">the list of format arguments</param>
+        </member>
+        <member name="M:Common.Logging.Factory.AbstractLogger.TraceFormat(System.String,System.Exception,System.Object[])">
+            <summary>
+            Log a message with the <see cref="F:Common.Logging.LogLevel.Trace"/> level.
+            </summary>
+            <param name="format">The format of the message object to log.<see cref="M:System.String.Format(System.String,System.Object[])"/> </param>
+            <param name="exception">The exception to log.</param>
+            <param name="args">the list of format arguments</param>
+        </member>
+        <member name="M:Common.Logging.Factory.AbstractLogger.Trace(System.Action{Common.Logging.FormatMessageHandler})">
+            <summary>
+            Log a message with the <see cref="F:Common.Logging.LogLevel.Trace"/> level using a callback to obtain the message
+            </summary>
+            <remarks>
+            Using this method avoids the cost of creating a message and evaluating message arguments 
+            that probably won't be logged due to loglevel settings.
+            </remarks>
+            <param name="formatMessageCallback">A callback used by the logger to obtain the message if log level is matched</param>
+        </member>
+        <member name="M:Common.Logging.Factory.AbstractLogger.Trace(System.Action{Common.Logging.FormatMessageHandler},System.Exception)">
+            <summary>
+            Log a message with the <see cref="F:Common.Logging.LogLevel.Trace"/> level using a callback to obtain the message
+            </summary>
+            <remarks>
+            Using this method avoids the cost of creating a message and evaluating message arguments 
+            that probably won't be logged due to loglevel settings.
+            </remarks>
+            <param name="formatMessageCallback">A callback used by the logger to obtain the message if log level is matched</param>
+            <param name="exception">The exception to log, including its stack trace.</param>
+        </member>
+        <member name="M:Common.Logging.Factory.AbstractLogger.Trace(System.IFormatProvider,System.Action{Common.Logging.FormatMessageHandler})">
+            <summary>
+            Log a message with the <see cref="F:Common.Logging.LogLevel.Trace"/> level using a callback to obtain the message
+            </summary>
+            <remarks>
+            Using this method avoids the cost of creating a message and evaluating message arguments 
+            that probably won't be logged due to loglevel settings.
+            </remarks>
+            <param name="formatProvider">An <see cref="T:System.IFormatProvider"/> that supplies culture-specific formatting information.</param>
+            <param name="formatMessageCallback">A callback used by the logger to obtain the message if log level is matched</param>
+        </member>
+        <member name="M:Common.Logging.Factory.AbstractLogger.Trace(System.IFormatProvider,System.Action{Common.Logging.FormatMessageHandler},System.Exception)">
+            <summary>
+            Log a message with the <see cref="F:Common.Logging.LogLevel.Trace"/> level using a callback to obtain the message
+            </summary>
+            <remarks>
+            Using this method avoids the cost of creating a message and evaluating message arguments 
+            that probably won't be logged due to loglevel settings.
+            </remarks>
+            <param name="formatProvider">An <see cref="T:System.IFormatProvider"/> that supplies culture-specific formatting information.</param>
+            <param name="formatMessageCallback">A callback used by the logger to obtain the message if log level is matched</param>
+            <param name="exception">The exception to log, including its stack trace.</param>
+        </member>
+        <member name="M:Common.Logging.Factory.AbstractLogger.Debug(System.Object)">
+            <summary>
+            Log a message object with the <see cref="F:Common.Logging.LogLevel.Debug"/> level.
+            </summary>
+            <param name="message">The message object to log.</param>
+        </member>
+        <member name="M:Common.Logging.Factory.AbstractLogger.Debug(System.Object,System.Exception)">
+            <summary>
+            Log a message object with the <see cref="F:Common.Logging.LogLevel.Debug"/> level including
+            the stack Debug of the <see cref="T:System.Exception"/> passed
+            as a parameter.
+            </summary>
+            <param name="message">The message object to log.</param>
+            <param name="exception">The exception to log, including its stack Debug.</param>
+        </member>
+        <member name="M:Common.Logging.Factory.AbstractLogger.DebugFormat(System.IFormatProvider,System.String,System.Object[])">
+            <summary>
+            Log a message with the <see cref="F:Common.Logging.LogLevel.Debug"/> level.
+            </summary>
+            <param name="formatProvider">An <see cref="T:System.IFormatProvider"/> that supplies culture-specific formatting information.</param>
+            <param name="format">The format of the message object to log.<see cref="M:System.String.Format(System.String,System.Object[])"/> </param>
+            <param name="args"></param>
+        </member>
+        <member name="M:Common.Logging.Factory.AbstractLogger.DebugFormat(System.IFormatProvider,System.String,System.Exception,System.Object[])">
+            <summary>
+            Log a message with the <see cref="F:Common.Logging.LogLevel.Debug"/> level.
+            </summary>
+            <param name="formatProvider">An <see cref="T:System.IFormatProvider"/> that supplies culture-specific formatting information.</param>
+            <param name="format">The format of the message object to log.<see cref="M:System.String.Format(System.String,System.Object[])"/> </param>
+            <param name="exception">The exception to log.</param>
+            <param name="args"></param>
+        </member>
+        <member name="M:Common.Logging.Factory.AbstractLogger.DebugFormat(System.String,System.Object[])">
+            <summary>
+            Log a message with the <see cref="F:Common.Logging.LogLevel.Debug"/> level.
+            </summary>
+            <param name="format">The format of the message object to log.<see cref="M:System.String.Format(System.String,System.Object[])"/> </param>
+            <param name="args">the list of format arguments</param>
+        </member>
+        <member name="M:Common.Logging.Factory.AbstractLogger.DebugFormat(System.String,System.Exception,System.Object[])">
+            <summary>
+            Log a message with the <see cref="F:Common.Logging.LogLevel.Debug"/> level.
+            </summary>
+            <param name="format">The format of the message object to log.<see cref="M:System.String.Format(System.String,System.Object[])"/> </param>
+            <param name="exception">The exception to log.</param>
+            <param name="args">the list of format arguments</param>
+        </member>
+        <member name="M:Common.Logging.Factory.AbstractLogger.Debug(System.Action{Common.Logging.FormatMessageHandler})">
+            <summary>
+            Log a message with the <see cref="F:Common.Logging.LogLevel.Debug"/> level using a callback to obtain the message
+            </summary>
+            <remarks>
+            Using this method avoids the cost of creating a message and evaluating message arguments 
+            that probably won't be logged due to loglevel settings.
+            </remarks>
+            <param name="formatMessageCallback">A callback used by the logger to obtain the message if log level is matched</param>
+        </member>
+        <member name="M:Common.Logging.Factory.AbstractLogger.Debug(System.Action{Common.Logging.FormatMessageHandler},System.Exception)">
+            <summary>
+            Log a message with the <see cref="F:Common.Logging.LogLevel.Debug"/> level using a callback to obtain the message
+            </summary>
+            <remarks>
+            Using this method avoids the cost of creating a message and evaluating message arguments 
+            that probably won't be logged due to loglevel settings.
+            </remarks>
+            <param name="formatMessageCallback">A callback used by the logger to obtain the message if log level is matched</param>
+            <param name="exception">The exception to log, including its stack Debug.</param>
+        </member>
+        <member name="M:Common.Logging.Factory.AbstractLogger.Debug(System.IFormatProvider,System.Action{Common.Logging.FormatMessageHandler})">
+            <summary>
+            Log a message with the <see cref="F:Common.Logging.LogLevel.Debug"/> level using a callback to obtain the message
+            </summary>
+            <remarks>
+            Using this method avoids the cost of creating a message and evaluating message arguments 
+            that probably won't be logged due to loglevel settings.
+            </remarks>
+            <param name="formatProvider">An <see cref="T:System.IFormatProvider"/> that supplies culture-specific formatting information.</param>
+            <param name="formatMessageCallback">A callback used by the logger to obtain the message if log level is matched</param>
+        </member>
+        <member name="M:Common.Logging.Factory.AbstractLogger.Debug(System.IFormatProvider,System.Action{Common.Logging.FormatMessageHandler},System.Exception)">
+            <summary>
+            Log a message with the <see cref="F:Common.Logging.LogLevel.Debug"/> level using a callback to obtain the message
+            </summary>
+            <remarks>
+            Using this method avoids the cost of creating a message and evaluating message arguments 
+            that probably won't be logged due to loglevel settings.
+            </remarks>
+            <param name="formatProvider">An <see cref="T:System.IFormatProvider"/> that supplies culture-specific formatting information.</param>
+            <param name="formatMessageCallback">A callback used by the logger to obtain the message if log level is matched</param>
+            <param name="exception">The exception to log, including its stack Debug.</param>
+        </member>
+        <member name="M:Common.Logging.Factory.AbstractLogger.Info(System.Object)">
+            <summary>
+            Log a message object with the <see cref="F:Common.Logging.LogLevel.Info"/> level.
+            </summary>
+            <param name="message">The message object to log.</param>
+        </member>
+        <member name="M:Common.Logging.Factory.AbstractLogger.Info(System.Object,System.Exception)">
+            <summary>
+            Log a message object with the <see cref="F:Common.Logging.LogLevel.Info"/> level including
+            the stack Info of the <see cref="T:System.Exception"/> passed
+            as a parameter.
+            </summary>
+            <param name="message">The message object to log.</param>
+            <param name="exception">The exception to log, including its stack Info.</param>
+        </member>
+        <member name="M:Common.Logging.Factory.AbstractLogger.InfoFormat(System.IFormatProvider,System.String,System.Object[])">
+            <summary>
+            Log a message with the <see cref="F:Common.Logging.LogLevel.Info"/> level.
+            </summary>
+            <param name="formatProvider">An <see cref="T:System.IFormatProvider"/> that supplies culture-specific formatting information.</param>
+            <param name="format">The format of the message object to log.<see cref="M:System.String.Format(System.String,System.Object[])"/> </param>
+            <param name="args"></param>
+        </member>
+        <member name="M:Common.Logging.Factory.AbstractLogger.InfoFormat(System.IFormatProvider,System.String,System.Exception,System.Object[])">
+            <summary>
+            Log a message with the <see cref="F:Common.Logging.LogLevel.Info"/> level.
+            </summary>
+            <param name="formatProvider">An <see cref="T:System.IFormatProvider"/> that supplies culture-specific formatting information.</param>
+            <param name="format">The format of the message object to log.<see cref="M:System.String.Format(System.String,System.Object[])"/> </param>
+            <param name="exception">The exception to log.</param>
+            <param name="args"></param>
+        </member>
+        <member name="M:Common.Logging.Factory.AbstractLogger.InfoFormat(System.String,System.Object[])">
+            <summary>
+            Log a message with the <see cref="F:Common.Logging.LogLevel.Info"/> level.
+            </summary>
+            <param name="format">The format of the message object to log.<see cref="M:System.String.Format(System.String,System.Object[])"/> </param>
+            <param name="args">the list of format arguments</param>
+        </member>
+        <member name="M:Common.Logging.Factory.AbstractLogger.InfoFormat(System.String,System.Exception,System.Object[])">
+            <summary>
+            Log a message with the <see cref="F:Common.Logging.LogLevel.Info"/> level.
+            </summary>
+            <param name="format">The format of the message object to log.<see cref="M:System.String.Format(System.String,System.Object[])"/> </param>
+            <param name="exception">The exception to log.</param>
+            <param name="args">the list of format arguments</param>
+        </member>
+        <member name="M:Common.Logging.Factory.AbstractLogger.Info(System.Action{Common.Logging.FormatMessageHandler})">
+            <summary>
+            Log a message with the <see cref="F:Common.Logging.LogLevel.Info"/> level using a callback to obtain the message
+            </summary>
+            <remarks>
+            Using this method avoids the cost of creating a message and evaluating message arguments 
+            that probably won't be logged due to loglevel settings.
+            </remarks>
+            <param name="formatMessageCallback">A callback used by the logger to obtain the message if log level is matched</param>
+        </member>
+        <member name="M:Common.Logging.Factory.AbstractLogger.Info(System.Action{Common.Logging.FormatMessageHandler},System.Exception)">
+            <summary>
+            Log a message with the <see cref="F:Common.Logging.LogLevel.Info"/> level using a callback to obtain the message
+            </summary>
+            <remarks>
+            Using this method avoids the cost of creating a message and evaluating message arguments 
+            that probably won't be logged due to loglevel settings.
+            </remarks>
+            <param name="formatMessageCallback">A callback used by the logger to obtain the message if log level is matched</param>
+            <param name="exception">The exception to log, including its stack Info.</param>
+        </member>
+        <member name="M:Common.Logging.Factory.AbstractLogger.Info(System.IFormatProvider,System.Action{Common.Logging.FormatMessageHandler})">
+            <summary>
+            Log a message with the <see cref="F:Common.Logging.LogLevel.Info"/> level using a callback to obtain the message
+            </summary>
+            <remarks>
+            Using this method avoids the cost of creating a message and evaluating message arguments 
+            that probably won't be logged due to loglevel settings.
+            </remarks>
+            <param name="formatProvider">An <see cref="T:System.IFormatProvider"/> that supplies culture-specific formatting information.</param>
+            <param name="formatMessageCallback">A callback used by the logger to obtain the message if log level is matched</param>
+        </member>
+        <member name="M:Common.Logging.Factory.AbstractLogger.Info(System.IFormatProvider,System.Action{Common.Logging.FormatMessageHandler},System.Exception)">
+            <summary>
+            Log a message with the <see cref="F:Common.Logging.LogLevel.Info"/> level using a callback to obtain the message
+            </summary>
+            <remarks>
+            Using this method avoids the cost of creating a message and evaluating message arguments 
+            that probably won't be logged due to loglevel settings.
+            </remarks>
+            <param name="formatProvider">An <see cref="T:System.IFormatProvider"/> that supplies culture-specific formatting information.</param>
+            <param name="formatMessageCallback">A callback used by the logger to obtain the message if log level is matched</param>
+            <param name="exception">The exception to log, including its stack Info.</param>
+        </member>
+        <member name="M:Common.Logging.Factory.AbstractLogger.Warn(System.Object)">
+            <summary>
+            Log a message object with the <see cref="F:Common.Logging.LogLevel.Warn"/> level.
+            </summary>
+            <param name="message">The message object to log.</param>
+        </member>
+        <member name="M:Common.Logging.Factory.AbstractLogger.Warn(System.Object,System.Exception)">
+            <summary>
+            Log a message object with the <see cref="F:Common.Logging.LogLevel.Warn"/> level including
+            the stack Warn of the <see cref="T:System.Exception"/> passed
+            as a parameter.
+            </summary>
+            <param name="message">The message object to log.</param>
+            <param name="exception">The exception to log, including its stack Warn.</param>
+        </member>
+        <member name="M:Common.Logging.Factory.AbstractLogger.WarnFormat(System.IFormatProvider,System.String,System.Object[])">
+            <summary>
+            Log a message with the <see cref="F:Common.Logging.LogLevel.Warn"/> level.
+            </summary>
+            <param name="formatProvider">An <see cref="T:System.IFormatProvider"/> that supplies culture-specific formatting Information.</param>
+            <param name="format">The format of the message object to log.<see cref="M:System.String.Format(System.String,System.Object[])"/> </param>
+            <param name="args"></param>
+        </member>
+        <member name="M:Common.Logging.Factory.AbstractLogger.WarnFormat(System.IFormatProvider,System.String,System.Exception,System.Object[])">
+            <summary>
+            Log a message with the <see cref="F:Common.Logging.LogLevel.Warn"/> level.
+            </summary>
+            <param name="formatProvider">An <see cref="T:System.IFormatProvider"/> that supplies culture-specific formatting information.</param>
+            <param name="format">The format of the message object to log.<see cref="M:System.String.Format(System.String,System.Object[])"/> </param>
+            <param name="exception">The exception to log.</param>
+            <param name="args"></param>
+        </member>
+        <member name="M:Common.Logging.Factory.AbstractLogger.WarnFormat(System.String,System.Object[])">
+            <summary>
+            Log a message with the <see cref="F:Common.Logging.LogLevel.Warn"/> level.
+            </summary>
+            <param name="format">The format of the message object to log.<see cref="M:System.String.Format(System.String,System.Object[])"/> </param>
+            <param name="args">the list of format arguments</param>
+        </member>
+        <member name="M:Common.Logging.Factory.AbstractLogger.WarnFormat(System.String,System.Exception,System.Object[])">
+            <summary>
+            Log a message with the <see cref="F:Common.Logging.LogLevel.Warn"/> level.
+            </summary>
+            <param name="format">The format of the message object to log.<see cref="M:System.String.Format(System.String,System.Object[])"/> </param>
+            <param name="exception">The exception to log.</param>
+            <param name="args">the list of format arguments</param>
+        </member>
+        <member name="M:Common.Logging.Factory.AbstractLogger.Warn(System.Action{Common.Logging.FormatMessageHandler})">
+            <summary>
+            Log a message with the <see cref="F:Common.Logging.LogLevel.Warn"/> level using a callback to obtain the message
+            </summary>
+            <remarks>
+            Using this method avoids the cost of creating a message and evaluating message arguments 
+            that probably won't be logged due to loglevel settings.
+            </remarks>
+            <param name="formatMessageCallback">A callback used by the logger to obtain the message if log level is matched</param>
+        </member>
+        <member name="M:Common.Logging.Factory.AbstractLogger.Warn(System.Action{Common.Logging.FormatMessageHandler},System.Exception)">
+            <summary>
+            Log a message with the <see cref="F:Common.Logging.LogLevel.Warn"/> level using a callback to obtain the message
+            </summary>
+            <remarks>
+            Using this method avoids the cost of creating a message and evaluating message arguments 
+            that probably won't be logged due to loglevel settings.
+            </remarks>
+            <param name="formatMessageCallback">A callback used by the logger to obtain the message if log level is matched</param>
+            <param name="exception">The exception to log, including its stack Warn.</param>
+        </member>
+        <member name="M:Common.Logging.Factory.AbstractLogger.Warn(System.IFormatProvider,System.Action{Common.Logging.FormatMessageHandler})">
+            <summary>
+            Log a message with the <see cref="F:Common.Logging.LogLevel.Warn"/> level using a callback to obtain the message
+            </summary>
+            <remarks>
+            Using this method avoids the cost of creating a message and evaluating message arguments 
+            that probably won't be logged due to loglevel settings.
+            </remarks>
+            <param name="formatProvider">An <see cref="T:System.IFormatProvider"/> that supplies culture-specific formatting information.</param>
+            <param name="formatMessageCallback">A callback used by the logger to obtain the message if log level is matched</param>
+        </member>
+        <member name="M:Common.Logging.Factory.AbstractLogger.Warn(System.IFormatProvider,System.Action{Common.Logging.FormatMessageHandler},System.Exception)">
+            <summary>
+            Log a message with the <see cref="F:Common.Logging.LogLevel.Warn"/> level using a callback to obtain the message
+            </summary>
+            <remarks>
+            Using this method avoids the cost of creating a message and evaluating message arguments 
+            that probably won't be logged due to loglevel settings.
+            </remarks>
+            <param name="formatProvider">An <see cref="T:System.IFormatProvider"/> that supplies culture-specific formatting information.</param>
+            <param name="formatMessageCallback">A callback used by the logger to obtain the message if log level is matched</param>
+            <param name="exception">The exception to log, including its stack Warn.</param>
+        </member>
+        <member name="M:Common.Logging.Factory.AbstractLogger.Error(System.Object)">
+            <summary>
+            Log a message object with the <see cref="F:Common.Logging.LogLevel.Error"/> level.
+            </summary>
+            <param name="message">The message object to log.</param>
+        </member>
+        <member name="M:Common.Logging.Factory.AbstractLogger.Error(System.Object,System.Exception)">
+            <summary>
+            Log a message object with the <see cref="F:Common.Logging.LogLevel.Error"/> level including
+            the stack Error of the <see cref="T:System.Exception"/> passed
+            as a parameter.
+            </summary>
+            <param name="message">The message object to log.</param>
+            <param name="exception">The exception to log, including its stack Error.</param>
+        </member>
+        <member name="M:Common.Logging.Factory.AbstractLogger.ErrorFormat(System.IFormatProvider,System.String,System.Object[])">
+            <summary>
+            Log a message with the <see cref="F:Common.Logging.LogLevel.Error"/> level.
+            </summary>
+            <param name="formatProvider">An <see cref="T:System.IFormatProvider"/> that supplies culture-specific formatting Errorrmation.</param>
+            <param name="format">The format of the message object to log.<see cref="M:System.String.Format(System.String,System.Object[])"/> </param>
+            <param name="args"></param>
+        </member>
+        <member name="M:Common.Logging.Factory.AbstractLogger.ErrorFormat(System.IFormatProvider,System.String,System.Exception,System.Object[])">
+            <summary>
+            Log a message with the <see cref="F:Common.Logging.LogLevel.Error"/> level.
+            </summary>
+            <param name="formatProvider">An <see cref="T:System.IFormatProvider"/> that supplies culture-specific formatting Errorrmation.</param>
+            <param name="format">The format of the message object to log.<see cref="M:System.String.Format(System.String,System.Object[])"/> </param>
+            <param name="exception">The exception to log.</param>
+            <param name="args"></param>
+        </member>
+        <member name="M:Common.Logging.Factory.AbstractLogger.ErrorFormat(System.String,System.Object[])">
+            <summary>
+            Log a message with the <see cref="F:Common.Logging.LogLevel.Error"/> level.
+            </summary>
+            <param name="format">The format of the message object to log.<see cref="M:System.String.Format(System.String,System.Object[])"/> </param>
+            <param name="args">the list of format arguments</param>
+        </member>
+        <member name="M:Common.Logging.Factory.AbstractLogger.ErrorFormat(System.String,System.Exception,System.Object[])">
+            <summary>
+            Log a message with the <see cref="F:Common.Logging.LogLevel.Error"/> level.
+            </summary>
+            <param name="format">The format of the message object to log.<see cref="M:System.String.Format(System.String,System.Object[])"/> </param>
+            <param name="exception">The exception to log.</param>
+            <param name="args">the list of format arguments</param>
+        </member>
+        <member name="M:Common.Logging.Factory.AbstractLogger.Error(System.Action{Common.Logging.FormatMessageHandler})">
+            <summary>
+            Log a message with the <see cref="F:Common.Logging.LogLevel.Error"/> level using a callback to obtain the message
+            </summary>
+            <remarks>
+            Using this method avoids the cost of creating a message and evaluating message arguments 
+            that probably won't be logged due to loglevel settings.
+            </remarks>
+            <param name="formatMessageCallback">A callback used by the logger to obtain the message if log level is matched</param>
+        </member>
+        <member name="M:Common.Logging.Factory.AbstractLogger.Error(System.Action{Common.Logging.FormatMessageHandler},System.Exception)">
+            <summary>
+            Log a message with the <see cref="F:Common.Logging.LogLevel.Error"/> level using a callback to obtain the message
+            </summary>
+            <remarks>
+            Using this method avoids the cost of creating a message and evaluating message arguments 
+            that probably won't be logged due to loglevel settings.
+            </remarks>
+            <param name="formatMessageCallback">A callback used by the logger to obtain the message if log level is matched</param>
+            <param name="exception">The exception to log, including its stack Error.</param>
+        </member>
+        <member name="M:Common.Logging.Factory.AbstractLogger.Error(System.IFormatProvider,System.Action{Common.Logging.FormatMessageHandler})">
+            <summary>
+            Log a message with the <see cref="F:Common.Logging.LogLevel.Error"/> level using a callback to obtain the message
+            </summary>
+            <remarks>
+            Using this method avoids the cost of creating a message and evaluating message arguments 
+            that probably won't be logged due to loglevel settings.
+            </remarks>
+            <param name="formatProvider">An <see cref="T:System.IFormatProvider"/> that supplies culture-specific formatting information.</param>
+            <param name="formatMessageCallback">A callback used by the logger to obtain the message if log level is matched</param>
+        </member>
+        <member name="M:Common.Logging.Factory.AbstractLogger.Error(System.IFormatProvider,System.Action{Common.Logging.FormatMessageHandler},System.Exception)">
+            <summary>
+            Log a message with the <see cref="F:Common.Logging.LogLevel.Error"/> level using a callback to obtain the message
+            </summary>
+            <remarks>
+            Using this method avoids the cost of creating a message and evaluating message arguments 
+            that probably won't be logged due to loglevel settings.
+            </remarks>
+            <param name="formatProvider">An <see cref="T:System.IFormatProvider"/> that supplies culture-specific formatting information.</param>
+            <param name="formatMessageCallback">A callback used by the logger to obtain the message if log level is matched</param>
+            <param name="exception">The exception to log, including its stack Error.</param>
+        </member>
+        <member name="M:Common.Logging.Factory.AbstractLogger.Fatal(System.Object)">
+            <summary>
+            Log a message object with the <see cref="F:Common.Logging.LogLevel.Fatal"/> level.
+            </summary>
+            <param name="message">The message object to log.</param>
+        </member>
+        <member name="M:Common.Logging.Factory.AbstractLogger.Fatal(System.Object,System.Exception)">
+            <summary>
+            Log a message object with the <see cref="F:Common.Logging.LogLevel.Fatal"/> level including
+            the stack Fatal of the <see cref="T:System.Exception"/> passed
+            as a parameter.
+            </summary>
+            <param name="message">The message object to log.</param>
+            <param name="exception">The exception to log, including its stack Fatal.</param>
+        </member>
+        <member name="M:Common.Logging.Factory.AbstractLogger.FatalFormat(System.IFormatProvider,System.String,System.Object[])">
+            <summary>
+            Log a message with the <see cref="F:Common.Logging.LogLevel.Fatal"/> level.
+            </summary>
+            <param name="formatProvider">An <see cref="T:System.IFormatProvider"/> that supplies culture-specific formatting Fatalrmation.</param>
+            <param name="format">The format of the message object to log.<see cref="M:System.String.Format(System.String,System.Object[])"/> </param>
+            <param name="args"></param>
+        </member>
+        <member name="M:Common.Logging.Factory.AbstractLogger.FatalFormat(System.IFormatProvider,System.String,System.Exception,System.Object[])">
+            <summary>
+            Log a message with the <see cref="F:Common.Logging.LogLevel.Fatal"/> level.
+            </summary>
+            <param name="formatProvider">An <see cref="T:System.IFormatProvider"/> that supplies culture-specific formatting Fatalrmation.</param>
+            <param name="format">The format of the message object to log.<see cref="M:System.String.Format(System.String,System.Object[])"/> </param>
+            <param name="exception">The exception to log.</param>
+            <param name="args"></param>
+        </member>
+        <member name="M:Common.Logging.Factory.AbstractLogger.FatalFormat(System.String,System.Object[])">
+            <summary>
+            Log a message with the <see cref="F:Common.Logging.LogLevel.Fatal"/> level.
+            </summary>
+            <param name="format">The format of the message object to log.<see cref="M:System.String.Format(System.String,System.Object[])"/> </param>
+            <param name="args">the list of format arguments</param>
+        </member>
+        <member name="M:Common.Logging.Factory.AbstractLogger.FatalFormat(System.String,System.Exception,System.Object[])">
+            <summary>
+            Log a message with the <see cref="F:Common.Logging.LogLevel.Fatal"/> level.
+            </summary>
+            <param name="format">The format of the message object to log.<see cref="M:System.String.Format(System.String,System.Object[])"/> </param>
+            <param name="exception">The exception to log.</param>
+            <param name="args">the list of format arguments</param>
+        </member>
+        <member name="M:Common.Logging.Factory.AbstractLogger.Fatal(System.Action{Common.Logging.FormatMessageHandler})">
+            <summary>
+            Log a message with the <see cref="F:Common.Logging.LogLevel.Fatal"/> level using a callback to obtain the message
+            </summary>
+            <remarks>
+            Using this method avoids the cost of creating a message and evaluating message arguments 
+            that probably won't be logged due to loglevel settings.
+            </remarks>
+            <param name="formatMessageCallback">A callback used by the logger to obtain the message if log level is matched</param>
+        </member>
+        <member name="M:Common.Logging.Factory.AbstractLogger.Fatal(System.Action{Common.Logging.FormatMessageHandler},System.Exception)">
+            <summary>
+            Log a message with the <see cref="F:Common.Logging.LogLevel.Fatal"/> level using a callback to obtain the message
+            </summary>
+            <remarks>
+            Using this method avoids the cost of creating a message and evaluating message arguments 
+            that probably won't be logged due to loglevel settings.
+            </remarks>
+            <param name="formatMessageCallback">A callback used by the logger to obtain the message if log level is matched</param>
+            <param name="exception">The exception to log, including its stack Fatal.</param>
+        </member>
+        <member name="M:Common.Logging.Factory.AbstractLogger.Fatal(System.IFormatProvider,System.Action{Common.Logging.FormatMessageHandler})">
+            <summary>
+            Log a message with the <see cref="F:Common.Logging.LogLevel.Fatal"/> level using a callback to obtain the message
+            </summary>
+            <remarks>
+            Using this method avoids the cost of creating a message and evaluating message arguments 
+            that probably won't be logged due to loglevel settings.
+            </remarks>
+            <param name="formatProvider">An <see cref="T:System.IFormatProvider"/> that supplies culture-specific formatting information.</param>
+            <param name="formatMessageCallback">A callback used by the logger to obtain the message if log level is matched</param>
+        </member>
+        <member name="M:Common.Logging.Factory.AbstractLogger.Fatal(System.IFormatProvider,System.Action{Common.Logging.FormatMessageHandler},System.Exception)">
+            <summary>
+            Log a message with the <see cref="F:Common.Logging.LogLevel.Fatal"/> level using a callback to obtain the message
+            </summary>
+            <remarks>
+            Using this method avoids the cost of creating a message and evaluating message arguments 
+            that probably won't be logged due to loglevel settings.
+            </remarks>
+            <param name="formatProvider">An <see cref="T:System.IFormatProvider"/> that supplies culture-specific formatting information.</param>
+            <param name="formatMessageCallback">A callback used by the logger to obtain the message if log level is matched</param>
+            <param name="exception">The exception to log, including its stack Fatal.</param>
+        </member>
+        <member name="P:Common.Logging.Factory.AbstractLogger.IsTraceEnabled">
+            <summary>
+            Checks if this logger is enabled for the <see cref="F:Common.Logging.LogLevel.Trace"/> level.
+            </summary>
+            <remarks>
+            Override this in your derived class to comply with the underlying logging system
+            </remarks>
+        </member>
+        <member name="P:Common.Logging.Factory.AbstractLogger.IsDebugEnabled">
+            <summary>
+            Checks if this logger is enabled for the <see cref="F:Common.Logging.LogLevel.Debug"/> level.
+            </summary>
+            <remarks>
+            Override this in your derived class to comply with the underlying logging system
+            </remarks>
+        </member>
+        <member name="P:Common.Logging.Factory.AbstractLogger.IsInfoEnabled">
+            <summary>
+            Checks if this logger is enabled for the <see cref="F:Common.Logging.LogLevel.Info"/> level.
+            </summary>
+            <remarks>
+            Override this in your derived class to comply with the underlying logging system
+            </remarks>
+        </member>
+        <member name="P:Common.Logging.Factory.AbstractLogger.IsWarnEnabled">
+            <summary>
+            Checks if this logger is enabled for the <see cref="F:Common.Logging.LogLevel.Warn"/> level.
+            </summary>
+            <remarks>
+            Override this in your derived class to comply with the underlying logging system
+            </remarks>
+        </member>
+        <member name="P:Common.Logging.Factory.AbstractLogger.IsErrorEnabled">
+            <summary>
+            Checks if this logger is enabled for the <see cref="F:Common.Logging.LogLevel.Error"/> level.
+            </summary>
+            <remarks>
+            Override this in your derived class to comply with the underlying logging system
+            </remarks>
+        </member>
+        <member name="P:Common.Logging.Factory.AbstractLogger.IsFatalEnabled">
+            <summary>
+            Checks if this logger is enabled for the <see cref="F:Common.Logging.LogLevel.Fatal"/> level.
+            </summary>
+            <remarks>
+            Override this in your derived class to comply with the underlying logging system
+            </remarks>
+        </member>
+        <member name="P:Common.Logging.Factory.AbstractLogger.GlobalVariablesContext">
+            <summary>
+            Returns the global context for variables
+            </summary>
+        </member>
+        <member name="P:Common.Logging.Factory.AbstractLogger.ThreadVariablesContext">
+            <summary>
+            Returns the thread-specific context for variables
+            </summary>
+        </member>
+        <member name="T:Common.Logging.Factory.AbstractLogger.FormatMessageCallbackFormattedMessage">
+            <summary>
+            Format message on demand.
+            </summary>
+        </member>
+        <member name="M:Common.Logging.Factory.AbstractLogger.FormatMessageCallbackFormattedMessage.#ctor(System.Action{Common.Logging.FormatMessageHandler})">
+            <summary>
+            Initializes a new instance of the <see cref="T:Common.Logging.Factory.AbstractLogger.FormatMessageCallbackFormattedMessage"/> class.
+            </summary>
+            <param name="formatMessageCallback">The format message callback.</param>
+        </member>
+        <member name="M:Common.Logging.Factory.AbstractLogger.FormatMessageCallbackFormattedMessage.#ctor(System.IFormatProvider,System.Action{Common.Logging.FormatMessageHandler})">
+            <summary>
+            Initializes a new instance of the <see cref="T:Common.Logging.Factory.AbstractLogger.FormatMessageCallbackFormattedMessage"/> class.
+            </summary>
+            <param name="formatProvider">The format provider.</param>
+            <param name="formatMessageCallback">The format message callback.</param>
+        </member>
+        <member name="M:Common.Logging.Factory.AbstractLogger.FormatMessageCallbackFormattedMessage.ToString">
+            <summary>
+            Calls <see cref="F:Common.Logging.Factory.AbstractLogger.FormatMessageCallbackFormattedMessage.formatMessageCallback"/> and returns result.
+            </summary>
+            <returns></returns>
+        </member>
+        <member name="T:Common.Logging.Factory.AbstractLogger.StringFormatFormattedMessage">
+            <summary>
+            Format string on demand.
+            </summary>
+        </member>
+        <member name="M:Common.Logging.Factory.AbstractLogger.StringFormatFormattedMessage.#ctor(System.IFormatProvider,System.String,System.Object[])">
+            <summary>
+            Initializes a new instance of the <see cref="T:Common.Logging.Factory.AbstractLogger.StringFormatFormattedMessage"/> class.
+            </summary>
+            <param name="formatProvider">The format provider.</param>
+            <param name="message">The message.</param>
+            <param name="args">The args.</param>
+        </member>
+        <member name="M:Common.Logging.Factory.AbstractLogger.StringFormatFormattedMessage.ToString">
+            <summary>
+            Runs <see cref="M:System.String.Format(System.IFormatProvider,System.String,System.Object[])"/> on supplied arguemnts.
+            </summary>
+            <returns>string</returns>
+        </member>
+        <member name="T:Common.Logging.Factory.AbstractLogger.WriteHandler">
+            <summary>
+            Represents a method responsible for writing a message to the log system.
+            </summary>
+        </member>
+        <member name="T:Common.Logging.LogManager">
+            <summary>
+            Use the LogManager's <see cref="M:Common.Logging.LogManager.GetLogger(System.String)"/> or <see cref="M:Common.Logging.LogManager.GetLogger(System.Type)"/> 
+            methods to obtain <see cref="T:Common.Logging.ILog"/> instances for logging.
+            </summary>
+            <remarks>
+            For configuring the underlying log system using application configuration, see the example 
+            at <c>System.Configuration.ConfigurationManager</c>
+            For configuring programmatically, see the example section below.
+            </remarks>
+            <example>
+            The example below shows the typical use of LogManager to obtain a reference to a logger
+            and log an exception:
+            <code>
+            
+            ILog log = LogManager.GetLogger(this.GetType());
+            ...
+            try 
+            { 
+              /* .... */ 
+            }
+            catch(Exception ex)
+            {
+              log.ErrorFormat("Hi {0}", ex, "dude");
+            }
+            
+            </code>
+            The example below shows programmatic configuration of the underlying log system:
+            <code>
+            
+            // create properties
+            NameValueCollection properties = new NameValueCollection();
+            properties["showDateTime"] = "true";
+            
+            // set Adapter
+            Common.Logging.LogManager.Adapter = new 
+            Common.Logging.Simple.ConsoleOutLoggerFactoryAdapter(properties);
+            
+            </code>
+            </example>
+            <seealso cref="T:Common.Logging.ILog"/>
+            <seealso cref="P:Common.Logging.LogManager.Adapter"/>
+            <seealso cref="T:Common.Logging.ILoggerFactoryAdapter"/>
+            <author>Gilles Bayon</author>
+        </member>
+        <member name="M:Common.Logging.LogManager.#cctor">
+            <summary>
+            Performs static 1-time init of LogManager by calling <see cref="M:Common.Logging.LogManager.Reset"/>
+            </summary>
+        </member>
+        <member name="M:Common.Logging.LogManager.Reset">
+            <summary>
+            Reset the <see cref="N:Common.Logging"/> infrastructure to its default settings. This means, that configuration settings
+            will be re-read from section <c>&lt;common/logging&gt;</c> of your <c>app.config</c>.
+            </summary>
+            <remarks>
+            This is mainly used for unit testing, you wouldn't normally use this in your applications.<br/>
+            <b>Note:</b><see cref="T:Common.Logging.ILog"/> instances already handed out from this LogManager are not(!) affected. 
+            Resetting LogManager only affects new instances being handed out.
+            </remarks>
+        </member>
+        <member name="M:Common.Logging.LogManager.Reset(Common.Logging.IConfigurationReader)">
+            <summary>
+            Reset the <see cref="N:Common.Logging"/> infrastructure to its default settings. This means, that configuration settings
+            will be re-read from section <c>&lt;common/logging&gt;</c> of your <c>app.config</c>.
+            </summary>
+            <remarks>
+            This is mainly used for unit testing, you wouldn't normally use this in your applications.<br/>
+            <b>Note:</b><see cref="T:Common.Logging.ILog"/> instances already handed out from this LogManager are not(!) affected. 
+            Resetting LogManager only affects new instances being handed out.
+            </remarks>
+            <param name="reader">
+            the <see cref="T:Common.Logging.IConfigurationReader"/> instance to obtain settings for 
+            re-initializing the LogManager.
+            </param>
+        </member>
+        <member name="M:Common.Logging.LogManager.GetCurrentClassLogger">
+            <summary>
+            Gets the logger by calling <see cref="M:Common.Logging.ILoggerFactoryAdapter.GetLogger(System.Type)"/>
+            on the currently configured <see cref="P:Common.Logging.LogManager.Adapter"/> using the type of the calling class.
+            </summary>
+            <remarks>
+            This method needs to inspect the <see cref="T:System.Diagnostics.StackTrace"/> in order to determine the calling 
+            class. This of course comes with a performance penalty, thus you shouldn't call it too
+            often in your application.
+            </remarks>
+            <seealso cref="M:Common.Logging.LogManager.GetLogger(System.Type)"/>
+            <returns>the logger instance obtained from the current <see cref="P:Common.Logging.LogManager.Adapter"/></returns>
+        </member>
+        <member name="M:Common.Logging.LogManager.Common#Logging#ILogManager#GetCurrentClassLogger">
+            <summary>
+            Gets the logger by calling <see cref="M:Common.Logging.ILoggerFactoryAdapter.GetLogger(System.Type)"/>
+            on the currently configured <see cref="P:Common.Logging.LogManager.Adapter"/> using the type of the calling class.
+            </summary>
+            <remarks>
+            This method needs to inspect the <see cref="T:System.Diagnostics.StackTrace"/> in order to determine the calling 
+            class. This of course comes with a performance penalty, thus you shouldn't call it too
+            often in your application.
+            </remarks>
+            <seealso cref="M:Common.Logging.LogManager.GetLogger(System.Type)"/>
+            <returns>the logger instance obtained from the current <see cref="P:Common.Logging.LogManager.Adapter"/></returns>
+        </member>
+        <member name="M:Common.Logging.LogManager.GetLogger``1">
+            <summary>
+            Gets the logger by calling <see cref="M:Common.Logging.ILoggerFactoryAdapter.GetLogger(System.Type)"/>
+            on the currently configured <see cref="P:Common.Logging.LogManager.Adapter"/> using the specified type.
+            </summary>
+            <returns>the logger instance obtained from the current <see cref="P:Common.Logging.LogManager.Adapter"/></returns>
+        </member>
+        <member name="M:Common.Logging.LogManager.Common#Logging#ILogManager#GetLogger``1">
+            <summary>
+            Gets the logger by calling <see cref="M:Common.Logging.ILoggerFactoryAdapter.GetLogger(System.Type)"/>
+            on the currently configured <see cref="P:Common.Logging.LogManager.Adapter"/> using the specified type.
+            </summary>
+            <returns>the logger instance obtained from the current <see cref="P:Common.Logging.LogManager.Adapter"/></returns>
+        </member>
+        <member name="M:Common.Logging.LogManager.GetLogger(System.Type)">
+            <summary>
+            Gets the logger by calling <see cref="M:Common.Logging.ILoggerFactoryAdapter.GetLogger(System.Type)"/>
+            on the currently configured <see cref="P:Common.Logging.LogManager.Adapter"/> using the specified type.
+            </summary>
+            <param name="type">The type.</param>
+            <returns>the logger instance obtained from the current <see cref="P:Common.Logging.LogManager.Adapter"/></returns>
+        </member>
+        <member name="M:Common.Logging.LogManager.Common#Logging#ILogManager#GetLogger(System.Type)">
+            <summary>
+            Gets the logger by calling <see cref="M:Common.Logging.ILoggerFactoryAdapter.GetLogger(System.Type)"/>
+            on the currently configured <see cref="P:Common.Logging.LogManager.Adapter"/> using the specified type.
+            </summary>
+            <param name="type">The type.</param>
+            <returns>the logger instance obtained from the current <see cref="P:Common.Logging.LogManager.Adapter"/></returns>
+        </member>
+        <member name="M:Common.Logging.LogManager.GetLogger(System.String)">
+            <summary>
+            Gets the logger by calling <see cref="M:Common.Logging.ILoggerFactoryAdapter.GetLogger(System.String)"/>
+            on the currently configured <see cref="P:Common.Logging.LogManager.Adapter"/> using the specified key.
+            </summary>
+            <param name="key">The key.</param>
+            <returns>the logger instance obtained from the current <see cref="P:Common.Logging.LogManager.Adapter"/></returns>
+        </member>
+        <member name="M:Common.Logging.LogManager.Common#Logging#ILogManager#GetLogger(System.String)">
+            <summary>
+            Gets the logger by calling <see cref="M:Common.Logging.ILoggerFactoryAdapter.GetLogger(System.String)"/>
+            on the currently configured <see cref="P:Common.Logging.LogManager.Adapter"/> using the specified key.
+            </summary>
+            <param name="key">The key.</param>
+            <returns>the logger instance obtained from the current <see cref="P:Common.Logging.LogManager.Adapter"/></returns>
+        </member>
+        <member name="M:Common.Logging.LogManager.BuildLoggerFactoryAdapter">
+            <summary>
+            Builds the logger factory adapter.
+            </summary>
+            <returns>a factory adapter instance. Is never <c>null</c>.</returns>
+        </member>
+        <member name="M:Common.Logging.LogManager.BuildLoggerFactoryAdapterFromLogSettings(Common.Logging.Configuration.LogSetting)">
+            <summary>
+            Builds a <see cref="T:Common.Logging.ILoggerFactoryAdapter"/> instance from the given <see cref="T:Common.Logging.Configuration.LogSetting"/>
+            using <see cref="T:System.Activator"/>.
+            </summary>
+            <param name="setting"></param>
+            <returns>the <see cref="T:Common.Logging.ILoggerFactoryAdapter"/> instance. Is never <c>null</c></returns>
+        </member>
+        <member name="P:Common.Logging.LogManager.COMMON_LOGGING_SECTION">
+            <summary>
+            The key of the default configuration section to read settings from.
+            </summary>
+            <remarks>
+            You can always change the source of your configuration settings by setting another <see cref="T:Common.Logging.IConfigurationReader"/> instance
+            on <see cref="P:Common.Logging.LogManager.ConfigurationReader"/>.
+            </remarks>
+        </member>
+        <member name="P:Common.Logging.LogManager.Common#Logging#ILogManager#COMMON_LOGGING_SECTION">
+            <summary>
+            The key of the default configuration section to read settings from.
+            </summary>
+            <remarks>
+            You can always change the source of your configuration settings by setting another <see cref="T:Common.Logging.IConfigurationReader"/> instance
+            on <see cref="P:Common.Logging.LogManager.ConfigurationReader"/>.
+            </remarks>
+        </member>
+        <member name="P:Common.Logging.LogManager.ConfigurationReader">
+            <summary>
+            Gets the configuration reader used to initialize the LogManager.
+            </summary>
+            <remarks>Primarily used for testing purposes but maybe useful to obtain configuration
+            information from some place other than the .NET application configuration file.</remarks>
+            <value>The configuration reader.</value>
+        </member>
+        <member name="P:Common.Logging.LogManager.Common#Logging#ILogManager#ConfigurationReader">
+            <summary>
+            Gets the configuration reader used to initialize the LogManager.
+            </summary>
+            <remarks>Primarily used for testing purposes but maybe useful to obtain configuration
+            information from some place other than the .NET application configuration file.</remarks>
+            <value>The configuration reader.</value>
+        </member>
+        <member name="P:Common.Logging.LogManager.Adapter">
+            <summary>
+            Gets or sets the adapter.
+            </summary>
+            <value>The adapter.</value>
+        </member>
+        <member name="P:Common.Logging.LogManager.Common#Logging#ILogManager#Adapter">
+            <summary>
+            Gets or sets the adapter.
+            </summary>
+            <value>The adapter.</value>
+        </member>
+        <member name="T:Common.Logging.Simple.AbstractSimpleLogger">
+            <summary>
+            Abstract class providing a standard implementation of simple loggers.
+            </summary>
+            <author>Erich Eichinger</author>
+        </member>
+        <member name="M:Common.Logging.Simple.AbstractSimpleLogger.#ctor(System.String,Common.Logging.LogLevel,System.Boolean,System.Boolean,System.Boolean,System.String)">
+            <summary>
+            Creates and initializes a the simple logger.
+            </summary>
+            <param name="logName">The key, usually type key of the calling class, of the logger.</param>
+            <param name="logLevel">The current logging threshold. Messages recieved that are beneath this threshold will not be logged.</param>
+            <param name="showlevel">Include level in the log message.</param>
+            <param name="showDateTime">Include the current time in the log message.</param>
+            <param name="showLogName">Include the instance key in the log message.</param>
+            <param name="dateTimeFormat">The date and time format to use in the log message.</param>
+        </member>
+        <member name="M:Common.Logging.Simple.AbstractSimpleLogger.FormatOutput(System.Text.StringBuilder,Common.Logging.LogLevel,System.Object,System.Exception)">
+            <summary>
+            Appends the formatted message to the specified <see cref="T:System.Text.StringBuilder"/>.
+            </summary>
+            <param name="stringBuilder">the <see cref="T:System.Text.StringBuilder"/> that receíves the formatted message.</param>
+            <param name="level"></param>
+            <param name="message"></param>
+            <param name="e"></param>
+        </member>
+        <member name="M:Common.Logging.Simple.AbstractSimpleLogger.IsLevelEnabled(Common.Logging.LogLevel)">
+            <summary>
+            Determines if the given log level is currently enabled.
+            </summary>
+            <param name="level"></param>
+            <returns></returns>
+        </member>
+        <member name="P:Common.Logging.Simple.AbstractSimpleLogger.Name">
+            <summary>
+            The key of the logger.
+            </summary>
+        </member>
+        <member name="P:Common.Logging.Simple.AbstractSimpleLogger.ShowLevel">
+            <summary>
+            Include the current log level in the log message.
+            </summary>
+        </member>
+        <member name="P:Common.Logging.Simple.AbstractSimpleLogger.ShowDateTime">
+            <summary>
+            Include the current time in the log message.
+            </summary>
+        </member>
+        <member name="P:Common.Logging.Simple.AbstractSimpleLogger.ShowLogName">
+            <summary>
+            Include the instance key in the log message.
+            </summary>
+        </member>
+        <member name="P:Common.Logging.Simple.AbstractSimpleLogger.CurrentLogLevel">
+            <summary>
+            The current logging threshold. Messages recieved that are beneath this threshold will not be logged.
+            </summary>
+        </member>
+        <member name="P:Common.Logging.Simple.AbstractSimpleLogger.DateTimeFormat">
+            <summary>
+            The date and time format to use in the log message.
+            </summary>
+        </member>
+        <member name="P:Common.Logging.Simple.AbstractSimpleLogger.HasDateTimeFormat">
+            <summary>
+            Determines Whether <see cref="P:Common.Logging.Simple.AbstractSimpleLogger.DateTimeFormat"/> is set.
+            </summary>
+        </member>
+        <member name="P:Common.Logging.Simple.AbstractSimpleLogger.IsTraceEnabled">
+            <summary>
+            Returns <see langword="true"/> if the current <see cref="T:Common.Logging.LogLevel"/> is greater than or
+            equal to <see cref="F:Common.Logging.LogLevel.Trace"/>. If it is, all messages will be sent to <see cref="P:System.Console.Out"/>.
+            </summary>
+        </member>
+        <member name="P:Common.Logging.Simple.AbstractSimpleLogger.IsDebugEnabled">
+            <summary>
+            Returns <see langword="true"/> if the current <see cref="T:Common.Logging.LogLevel"/> is greater than or
+            equal to <see cref="F:Common.Logging.LogLevel.Debug"/>. If it is, all messages will be sent to <see cref="P:System.Console.Out"/>.
+            </summary>
+        </member>
+        <member name="P:Common.Logging.Simple.AbstractSimpleLogger.IsInfoEnabled">
+            <summary>
+            Returns <see langword="true"/> if the current <see cref="T:Common.Logging.LogLevel"/> is greater than or
+            equal to <see cref="F:Common.Logging.LogLevel.Info"/>. If it is, only messages with a <see cref="T:Common.Logging.LogLevel"/> of
+            <see cref="F:Common.Logging.LogLevel.Info"/>, <see cref="F:Common.Logging.LogLevel.Warn"/>, <see cref="F:Common.Logging.LogLevel.Error"/>, and 
+            <see cref="F:Common.Logging.LogLevel.Fatal"/> will be sent to <see cref="P:System.Console.Out"/>.
+            </summary>
+        </member>
+        <member name="P:Common.Logging.Simple.AbstractSimpleLogger.IsWarnEnabled">
+            <summary>
+            Returns <see langword="true"/> if the current <see cref="T:Common.Logging.LogLevel"/> is greater than or
+            equal to <see cref="F:Common.Logging.LogLevel.Warn"/>. If it is, only messages with a <see cref="T:Common.Logging.LogLevel"/> of
+            <see cref="F:Common.Logging.LogLevel.Warn"/>, <see cref="F:Common.Logging.LogLevel.Error"/>, and <see cref="F:Common.Logging.LogLevel.Fatal"/> 
+            will be sent to <see cref="P:System.Console.Out"/>.
+            </summary>
+        </member>
+        <member name="P:Common.Logging.Simple.AbstractSimpleLogger.IsErrorEnabled">
+            <summary>
+            Returns <see langword="true"/> if the current <see cref="T:Common.Logging.LogLevel"/> is greater than or
+            equal to <see cref="F:Common.Logging.LogLevel.Error"/>. If it is, only messages with a <see cref="T:Common.Logging.LogLevel"/> of
+            <see cref="F:Common.Logging.LogLevel.Error"/> and <see cref="F:Common.Logging.LogLevel.Fatal"/> will be sent to <see cref="P:System.Console.Out"/>.
+            </summary>
+        </member>
+        <member name="P:Common.Logging.Simple.AbstractSimpleLogger.IsFatalEnabled">
+            <summary>
+            Returns <see langword="true"/> if the current <see cref="T:Common.Logging.LogLevel"/> is greater than or
+            equal to <see cref="F:Common.Logging.LogLevel.Fatal"/>. If it is, only messages with a <see cref="T:Common.Logging.LogLevel"/> of
+            <see cref="F:Common.Logging.LogLevel.Fatal"/> will be sent to <see cref="P:System.Console.Out"/>.
+            </summary>
+        </member>
+        <member name="T:Common.Logging.Simple.AbstractSimpleLoggerFactoryAdapter">
+            <summary>
+            Base factory implementation for creating simple <see cref="T:Common.Logging.ILog"/> instances.
+            </summary>
+            <remarks>Default settings are LogLevel.All, showDateTime = true, showLogName = true, and no DateTimeFormat.
+            The keys in the NameValueCollection to configure this adapter are the following
+            <list type="bullet">
+                <item>level</item>
+                <item>showDateTime</item>
+                <item>showLogName</item>
+                <item>dateTimeFormat</item>
+            </list>
+            <example>
+            Here is an example how to implement your own logging adapter:
+            <code>
+            public class ConsoleOutLogger : AbstractSimpleLogger
+            {
+              public ConsoleOutLogger(string logName, LogLevel logLevel, bool showLevel, bool showDateTime, 
+            bool showLogName, string dateTimeFormat)
+                  : base(logName, logLevel, showLevel, showDateTime, showLogName, dateTimeFormat)
+              {
+              }
+            
+              protected override void WriteInternal(LogLevel level, object message, Exception e)
+              {
+                  // Use a StringBuilder for better performance
+                  StringBuilder sb = new StringBuilder();
+                  FormatOutput(sb, level, message, e);
+            
+                  // Print to the appropriate destination
+                  Console.Out.WriteLine(sb.ToString());
+              }
+            }
+            
+            public class ConsoleOutLoggerFactoryAdapter : AbstractSimpleLoggerFactoryAdapter
+            {
+              public ConsoleOutLoggerFactoryAdapter(NameValueCollection properties)
+                  : base(properties)
+              { }
+            
+              protected override ILog CreateLogger(string key, LogLevel level, bool showLevel, bool 
+            showDateTime, bool showLogName, string dateTimeFormat)
+              {
+                  ILog log = new ConsoleOutLogger(key, level, showLevel, showDateTime, showLogName, 
+            dateTimeFormat);
+                  return log;
+              }
+            }
+            </code>
+            </example>
+            </remarks>
+            <author>Gilles Bayon</author>
+            <author>Mark Pollack</author>
+            <author>Erich Eichinger</author>
+        </member>
+        <member name="M:Common.Logging.Simple.AbstractSimpleLoggerFactoryAdapter.#ctor(Common.Logging.Configuration.NameValueCollection)">
+            <summary>
+            Initializes a new instance of the <see cref="T:Common.Logging.Simple.AbstractSimpleLoggerFactoryAdapter"/> class.
+            </summary>
+            <remarks>
+            Looks for level, showDateTime, showLogName, dateTimeFormat items from 
+            <paramref key="properties"/> for use when the GetLogger methods are called.
+            <c>System.Configuration.ConfigurationManager</c> for more information on how to use the 
+            standard .NET application configuraiton file (App.config/Web.config) 
+            to configure this adapter.
+            </remarks>
+            <param name="properties">The key value collection, typically specified by the user in 
+            a configuration section named common/logging.</param>
+        </member>
+        <member name="M:Common.Logging.Simple.AbstractSimpleLoggerFactoryAdapter.#ctor(Common.Logging.LogLevel,System.Boolean,System.Boolean,System.Boolean,System.String)">
+            <summary>
+            Initializes a new instance of the <see cref="T:Common.Logging.Simple.AbstractSimpleLoggerFactoryAdapter"/> class with 
+            default settings for the loggers created by this factory.
+            </summary>
+        </member>
+        <member name="M:Common.Logging.Simple.AbstractSimpleLoggerFactoryAdapter.CreateLogger(System.String)">
+            <summary>
+            Create the specified logger instance
+            </summary>
+        </member>
+        <member name="M:Common.Logging.Simple.AbstractSimpleLoggerFactoryAdapter.CreateLogger(System.String,Common.Logging.LogLevel,System.Boolean,System.Boolean,System.Boolean,System.String)">
+            <summary>
+            Derived factories need to implement this method to create the
+            actual logger instance.
+            </summary>
+            <returns>a new logger instance. Must never be <c>null</c>!</returns>
+        </member>
+        <member name="P:Common.Logging.Simple.AbstractSimpleLoggerFactoryAdapter.Level">
+            <summary>
+            The default <see cref="T:Common.Logging.LogLevel"/> to use when creating new <see cref="T:Common.Logging.ILog"/> instances.
+            </summary>
+        </member>
+        <member name="P:Common.Logging.Simple.AbstractSimpleLoggerFactoryAdapter.ShowLevel">
+            <summary>
+            The default setting to use when creating new <see cref="T:Common.Logging.ILog"/> instances.
+            </summary>
+        </member>
+        <member name="P:Common.Logging.Simple.AbstractSimpleLoggerFactoryAdapter.ShowDateTime">
+            <summary>
+            The default setting to use when creating new <see cref="T:Common.Logging.ILog"/> instances.
+            </summary>
+        </member>
+        <member name="P:Common.Logging.Simple.AbstractSimpleLoggerFactoryAdapter.ShowLogName">
+            <summary>
+            The default setting to use when creating new <see cref="T:Common.Logging.ILog"/> instances.
+            </summary>
+        </member>
+        <member name="P:Common.Logging.Simple.AbstractSimpleLoggerFactoryAdapter.DateTimeFormat">
+            <summary>
+            The default setting to use when creating new <see cref="T:Common.Logging.ILog"/> instances.
+            </summary>
+        </member>
+        <member name="T:Common.Logging.Simple.CapturingLogger">
+            <summary>
+            A logger created by <see cref="T:Common.Logging.Simple.CapturingLoggerFactoryAdapter"/> that 
+            sends all log events to the owning adapter's <see cref="M:Common.Logging.Simple.CapturingLoggerFactoryAdapter.AddEvent(Common.Logging.Simple.CapturingLoggerEvent)"/>
+            </summary>
+            <author>Erich Eichinger</author>
+        </member>
+        <member name="F:Common.Logging.Simple.CapturingLogger.Owner">
+            <summary>
+            The adapter that created this logger instance.
+            </summary>
+        </member>
+        <member name="M:Common.Logging.Simple.CapturingLogger.Clear">
+            <summary>
+            Clears all captured events
+            </summary>
+        </member>
+        <member name="M:Common.Logging.Simple.CapturingLogger.ClearLastEvent">
+            <summary>
+            Resets the <see cref="P:Common.Logging.Simple.CapturingLogger.LastEvent"/> to <c>null</c>.
+            </summary>
+        </member>
+        <member name="F:Common.Logging.Simple.CapturingLogger.LoggerEvents">
+            <summary>
+            Holds the list of logged events.
+            </summary>
+            <remarks>
+            To access this collection in a multithreaded application, put a lock on the list instance.
+            </remarks>
+        </member>
+        <member name="M:Common.Logging.Simple.CapturingLogger.AddEvent(Common.Logging.Simple.CapturingLoggerEvent)">
+            <summary>
+            <see cref="T:Common.Logging.Simple.CapturingLogger"/> instances send their captured log events to this method.
+            </summary>
+        </member>
+        <member name="M:Common.Logging.Simple.CapturingLogger.#ctor(Common.Logging.Simple.CapturingLoggerFactoryAdapter,System.String)">
+            <summary>
+            Create a new logger instance.
+            </summary>
+        </member>
+        <member name="M:Common.Logging.Simple.CapturingLogger.WriteInternal(Common.Logging.LogLevel,System.Object,System.Exception)">
+            <summary>
+            Create a new <see cref="T:Common.Logging.Simple.CapturingLoggerEvent"/> and send it to <see cref="M:Common.Logging.Simple.CapturingLoggerFactoryAdapter.AddEvent(Common.Logging.Simple.CapturingLoggerEvent)"/>
+            </summary>
+            <param name="level"></param>
+            <param name="message"></param>
+            <param name="exception"></param>
+        </member>
+        <member name="P:Common.Logging.Simple.CapturingLogger.LastEvent">
+            <summary>
+            Holds the last log event received from any of this adapter's loggers.
+            </summary>
+        </member>
+        <member name="T:Common.Logging.Simple.CapturingLoggerEvent">
+            <summary>
+            A logging event captured by <see cref="T:Common.Logging.Simple.CapturingLogger"/>
+            </summary>
+            <author>Erich Eichinger</author>
+        </member>
+        <member name="F:Common.Logging.Simple.CapturingLoggerEvent.Source">
+            <summary>
+            The logger that logged this event
+            </summary>
+        </member>
+        <member name="F:Common.Logging.Simple.CapturingLoggerEvent.Level">
+            <summary>
+            The level used to log this event
+            </summary>
+        </member>
+        <member name="F:Common.Logging.Simple.CapturingLoggerEvent.MessageObject">
+            <summary>
+            The raw message object
+            </summary>
+        </member>
+        <member name="F:Common.Logging.Simple.CapturingLoggerEvent.Exception">
+            <summary>
+            A logged exception
+            </summary>
+        </member>
+        <member name="M:Common.Logging.Simple.CapturingLoggerEvent.#ctor(Common.Logging.Simple.CapturingLogger,Common.Logging.LogLevel,System.Object,System.Exception)">
+            <summary>
+            Create a new event instance
+            </summary>
+        </member>
+        <member name="P:Common.Logging.Simple.CapturingLoggerEvent.RenderedMessage">
+            <summary>
+            Retrieves the formatted message text
+            </summary>
+        </member>
+        <member name="T:Common.Logging.Simple.CapturingLoggerFactoryAdapter">
+            <summary>
+            An adapter, who's loggers capture all log events and send them to <see cref="M:Common.Logging.Simple.CapturingLoggerFactoryAdapter.AddEvent(Common.Logging.Simple.CapturingLoggerEvent)"/>. 
+            Retrieve the list of log events from <see cref="F:Common.Logging.Simple.CapturingLoggerFactoryAdapter.LoggerEvents"/>.
+            </summary>
+            <remarks>
+            This logger factory is mainly for debugging and test purposes.
+            <example>
+            This is an example how you might use this adapter for testing:
+            <code>
+            // configure for capturing
+            CapturingLoggerFactoryAdapter adapter = new CapturingLoggerFactoryAdapter();
+            LogManager.Adapter = adapter;
+            
+            // reset capture state
+            adapter.Clear();
+            // log something
+            ILog log = LogManager.GetCurrentClassLogger();
+            log.DebugFormat("Current Time:{0}", DateTime.Now);
+            
+            // check logged data
+            Assert.AreEqual(1, adapter.LoggerEvents.Count);
+            Assert.AreEqual(LogLevel.Debug, adapter.LastEvent.Level);
+            </code>
+            </example>
+            </remarks>
+            <author>Erich Eichinger</author>
+        </member>
+        <member name="M:Common.Logging.Simple.CapturingLoggerFactoryAdapter.Clear">
+            <summary>
+            Clears all captured events
+            </summary>
+        </member>
+        <member name="M:Common.Logging.Simple.CapturingLoggerFactoryAdapter.ClearLastEvent">
+            <summary>
+            Resets the <see cref="P:Common.Logging.Simple.CapturingLoggerFactoryAdapter.LastEvent"/> to <c>null</c>.
+            </summary>
+        </member>
+        <member name="F:Common.Logging.Simple.CapturingLoggerFactoryAdapter.LoggerEvents">
+            <summary>
+            Holds the list of logged events.
+            </summary>
+            <remarks>
+            To access this collection in a multithreaded application, put a lock on the list instance.
+            </remarks>
+        </member>
+        <member name="M:Common.Logging.Simple.CapturingLoggerFactoryAdapter.AddEvent(Common.Logging.Simple.CapturingLoggerEvent)">
+            <summary>
+            <see cref="T:Common.Logging.Simple.CapturingLogger"/> instances send their captured log events to this method.
+            </summary>
+        </member>
+        <member name="M:Common.Logging.Simple.CapturingLoggerFactoryAdapter.GetLogger(System.Type)">
+            <summary>
+            Get a <see cref="T:Common.Logging.Simple.CapturingLogger"/> instance for the given type.
+            </summary>
+        </member>
+        <member name="M:Common.Logging.Simple.CapturingLoggerFactoryAdapter.GetLogger(System.String)">
+            <summary>
+            Get a <see cref="T:Common.Logging.Simple.CapturingLogger"/> instance for the given key.
+            </summary>
+        </member>
+        <member name="P:Common.Logging.Simple.CapturingLoggerFactoryAdapter.LastEvent">
+            <summary>
+            Holds the last log event received from any of this adapter's loggers.
+            </summary>
+        </member>
+        <member name="T:Common.Logging.Simple.DebugOutLogger">
+            <summary>
+            Sends log messages to <see cref="T:System.Diagnostics.Debug"/>.
+            </summary>
+            <author>Gilles Bayon</author>
+        </member>
+        <member name="M:Common.Logging.Simple.DebugOutLogger.#ctor(System.String,Common.Logging.LogLevel,System.Boolean,System.Boolean,System.Boolean,System.String)">
+            <summary>
+            Creates and initializes a logger that writes messages to <see cref="T:System.Diagnostics.Debug"/>.
+            </summary>
+            <param name="logName">The key, usually type key of the calling class, of the logger.</param>
+            <param name="logLevel">The current logging threshold. Messages recieved that are beneath this threshold will not be logged.</param>
+            <param name="showLevel">Include the current log level in the log message.</param>
+            <param name="showDateTime">Include the current time in the log message.</param>
+            <param name="showLogName">Include the instance key in the log message.</param>
+            <param name="dateTimeFormat">The date and time format to use in the log message.</param>
+        </member>
+        <member name="M:Common.Logging.Simple.DebugOutLogger.WriteInternal(Common.Logging.LogLevel,System.Object,System.Exception)">
+            <summary>
+            Do the actual logging by constructing the log message using a <see cref="T:System.Text.StringBuilder"/> then
+            sending the output to <see cref="P:System.Console.Out"/>.
+            </summary>
+            <param name="level">The <see cref="T:Common.Logging.LogLevel"/> of the message.</param>
+            <param name="message">The log message.</param>
+            <param name="e">An optional <see cref="T:System.Exception"/> associated with the message.</param>
+        </member>
+        <member name="T:Common.Logging.Simple.DebugLoggerFactoryAdapter">
+            <summary>
+            Factory for creating <see cref="T:Common.Logging.ILog"/> instances that write data using <see cref="M:System.Diagnostics.Debug.WriteLine(System.String)"/>.
+            </summary>
+            <remarks>
+            <example>
+            Below is an example how to configure this adapter:
+            <code>
+            &lt;configuration&gt;
+            
+              &lt;configSections&gt;
+                &lt;sectionGroup key="common"&gt;
+                  &lt;section key="logging"
+                           type="Common.Logging.ConfigurationSectionHandler, Common.Logging"
+                           requirePermission="false" /&gt;
+                &lt;/sectionGroup&gt;
+              &lt;/configSections&gt;
+            
+              &lt;common&gt;
+                &lt;logging&gt;
+                  &lt;factoryAdapter type="Common.Logging.Simple.DebugLoggerFactoryAdapter, Common.Logging"&gt;
+                    &lt;arg key="level" value="ALL" /&gt;
+                  &lt;/factoryAdapter&gt;
+                &lt;/logging&gt;
+              &lt;/common&gt;
+            
+            &lt;/configuration&gt;
+            </code>
+            </example>
+            </remarks>
+            <seealso cref="T:Common.Logging.Simple.AbstractSimpleLoggerFactoryAdapter"/>
+            <author>Gilles Bayon</author>
+            <author>Mark Pollack</author>
+            <author>Erich Eichinger</author>
+        </member>
+        <member name="M:Common.Logging.Simple.DebugLoggerFactoryAdapter.#ctor">
+            <summary>
+            Initializes a new instance of the <see cref="T:Common.Logging.Simple.DebugLoggerFactoryAdapter"/> class using default 
+            settings.
+            </summary>
+        </member>
+        <member name="M:Common.Logging.Simple.DebugLoggerFactoryAdapter.#ctor(Common.Logging.Configuration.NameValueCollection)">
+            <summary>
+            Initializes a new instance of the <see cref="T:Common.Logging.Simple.DebugLoggerFactoryAdapter"/> class.
+            </summary>
+            <remarks>
+            Looks for level, showDateTime, showLogName, dateTimeFormat items from 
+            <paramref key="properties"/> for use when the GetLogger methods are called.
+            <see cref="T:System.Configuration.ConfigurationManager"/> for more information on how to use the 
+            standard .NET application configuraiton file (App.config/Web.config) 
+            to configure this adapter.
+            </remarks>
+            <param name="properties">The key value collection, typically specified by the user in 
+            a configuration section named common/logging.</param>
+        </member>
+        <member name="M:Common.Logging.Simple.DebugLoggerFactoryAdapter.#ctor(Common.Logging.LogLevel,System.Boolean,System.Boolean,System.Boolean,System.String)">
+            <summary>
+            Initializes a new instance of the <see cref="T:Common.Logging.Simple.AbstractSimpleLoggerFactoryAdapter"/> class with 
+            default settings for the loggers created by this factory.
+            </summary>
+        </member>
+        <member name="M:Common.Logging.Simple.DebugLoggerFactoryAdapter.CreateLogger(System.String,Common.Logging.LogLevel,System.Boolean,System.Boolean,System.Boolean,System.String)">
+            <summary>
+            Creates a new <see cref="T:Common.Logging.Simple.DebugOutLogger"/> instance.
+            </summary>
+        </member>
+        <member name="T:Common.Logging.Simple.ExceptionFormatter">
+            <summary>
+            </summary>
+        </member>
+        <member name="T:Common.Logging.Simple.NoOpLogger">
+            <summary>
+            Silently ignores all log messages.
+            </summary>
+            <author>Gilles Bayon</author>
+            <author>Erich Eichinger</author>
+        </member>
+        <member name="M:Common.Logging.Simple.NoOpLogger.Trace(System.Object)">
+            <summary>
+            Ignores message.
+            </summary>
+            <param name="message"></param>
+        </member>
+        <member name="M:Common.Logging.Simple.NoOpLogger.Trace(System.Object,System.Exception)">
+            <summary>
+            Ignores message.
+            </summary>
+            <param name="message"></param>
+            <param name="e"></param>
+        </member>
+        <member name="M:Common.Logging.Simple.NoOpLogger.TraceFormat(System.String,System.Object[])">
+            <summary>
+            Ignores message.
+            </summary>
+            <param name="format">The format of the message object to log.<see cref="M:System.String.Format(System.String,System.Object[])"/> </param>
+            <param name="args"></param>
+        </member>
+        <member name="M:Common.Logging.Simple.NoOpLogger.TraceFormat(System.String,System.Exception,System.Object[])">
+            <summary>
+            Ignores message.
+            </summary>
+            <param name="format">The format of the message object to log.<see cref="M:System.String.Format(System.String,System.Object[])"/> </param>
+            <param name="exception">The exception to log.</param>
+            <param name="args">the list of message format arguments</param>
+        </member>
+        <member name="M:Common.Logging.Simple.NoOpLogger.TraceFormat(System.IFormatProvider,System.String,System.Object[])">
+            <summary>
+            Ignores message.
+            </summary>
+            <param name="formatProvider">An <see cref="T:System.IFormatProvider"/> that supplies culture-specific formatting information.</param>
+            <param name="format">The format of the message object to log.<see cref="M:System.String.Format(System.String,System.Object[])"/> </param>
+            <param name="args">the list of message format arguments</param>
+        </member>
+        <member name="M:Common.Logging.Simple.NoOpLogger.TraceFormat(System.IFormatProvider,System.String,System.Exception,System.Object[])">
+            <summary>
+            Ignores message.
+            </summary>
+            <param name="formatProvider">An <see cref="T:System.IFormatProvider"/> that supplies culture-specific formatting information.</param>
+            <param name="format">The format of the message object to log.<see cref="M:System.String.Format(System.String,System.Object[])"/> </param>
+            <param name="exception">The exception to log.</param>
+            <param name="args">the list of message format arguments</param>
+        </member>
+        <member name="M:Common.Logging.Simple.NoOpLogger.Trace(System.Action{Common.Logging.FormatMessageHandler})">
+            <summary>
+            Ignores message.
+            </summary>
+            <param name="formatMessageCallback">A callback used by the logger to obtain the message if log level is matched</param>
+        </member>
+        <member name="M:Common.Logging.Simple.NoOpLogger.Trace(System.Action{Common.Logging.FormatMessageHandler},System.Exception)">
+            <summary>
+            Ignores message.
+            </summary>
+            <param name="formatMessageCallback">A callback used by the logger to obtain the message if log level is matched</param>
+            <param name="exception">The exception to log, including its stack trace.</param>
+        </member>
+        <member name="M:Common.Logging.Simple.NoOpLogger.Trace(System.IFormatProvider,System.Action{Common.Logging.FormatMessageHandler})">
+            <summary>
+            Ignores message.
+            </summary>
+            <param name="formatProvider">An <see cref="T:System.IFormatProvider"/> that supplies culture-specific formatting information.</param>
+            <param name="formatMessageCallback">A callback used by the logger to obtain the message if log level is matched</param>
+        </member>
+        <member name="M:Common.Logging.Simple.NoOpLogger.Trace(System.IFormatProvider,System.Action{Common.Logging.FormatMessageHandler},System.Exception)">
+            <summary>
+            Ignores message.
+            </summary>
+            <param name="formatProvider">An <see cref="T:System.IFormatProvider"/> that supplies culture-specific formatting information.</param>
+            <param name="formatMessageCallback">A callback used by the logger to obtain the message if log level is matched</param>
+            <param name="exception">The exception to log, including its stack trace.</param>
+        </member>
+        <member name="M:Common.Logging.Simple.NoOpLogger.Debug(System.Object)">
+            <summary>
+            Ignores message.
+            </summary>
+            <param name="message"></param>
+        </member>
+        <member name="M:Common.Logging.Simple.NoOpLogger.Debug(System.Object,System.Exception)">
+            <summary>
+            Ignores message.
+            </summary>
+            <param name="message"></param>
+            <param name="e"></param>
+        </member>
+        <member name="M:Common.Logging.Simple.NoOpLogger.DebugFormat(System.String,System.Object[])">
+            <summary>
+            Ignores message.
+            </summary>
+            <param name="format">The format of the message object to log.<see cref="M:System.String.Format(System.String,System.Object[])"/> </param>
+            <param name="args"></param>
+        </member>
+        <member name="M:Common.Logging.Simple.NoOpLogger.DebugFormat(System.String,System.Exception,System.Object[])">
+            <summary>
+            Ignores message.
+            </summary>
+            <param name="format">The format of the message object to log.<see cref="M:System.String.Format(System.String,System.Object[])"/> </param>
+            <param name="exception">The exception to log.</param>
+            <param name="args">the list of message format arguments</param>
+        </member>
+        <member name="M:Common.Logging.Simple.NoOpLogger.DebugFormat(System.IFormatProvider,System.String,System.Object[])">
+            <summary>
+            Ignores message.
+            </summary>
+            <param name="formatProvider">An <see cref="T:System.IFormatProvider"/> that supplies culture-specific formatting information.</param>
+            <param name="format">The format of the message object to log.<see cref="M:System.String.Format(System.String,System.Object[])"/> </param>
+            <param name="args">the list of message format arguments</param>
+        </member>
+        <member name="M:Common.Logging.Simple.NoOpLogger.DebugFormat(System.IFormatProvider,System.String,System.Exception,System.Object[])">
+            <summary>
+            Ignores message.
+            </summary>
+            <param name="formatProvider">An <see cref="T:System.IFormatProvider"/> that supplies culture-specific formatting information.</param>
+            <param name="format">The format of the message object to log.<see cref="M:System.String.Format(System.String,System.Object[])"/> </param>
+            <param name="exception">The exception to log.</param>
+            <param name="args">the list of message format arguments</param>
+        </member>
+        <member name="M:Common.Logging.Simple.NoOpLogger.Debug(System.Action{Common.Logging.FormatMessageHandler})">
+            <summary>
+            Ignores message.
+            </summary>
+            <param name="formatMessageCallback">A callback used by the logger to obtain the message if log level is matched</param>
+        </member>
+        <member name="M:Common.Logging.Simple.NoOpLogger.Debug(System.Action{Common.Logging.FormatMessageHandler},System.Exception)">
+            <summary>
+            Ignores message.
+            </summary>
+            <param name="formatMessageCallback">A callback used by the logger to obtain the message if log level is matched</param>
+            <param name="exception">The exception to log, including its stack Debug.</param>
+        </member>
+        <member name="M:Common.Logging.Simple.NoOpLogger.Debug(System.IFormatProvider,System.Action{Common.Logging.FormatMessageHandler})">
+            <summary>
+            Ignores message.
+            </summary>
+            <param name="formatProvider">An <see cref="T:System.IFormatProvider"/> that supplies culture-specific formatting information.</param>
+            <param name="formatMessageCallback">A callback used by the logger to obtain the message if log level is matched</param>
+        </member>
+        <member name="M:Common.Logging.Simple.NoOpLogger.Debug(System.IFormatProvider,System.Action{Common.Logging.FormatMessageHandler},System.Exception)">
+            <summary>
+            Ignores message.
+            </summary>
+            <param name="formatProvider">An <see cref="T:System.IFormatProvider"/> that supplies culture-specific formatting information.</param>
+            <param name="formatMessageCallback">A callback used by the logger to obtain the message if log level is matched</param>
+            <param name="exception">The exception to log, including its stack Debug.</param>
+        </member>
+        <member name="M:Common.Logging.Simple.NoOpLogger.Info(System.Object)">
+            <summary>
+            Ignores message.
+            </summary>
+            <param name="message"></param>
+        </member>
+        <member name="M:Common.Logging.Simple.NoOpLogger.Info(System.Object,System.Exception)">
+            <summary>
+            Ignores message.
+            </summary>
+            <param name="message"></param>
+            <param name="e"></param>
+        </member>
+        <member name="M:Common.Logging.Simple.NoOpLogger.InfoFormat(System.String,System.Object[])">
+            <summary>
+            Ignores message.
+            </summary>
+            <param name="format">The format of the message object to log.<see cref="M:System.String.Format(System.String,System.Object[])"/> </param>
+            <param name="args"></param>
+        </member>
+        <member name="M:Common.Logging.Simple.NoOpLogger.InfoFormat(System.String,System.Exception,System.Object[])">
+            <summary>
+            Ignores message.
+            </summary>
+            <param name="format">The format of the message object to log.<see cref="M:System.String.Format(System.String,System.Object[])"/> </param>
+            <param name="exception">The exception to log.</param>
+            <param name="args">the list of message format arguments</param>
+        </member>
+        <member name="M:Common.Logging.Simple.NoOpLogger.InfoFormat(System.IFormatProvider,System.String,System.Object[])">
+            <summary>
+            Ignores message.
+            </summary>
+            <param name="formatProvider">An <see cref="T:System.IFormatProvider"/> that supplies culture-specific formatting information.</param>
+            <param name="format">The format of the message object to log.<see cref="M:System.String.Format(System.String,System.Object[])"/> </param>
+            <param name="args">the list of message format arguments</param>
+        </member>
+        <member name="M:Common.Logging.Simple.NoOpLogger.InfoFormat(System.IFormatProvider,System.String,System.Exception,System.Object[])">
+            <summary>
+            Ignores message.
+            </summary>
+            <param name="formatProvider">An <see cref="T:System.IFormatProvider"/> that supplies culture-specific formatting information.</param>
+            <param name="format">The format of the message object to log.<see cref="M:System.String.Format(System.String,System.Object[])"/> </param>
+            <param name="exception">The exception to log.</param>
+            <param name="args">the list of message format arguments</param>
+        </member>
+        <member name="M:Common.Logging.Simple.NoOpLogger.Info(System.Action{Common.Logging.FormatMessageHandler})">
+            <summary>
+            Ignores message.
+            </summary>
+            <param name="formatMessageCallback">A callback used by the logger to obtain the message if log level is matched</param>
+        </member>
+        <member name="M:Common.Logging.Simple.NoOpLogger.Info(System.Action{Common.Logging.FormatMessageHandler},System.Exception)">
+            <summary>
+            Ignores message.
+            </summary>
+            <param name="formatMessageCallback">A callback used by the logger to obtain the message if log level is matched</param>
+            <param name="exception">The exception to log, including its stack Info.</param>
+        </member>
+        <member name="M:Common.Logging.Simple.NoOpLogger.Info(System.IFormatProvider,System.Action{Common.Logging.FormatMessageHandler})">
+            <summary>
+            Ignores message.
+            </summary>
+            <param name="formatProvider">An <see cref="T:System.IFormatProvider"/> that supplies culture-specific formatting information.</param>
+            <param name="formatMessageCallback">A callback used by the logger to obtain the message if log level is matched</param>
+        </member>
+        <member name="M:Common.Logging.Simple.NoOpLogger.Info(System.IFormatProvider,System.Action{Common.Logging.FormatMessageHandler},System.Exception)">
+            <summary>
+            Ignores message.
+            </summary>
+            <param name="formatProvider">An <see cref="T:System.IFormatProvider"/> that supplies culture-specific formatting information.</param>
+            <param name="formatMessageCallback">A callback used by the logger to obtain the message if log level is matched</param>
+            <param name="exception">The exception to log, including its stack Info.</param>
+        </member>
+        <member name="M:Common.Logging.Simple.NoOpLogger.Warn(System.Object)">
+            <summary>
+            Ignores message.
+            </summary>
+            <param name="message"></param>
+        </member>
+        <member name="M:Common.Logging.Simple.NoOpLogger.Warn(System.Object,System.Exception)">
+            <summary>
+            Ignores message.
+            </summary>
+            <param name="message"></param>
+            <param name="e"></param>
+        </member>
+        <member name="M:Common.Logging.Simple.NoOpLogger.WarnFormat(System.String,System.Object[])">
+            <summary>
+            Ignores message.
+            </summary>
+            <param name="format">The format of the message object to log.<see cref="M:System.String.Format(System.String,System.Object[])"/> </param>
+            <param name="args"></param>
+        </member>
+        <member name="M:Common.Logging.Simple.NoOpLogger.WarnFormat(System.String,System.Exception,System.Object[])">
+            <summary>
+            Ignores message.
+            </summary>
+            <param name="format">The format of the message object to log.<see cref="M:System.String.Format(System.String,System.Object[])"/> </param>
+            <param name="exception">The exception to log.</param>
+            <param name="args">the list of message format arguments</param>
+        </member>
+        <member name="M:Common.Logging.Simple.NoOpLogger.WarnFormat(System.IFormatProvider,System.String,System.Object[])">
+            <summary>
+            Ignores message.
+            </summary>
+            <param name="formatProvider">An <see cref="T:System.IFormatProvider"/> that supplies culture-specific formatting Information.</param>
+            <param name="format">The format of the message object to log.<see cref="M:System.String.Format(System.String,System.Object[])"/> </param>
+            <param name="args">the list of message format arguments</param>
+        </member>
+        <member name="M:Common.Logging.Simple.NoOpLogger.WarnFormat(System.IFormatProvider,System.String,System.Exception,System.Object[])">
+            <summary>
+            Ignores message.
+            </summary>
+            <param name="formatProvider">An <see cref="T:System.IFormatProvider"/> that supplies culture-specific formatting Information.</param>
+            <param name="format">The format of the message object to log.<see cref="M:System.String.Format(System.String,System.Object[])"/> </param>
+            <param name="exception">The exception to log.</param>
+            <param name="args">the list of message format arguments</param>
+        </member>
+        <member name="M:Common.Logging.Simple.NoOpLogger.Warn(System.Action{Common.Logging.FormatMessageHandler})">
+            <summary>
+            Ignores message.
+            </summary>
+            <param name="formatMessageCallback">A callback used by the logger to obtain the message if log level is matched</param>
+        </member>
+        <member name="M:Common.Logging.Simple.NoOpLogger.Warn(System.Action{Common.Logging.FormatMessageHandler},System.Exception)">
+            <summary>
+            Ignores message.
+            </summary>
+            <param name="formatMessageCallback">A callback used by the logger to obtain the message if log level is matched</param>
+            <param name="exception">The exception to log, including its stack Warn.</param>
+        </member>
+        <member name="M:Common.Logging.Simple.NoOpLogger.Warn(System.IFormatProvider,System.Action{Common.Logging.FormatMessageHandler})">
+            <summary>
+            Ignores message.
+            </summary>
+            <param name="formatProvider">An <see cref="T:System.IFormatProvider"/> that supplies culture-specific formatting information.</param>
+            <param name="formatMessageCallback">A callback used by the logger to obtain the message if log level is matched</param>
+        </member>
+        <member name="M:Common.Logging.Simple.NoOpLogger.Warn(System.IFormatProvider,System.Action{Common.Logging.FormatMessageHandler},System.Exception)">
+            <summary>
+            Ignores message.
+            </summary>
+            <param name="formatProvider">An <see cref="T:System.IFormatProvider"/> that supplies culture-specific formatting information.</param>
+            <param name="formatMessageCallback">A callback used by the logger to obtain the message if log level is matched</param>
+            <param name="exception">The exception to log, including its stack Warn.</param>
+        </member>
+        <member name="M:Common.Logging.Simple.NoOpLogger.Error(System.Object)">
+            <summary>
+            Ignores message.
+            </summary>
+            <param name="message"></param>
+        </member>
+        <member name="M:Common.Logging.Simple.NoOpLogger.Error(System.Object,System.Exception)">
+            <summary>
+            Ignores message.
+            </summary>
+            <param name="message"></param>
+            <param name="e"></param>
+        </member>
+        <member name="M:Common.Logging.Simple.NoOpLogger.ErrorFormat(System.String,System.Object[])">
+            <summary>
+            Ignores message.
+            </summary>
+            <param name="format">The format of the message object to log.<see cref="M:System.String.Format(System.String,System.Object[])"/> </param>
+            <param name="args"></param>
+        </member>
+        <member name="M:Common.Logging.Simple.NoOpLogger.ErrorFormat(System.String,System.Exception,System.Object[])">
+            <summary>
+            Ignores message.
+            </summary>
+            <param name="format">The format of the message object to log.<see cref="M:System.String.Format(System.String,System.Object[])"/> </param>
+            <param name="exception">The exception to log.</param>
+            <param name="args">the list of message format arguments</param>
+        </member>
+        <member name="M:Common.Logging.Simple.NoOpLogger.ErrorFormat(System.IFormatProvider,System.String,System.Object[])">
+            <summary>
+            Ignores message.
+            </summary>
+            <param name="formatProvider">An <see cref="T:System.IFormatProvider"/> that supplies culture-specific formatting Errorrmation.</param>
+            <param name="format">The format of the message object to log.<see cref="M:System.String.Format(System.String,System.Object[])"/> </param>
+            <param name="args">the list of message format arguments</param>
+        </member>
+        <member name="M:Common.Logging.Simple.NoOpLogger.ErrorFormat(System.IFormatProvider,System.String,System.Exception,System.Object[])">
+            <summary>
+            Ignores message.
+            </summary>
+            <param name="formatProvider">An <see cref="T:System.IFormatProvider"/> that supplies culture-specific formatting Errorrmation.</param>
+            <param name="format">The format of the message object to log.<see cref="M:System.String.Format(System.String,System.Object[])"/> </param>
+            <param name="exception">The exception to log.</param>
+            <param name="args">the list of message format arguments</param>
+        </member>
+        <member name="M:Common.Logging.Simple.NoOpLogger.Error(System.Action{Common.Logging.FormatMessageHandler})">
+            <summary>
+            Ignores message.
+            </summary>
+            <param name="formatMessageCallback">A callback used by the logger to obtain the message if log level is matched</param>
+        </member>
+        <member name="M:Common.Logging.Simple.NoOpLogger.Error(System.Action{Common.Logging.FormatMessageHandler},System.Exception)">
+            <summary>
+            Ignores message.
+            </summary>
+            <param name="formatMessageCallback">A callback used by the logger to obtain the message if log level is matched</param>
+            <param name="exception">The exception to log, including its stack Error.</param>
+        </member>
+        <member name="M:Common.Logging.Simple.NoOpLogger.Error(System.IFormatProvider,System.Action{Common.Logging.FormatMessageHandler})">
+            <summary>
+            Ignores message.
+            </summary>
+            <param name="formatProvider">An <see cref="T:System.IFormatProvider"/> that supplies culture-specific formatting information.</param>
+            <param name="formatMessageCallback">A callback used by the logger to obtain the message if log level is matched</param>
+        </member>
+        <member name="M:Common.Logging.Simple.NoOpLogger.Error(System.IFormatProvider,System.Action{Common.Logging.FormatMessageHandler},System.Exception)">
+            <summary>
+            Ignores message.
+            </summary>
+            <param name="formatProvider">An <see cref="T:System.IFormatProvider"/> that supplies culture-specific formatting information.</param>
+            <param name="formatMessageCallback">A callback used by the logger to obtain the message if log level is matched</param>
+            <param name="exception">The exception to log, including its stack Error.</param>
+        </member>
+        <member name="M:Common.Logging.Simple.NoOpLogger.Fatal(System.Object)">
+            <summary>
+            Ignores message.
+            </summary>
+            <param name="message"></param>
+        </member>
+        <member name="M:Common.Logging.Simple.NoOpLogger.Fatal(System.Object,System.Exception)">
+            <summary>
+            Ignores message.
+            </summary>
+            <param name="message"></param>
+            <param name="e"></param>
+        </member>
+        <member name="M:Common.Logging.Simple.NoOpLogger.FatalFormat(System.String,System.Object[])">
+            <summary>
+            Ignores message.
+            </summary>
+            <param name="format">The format of the message object to log.<see cref="M:System.String.Format(System.String,System.Object[])"/> </param>
+            <param name="args"></param>
+        </member>
+        <member name="M:Common.Logging.Simple.NoOpLogger.FatalFormat(System.String,System.Exception,System.Object[])">
+            <summary>
+            Ignores message.
+            </summary>
+            <param name="format">The format of the message object to log.<see cref="M:System.String.Format(System.String,System.Object[])"/> </param>
+            <param name="exception">The exception to log.</param>
+            <param name="args">the list of message format arguments</param>
+        </member>
+        <member name="M:Common.Logging.Simple.NoOpLogger.FatalFormat(System.IFormatProvider,System.String,System.Object[])">
+            <summary>
+            Ignores message.
+            </summary>
+            <param name="formatProvider">An <see cref="T:System.IFormatProvider"/> that supplies culture-specific formatting Fatalrmation.</param>
+            <param name="format">The format of the message object to log.<see cref="M:System.String.Format(System.String,System.Object[])"/> </param>
+            <param name="args">the list of message format arguments</param>
+        </member>
+        <member name="M:Common.Logging.Simple.NoOpLogger.FatalFormat(System.IFormatProvider,System.String,System.Exception,System.Object[])">
+            <summary>
+            Ignores message.
+            </summary>
+            <param name="formatProvider">An <see cref="T:System.IFormatProvider"/> that supplies culture-specific formatting Fatalrmation.</param>
+            <param name="format">The format of the message object to log.<see cref="M:System.String.Format(System.String,System.Object[])"/> </param>
+            <param name="exception">The exception to log.</param>
+            <param name="args">the list of message format arguments</param>
+        </member>
+        <member name="M:Common.Logging.Simple.NoOpLogger.Fatal(System.Action{Common.Logging.FormatMessageHandler})">
+            <summary>
+            Ignores message.
+            </summary>
+            <param name="formatMessageCallback">A callback used by the logger to obtain the message if log level is matched</param>
+        </member>
+        <member name="M:Common.Logging.Simple.NoOpLogger.Fatal(System.Action{Common.Logging.FormatMessageHandler},System.Exception)">
+            <summary>
+            Ignores message.
+            </summary>
+            <param name="formatMessageCallback">A callback used by the logger to obtain the message if log level is matched</param>
+            <param name="exception">The exception to log, including its stack Fatal.</param>
+        </member>
+        <member name="M:Common.Logging.Simple.NoOpLogger.Fatal(System.IFormatProvider,System.Action{Common.Logging.FormatMessageHandler})">
+            <summary>
+            Ignores message.
+            </summary>
+            <param name="formatProvider">An <see cref="T:System.IFormatProvider"/> that supplies culture-specific formatting information.</param>
+            <param name="formatMessageCallback">A callback used by the logger to obtain the message if log level is matched</param>
+        </member>
+        <member name="M:Common.Logging.Simple.NoOpLogger.Fatal(System.IFormatProvider,System.Action{Common.Logging.FormatMessageHandler},System.Exception)">
+            <summary>
+            Ignores message.
+            </summary>
+            <param name="formatProvider">An <see cref="T:System.IFormatProvider"/> that supplies culture-specific formatting information.</param>
+            <param name="formatMessageCallback">A callback used by the logger to obtain the message if log level is matched</param>
+            <param name="exception">The exception to log, including its stack Fatal.</param>
+        </member>
+        <member name="P:Common.Logging.Simple.NoOpLogger.IsTraceEnabled">
+            <summary>
+            Always returns <see langword="false" />.
+            </summary>
+        </member>
+        <member name="P:Common.Logging.Simple.NoOpLogger.IsDebugEnabled">
+            <summary>
+            Always returns <see langword="false" />.
+            </summary>
+        </member>
+        <member name="P:Common.Logging.Simple.NoOpLogger.IsInfoEnabled">
+            <summary>
+            Always returns <see langword="false" />.
+            </summary>
+        </member>
+        <member name="P:Common.Logging.Simple.NoOpLogger.IsWarnEnabled">
+            <summary>
+            Always returns <see langword="false" />.
+            </summary>
+        </member>
+        <member name="P:Common.Logging.Simple.NoOpLogger.IsErrorEnabled">
+            <summary>
+            Always returns <see langword="false" />.
+            </summary>
+        </member>
+        <member name="P:Common.Logging.Simple.NoOpLogger.IsFatalEnabled">
+            <summary>
+            Always returns <see langword="false" />.
+            </summary>
+        </member>
+        <member name="P:Common.Logging.Simple.NoOpLogger.GlobalVariablesContext">
+            <summary>
+            Returns the global context for variables
+            </summary>
+        </member>
+        <member name="P:Common.Logging.Simple.NoOpLogger.ThreadVariablesContext">
+            <summary>
+            Returns the thread-specific context for variables
+            </summary>
+        </member>
+        <member name="T:Common.Logging.Simple.NoOpLoggerFactoryAdapter">
+            <summary>
+            Factory for creating <see cref="T:Common.Logging.ILog"/> instances that silently ignores
+            logging requests.
+            </summary>
+            <remarks>
+            This logger adapter is the default used by Common.Logging if unconfigured. Using this logger adapter is the most efficient
+            way to suppress any logging output.
+            <example>
+            Below is an example how to configure this adapter:
+            <code>
+            &lt;configuration&gt;
+            
+              &lt;configSections&gt;
+                &lt;sectionGroup key="common"&gt;
+                  &lt;section key="logging"
+                           type="Common.Logging.ConfigurationSectionHandler, Common.Logging"
+                           requirePermission="false" /&gt;
+                &lt;/sectionGroup&gt;
+              &lt;/configSections&gt;
+            
+              &lt;common&gt;
+                &lt;logging&gt;
+                  &lt;factoryAdapter type="Common.Logging.Simple.NoOpLoggerFactoryAdapter, Common.Logging"&gt;
+                    &lt;arg key="level" value="ALL" /&gt;
+                  &lt;/factoryAdapter&gt;
+                &lt;/logging&gt;
+              &lt;/common&gt;
+            
+            &lt;/configuration&gt;
+            </code>
+            </example>
+            </remarks>
+            <author>Gilles Bayon</author>
+        </member>
+        <member name="M:Common.Logging.Simple.NoOpLoggerFactoryAdapter.#ctor">
+            <summary>
+            Constructor
+            </summary>
+        </member>
+        <member name="M:Common.Logging.Simple.NoOpLoggerFactoryAdapter.#ctor(Common.Logging.Configuration.NameValueCollection)">
+            <summary>
+            Constructor
+            </summary>
+        </member>
+        <member name="M:Common.Logging.Simple.NoOpLoggerFactoryAdapter.#ctor(System.Collections.Specialized.NameValueCollection)">
+            <summary>
+            Constructor for binary backwards compatibility
+            </summary>
+            <param name="properties">The properties.</param>
+        </member>
+        <member name="M:Common.Logging.Simple.NoOpLoggerFactoryAdapter.GetLogger(System.Type)">
+            <summary>
+            Get a ILog instance by type 
+            </summary>
+            <param name="type"></param>
+            <returns></returns>
+        </member>
+        <member name="M:Common.Logging.Simple.NoOpLoggerFactoryAdapter.Common#Logging#ILoggerFactoryAdapter#GetLogger(System.String)">
+            <summary>
+            Get a ILog instance by type key 
+            </summary>
+            <param name="key"></param>
+            <returns></returns>
+        </member>
+        <member name="T:Common.Logging.Simple.NoOpVariablesContext">
+            <summary>
+            A null-functionality implementation of <see cref="T:Common.Logging.IVariablesContext"/>
+            </summary>
+        </member>
+        <member name="M:Common.Logging.Simple.NoOpVariablesContext.Set(System.String,System.Object)">
+            <summary>
+            Sets the value of a new or existing variable within the global context
+            </summary>
+            <param name="key">The key of the variable that is to be added</param>
+            <param name="value">The value to add</param>
+        </member>
+        <member name="M:Common.Logging.Simple.NoOpVariablesContext.Get(System.String)">
+            <summary>
+            Gets the value of a variable within the global context
+            </summary>
+            <param name="key">The key of the variable to get</param>
+            <returns>The value or null if not found</returns>
+        </member>
+        <member name="M:Common.Logging.Simple.NoOpVariablesContext.Contains(System.String)">
+            <summary>
+            Checks if a variable is set within the global context
+            </summary>
+            <param name="key">The key of the variable to check for</param>
+            <returns>True if the variable is set</returns>
+        </member>
+        <member name="M:Common.Logging.Simple.NoOpVariablesContext.Remove(System.String)">
+            <summary>
+            Removes a variable from the global context by key
+            </summary>
+            <param name="key">The key of the variable to remove</param>
+        </member>
+        <member name="M:Common.Logging.Simple.NoOpVariablesContext.Clear">
+            <summary>
+            Clears the global context variables
+            </summary>
+        </member>
+        <member name="T:Common.Logging.Simple.CommonLoggingTraceListener">
+            <summary>
+            A <see cref="T:System.Diagnostics.TraceListener"/> implementation sending all <see cref="T:System.Diagnostics.Trace">System.Diagnostics.Trace</see> output to 
+            the Common.Logging infrastructure.
+            </summary>
+            <remarks>
+            This listener captures all output sent by calls to <see cref="T:System.Diagnostics.Trace">System.Diagnostics.Trace</see> and
+            and <see cref="T:System.Diagnostics.TraceSource"/> and sends it to an <see cref="T:Common.Logging.ILog"/> instance.<br/>
+            The <see cref="T:Common.Logging.ILog"/> instance to be used is obtained by calling
+            <see cref="M:Common.Logging.LogManager.GetLogger(System.String)"/>. The name of the logger is created by passing 
+            this listener's <see cref="P:System.Diagnostics.TraceListener.Name"/> and any <c>source</c> or <c>category</c> passed 
+            into this listener (see <see cref="M:System.Diagnostics.TraceListener.WriteLine(System.Object,System.String)"/> or <see cref="M:System.Diagnostics.TraceListener.TraceEvent(System.Diagnostics.TraceEventCache,System.String,System.Diagnostics.TraceEventType,System.Int32,System.String,System.Object[])"/> for example).
+            </remarks>
+            <example>
+            The snippet below shows how to add and configure this listener to your app.config:
+            <code lang="XML">
+            &lt;system.diagnostics&gt;
+              &lt;sharedListeners&gt;
+                &lt;add name="Diagnostics"
+                     type="Common.Logging.Simple.CommonLoggingTraceListener, Common.Logging"
+                     initializeData="DefaultTraceEventType=Information; LoggerNameFormat={listenerName}.{sourceName}"&gt;
+                  &lt;filter type="System.Diagnostics.EventTypeFilter" initializeData="Information"/&gt;
+                &lt;/add&gt;
+              &lt;/sharedListeners&gt;
+              &lt;trace&gt;
+                &lt;listeners&gt;
+                  &lt;add name="Diagnostics" /&gt;
+                &lt;/listeners&gt;
+              &lt;/trace&gt;
+            &lt;/system.diagnostics&gt;
+            </code>
+            </example>
+            <author>Erich Eichinger</author>
+        </member>
+        <member name="M:Common.Logging.Simple.CommonLoggingTraceListener.#ctor">
+            <summary>
+            Creates a new instance with the default name "Diagnostics" and <see cref="T:Common.Logging.LogLevel"/> "Trace".
+            </summary>
+        </member>
+        <member name="M:Common.Logging.Simple.CommonLoggingTraceListener.#ctor(System.String)">
+            <summary>
+            Creates a new instance initialized with properties from the <paramref name="initializeData"/>. string.
+            </summary>
+            <remarks>
+            <paramref name="initializeData"/> is a semicolon separated string of name/value pairs, where each pair has
+            the form <c>key=value</c>. E.g.
+            "<c>Name=MyLoggerName;LogLevel=Debug</c>"
+            </remarks>
+            <param name="initializeData">a semicolon separated list of name/value pairs.</param>
+        </member>
+        <member name="M:Common.Logging.Simple.CommonLoggingTraceListener.#ctor(System.Collections.Specialized.NameValueCollection)">
+            <summary>
+            Creates a new instance initialized with the specified properties.
+            </summary>
+            <param name="properties">name/value configuration properties.</param>
+        </member>
+        <member name="M:Common.Logging.Simple.CommonLoggingTraceListener.Log(System.Diagnostics.TraceEventType,System.String,System.Int32,System.String,System.Object[])">
+            <summary>
+            Logs the given message to the Common.Logging infrastructure.
+            </summary>
+            <param name="eventType">the eventType</param>
+            <param name="source">the <see cref="T:System.Diagnostics.TraceSource"/> name or category name passed into e.g. <see cref="M:System.Diagnostics.Trace.Write(System.Object,System.String)"/>.</param>
+            <param name="id">the id of this event</param>
+            <param name="format">the message format</param>
+            <param name="args">the message arguments</param>
+        </member>
+        <member name="M:Common.Logging.Simple.CommonLoggingTraceListener.Write(System.Object)">
+            <summary>
+            Writes message to logger provided by <see cref="M:Common.Logging.LogManager.GetLogger(System.String)"/>.
+            </summary>
+        </member>
+        <member name="M:Common.Logging.Simple.CommonLoggingTraceListener.Write(System.Object,System.String)">
+            <summary>
+            Writes message to logger provided by <see cref="M:Common.Logging.LogManager.GetLogger(System.String)"/>.
+            </summary>
+        </member>
+        <member name="M:Common.Logging.Simple.CommonLoggingTraceListener.Write(System.String)">
+            <summary>
+            Writes message to logger provided by <see cref="M:Common.Logging.LogManager.GetLogger(System.String)"/>.
+            </summary>
+        </member>
+        <member name="M:Common.Logging.Simple.CommonLoggingTraceListener.Write(System.String,System.String)">
+            <summary>
+            Writes message to logger provided by <see cref="M:Common.Logging.LogManager.GetLogger(System.String)"/>.
+            </summary>
+        </member>
+        <member name="M:Common.Logging.Simple.CommonLoggingTraceListener.WriteLine(System.Object)">
+            <summary>
+            Writes message to logger provided by <see cref="M:Common.Logging.LogManager.GetLogger(System.String)"/>.
+            </summary>
+        </member>
+        <member name="M:Common.Logging.Simple.CommonLoggingTraceListener.WriteLine(System.Object,System.String)">
+            <summary>
+            Writes message to logger provided by <see cref="M:Common.Logging.LogManager.GetLogger(System.String)"/>.
+            </summary>
+        </member>
+        <member name="M:Common.Logging.Simple.CommonLoggingTraceListener.WriteLine(System.String)">
+            <summary>
+            Writes message to logger provided by <see cref="M:Common.Logging.LogManager.GetLogger(System.String)"/>.
+            </summary>
+        </member>
+        <member name="M:Common.Logging.Simple.CommonLoggingTraceListener.WriteLine(System.String,System.String)">
+            <summary>
+            Writes message to logger provided by <see cref="M:Common.Logging.LogManager.GetLogger(System.String)"/>
+            </summary>
+        </member>
+        <member name="M:Common.Logging.Simple.CommonLoggingTraceListener.TraceEvent(System.Diagnostics.TraceEventCache,System.String,System.Diagnostics.TraceEventType,System.Int32)">
+            <summary>
+            Writes message to logger provided by <see cref="M:Common.Logging.LogManager.GetLogger(System.String)"/>
+            </summary>
+        </member>
+        <member name="M:Common.Logging.Simple.CommonLoggingTraceListener.TraceEvent(System.Diagnostics.TraceEventCache,System.String,System.Diagnostics.TraceEventType,System.Int32,System.String)">
+            <summary>
+            Writes message to logger provided by <see cref="M:Common.Logging.LogManager.GetLogger(System.String)"/>
+            </summary>
+        </member>
+        <member name="M:Common.Logging.Simple.CommonLoggingTraceListener.TraceEvent(System.Diagnostics.TraceEventCache,System.String,System.Diagnostics.TraceEventType,System.Int32,System.String,System.Object[])">
+            <summary>
+            Writes message to logger provided by <see cref="M:Common.Logging.LogManager.GetLogger(System.String)"/>
+            </summary>
+        </member>
+        <member name="M:Common.Logging.Simple.CommonLoggingTraceListener.TraceData(System.Diagnostics.TraceEventCache,System.String,System.Diagnostics.TraceEventType,System.Int32,System.Object[])">
+            <summary>
+            Writes message to logger provided by <see cref="M:Common.Logging.LogManager.GetLogger(System.String)"/>
+            </summary>
+        </member>
+        <member name="M:Common.Logging.Simple.CommonLoggingTraceListener.TraceData(System.Diagnostics.TraceEventCache,System.String,System.Diagnostics.TraceEventType,System.Int32,System.Object)">
+            <summary>
+            Writes message to logger provided by <see cref="M:Common.Logging.LogManager.GetLogger(System.String)"/>
+            </summary>
+        </member>
+        <member name="P:Common.Logging.Simple.CommonLoggingTraceListener.DefaultTraceEventType">
+            <summary>
+            Sets the default <see cref="T:System.Diagnostics.TraceEventType"/> to use for logging
+            all events emitted by <see cref="T:System.Diagnostics.Trace"/><c>.Write(...)</c> and
+            <see cref="T:System.Diagnostics.Trace"/><c>.WriteLine(...)</c> methods.
+            </summary>
+            <remarks>
+            This listener captures all output sent by calls to <see cref="T:System.Diagnostics.Trace"/> and
+            sends it to an <see cref="T:Common.Logging.ILog"/> instance using the <see cref="T:Common.Logging.LogLevel"/> specified
+            on <see cref="T:Common.Logging.LogLevel"/>.
+            </remarks>
+        </member>
+        <member name="P:Common.Logging.Simple.CommonLoggingTraceListener.LoggerNameFormat">
+            <summary>
+            Format to use for creating the logger name. Defaults to "{listenerName}.{sourceName}".
+            </summary>
+            <remarks>
+            Available placeholders are:
+            <list type="bullet">
+            <item>{listenerName}: the configured name of this listener instance.</item>
+            <item>{sourceName}: the trace source name an event originates from (see e.g. <see cref="M:System.Diagnostics.TraceListener.TraceEvent(System.Diagnostics.TraceEventCache,System.String,System.Diagnostics.TraceEventType,System.Int32,System.String,System.Object[])"/>.</item>
+            </list>
+            </remarks>
+        </member>
+        <member name="T:Common.Logging.ConfigurationSectionHandler">
+            <summary>
+            Used in an application's configuration file (App.Config or Web.Config) to configure the logging subsystem.
+            </summary>
+            <example>
+            An example configuration section that writes log messages to the Console using the
+            built-in Console Logger.
+            <code lang="XML">
+            &lt;configuration&gt;
+                &lt;configSections&gt;
+                  &lt;sectionGroup name=&quot;common&quot;&gt;
+                    &lt;section name=&quot;logging&quot; type=&quot;Common.Logging.ConfigurationSectionHandler, Common.Logging&quot; /&gt;
+                  &lt;/sectionGroup&gt; 
+                &lt;/configSections&gt;
+                &lt;common&gt;
+                  &lt;logging&gt;
+                    &lt;factoryAdapter type=&quot;Common.Logging.Simple.ConsoleOutLoggerFactoryAdapter, Common.Logging&quot;&gt;
+                      &lt;arg key=&quot;showLogName&quot; value=&quot;true&quot; /&gt;
+                      &lt;arg key=&quot;showDateTime&quot; value=&quot;true&quot; /&gt;
+                      &lt;arg key=&quot;level&quot; value=&quot;ALL&quot; /&gt;
+                      &lt;arg key=&quot;dateTimeFormat&quot; value=&quot;yyyy/MM/dd HH:mm:ss:fff&quot; /&gt;
+                    &lt;/factoryAdapter&gt;
+                  &lt;/logging&gt;
+                &lt;/common&gt;
+            &lt;/configuration&gt;
+            </code>
+            </example>
+        </member>
+        <member name="M:Common.Logging.ConfigurationSectionHandler.#cctor">
+            <summary>
+            Ensure static fields get initialized before any class member 
+            can be accessed (avoids beforeFieldInit)
+            </summary>
+        </member>
+        <member name="M:Common.Logging.ConfigurationSectionHandler.#ctor">
+            <summary>
+            Constructor
+            </summary>
+        </member>
+        <member name="M:Common.Logging.ConfigurationSectionHandler.ReadConfiguration(System.Xml.XmlNode)">
+            <summary>
+            Retrieves the <see cref="T:System.Type"/> of the logger the use by looking at the logFactoryAdapter element
+            of the logging configuration element.
+            </summary>
+            <param name="section"></param>
+            <returns>
+            A <see cref="T:Common.Logging.Configuration.LogSetting"/> object containing the specified type that implements 
+            <see cref="T:Common.Logging.ILoggerFactoryAdapter"/> along with zero or more properties that will be 
+            passed to the logger factory adapter's constructor as an <see cref="T:System.Collections.IDictionary"/>.
+            </returns>
+        </member>
+        <member name="M:Common.Logging.ConfigurationSectionHandler.Create(Common.Logging.Configuration.LogSetting,System.Object,System.Xml.XmlNode)">
+            <summary>
+            Verifies that the logFactoryAdapter element appears once in the configuration section.
+            </summary>
+            <param name="parent">settings of a parent section - atm this must always be null</param>
+            <param name="configContext">Additional information about the configuration process.</param>
+            <param name="section">The configuration section to apply an XPath query too.</param>
+            <returns>
+            A <see cref="T:Common.Logging.Configuration.LogSetting"/> object containing the specified logFactoryAdapter type
+            along with user supplied configuration properties.
+            </returns>
+        </member>
+        <member name="M:Common.Logging.ConfigurationSectionHandler.System#Configuration#IConfigurationSectionHandler#Create(System.Object,System.Object,System.Xml.XmlNode)">
+            <summary>
+            Verifies that the logFactoryAdapter element appears once in the configuration section.
+            </summary>
+            <param name="parent">The parent of the current item.</param>
+            <param name="configContext">Additional information about the configuration process.</param>
+            <param name="section">The configuration section to apply an XPath query too.</param>
+            <returns>
+            A <see cref="T:Common.Logging.Configuration.LogSetting"/> object containing the specified logFactoryAdapter type
+            along with user supplied configuration properties.
+            </returns>
+        </member>
+        <member name="T:Common.Logging.Simple.ConsoleOutLogger">
+            <summary>
+            Sends log messages to <see cref="P:System.Console.Out"/>.
+            </summary>
+            <author>Gilles Bayon</author>
+        </member>
+        <member name="M:Common.Logging.Simple.ConsoleOutLogger.#ctor(System.String,Common.Logging.LogLevel,System.Boolean,System.Boolean,System.Boolean,System.String)">
+            <summary>
+            Creates and initializes a logger that writes messages to <see cref="P:System.Console.Out"/>.
+            </summary>
+            <param name="logName">The name, usually type name of the calling class, of the logger.</param>
+            <param name="logLevel">The current logging threshold. Messages recieved that are beneath this threshold will not be logged.</param>
+            <param name="showLevel">Include the current log level in the log message.</param>
+            <param name="showDateTime">Include the current time in the log message.</param>
+            <param name="showLogName">Include the instance name in the log message.</param>
+            <param name="dateTimeFormat">The date and time format to use in the log message.</param>
+        </member>
+        <member name="M:Common.Logging.Simple.ConsoleOutLogger.#ctor(System.String,Common.Logging.LogLevel,System.Boolean,System.Boolean,System.Boolean,System.String,System.Boolean)">
+            <summary>
+            Creates and initializes a logger that writes messages to <see cref="P:System.Console.Out"/>.
+            </summary>
+            <param name="logName">The name, usually type name of the calling class, of the logger.</param>
+            <param name="logLevel">The current logging threshold. Messages recieved that are beneath this threshold will not be logged.</param>
+            <param name="showLevel">Include the current log level in the log message.</param>
+            <param name="showDateTime">Include the current time in the log message.</param>
+            <param name="showLogName">Include the instance name in the log message.</param>
+            <param name="dateTimeFormat">The date and time format to use in the log message.</param>
+            <param name="useColor">Use color when writing the log message.</param>
+        </member>
+        <member name="M:Common.Logging.Simple.ConsoleOutLogger.WriteInternal(Common.Logging.LogLevel,System.Object,System.Exception)">
+            <summary>
+            Do the actual logging by constructing the log message using a <see cref="T:System.Text.StringBuilder"/> then
+            sending the output to <see cref="P:System.Console.Out"/>.
+            </summary>
+            <param name="level">The <see cref="T:Common.Logging.LogLevel"/> of the message.</param>
+            <param name="message">The log message.</param>
+            <param name="e">An optional <see cref="T:System.Exception"/> associated with the message.</param>
+        </member>
+        <member name="T:Common.Logging.Simple.ConsoleOutLoggerFactoryAdapter">
+            <summary>
+            Factory for creating <see cref="T:Common.Logging.ILog"/> instances that write data to <see cref="P:System.Console.Out"/>.
+            </summary>
+            <remarks>
+            <example>
+            Below is an example how to configure this adapter:
+            <code>
+            &lt;configuration&gt;
+            
+              &lt;configSections&gt;
+                &lt;sectionGroup name="common"&gt;
+                  &lt;section name="logging"
+                           type="Common.Logging.ConfigurationSectionHandler, Common.Logging"
+                           requirePermission="false" /&gt;
+                &lt;/sectionGroup&gt;
+              &lt;/configSections&gt;
+            
+              &lt;common&gt;
+                &lt;logging&gt;
+                  &lt;factoryAdapter type="Common.Logging.Simple.ConsoleOutLoggerFactoryAdapter, Common.Logging"&gt;
+                    &lt;arg key="level" value="ALL" /&gt;
+                  &lt;/factoryAdapter&gt;
+                &lt;/logging&gt;
+              &lt;/common&gt;
+            
+            &lt;/configuration&gt;
+            </code>
+            </example>
+            </remarks>
+            <seealso cref="T:Common.Logging.Simple.AbstractSimpleLoggerFactoryAdapter"/>
+            <seealso cref="P:Common.Logging.LogManager.Adapter"/>
+            <seealso cref="T:Common.Logging.ConfigurationSectionHandler"/>
+            <author>Gilles Bayon</author>
+            <author>Mark Pollack</author>
+            <author>Erich Eichinger</author>
+        </member>
+        <member name="M:Common.Logging.Simple.ConsoleOutLoggerFactoryAdapter.#ctor">
+            <summary>
+            Initializes a new instance of the <see cref="T:Common.Logging.Simple.ConsoleOutLoggerFactoryAdapter"/> class using default 
+            settings.
+            </summary>
+        </member>
+        <member name="M:Common.Logging.Simple.ConsoleOutLoggerFactoryAdapter.#ctor(Common.Logging.Configuration.NameValueCollection)">
+            <summary>
+            Initializes a new instance of the <see cref="T:Common.Logging.Simple.ConsoleOutLoggerFactoryAdapter"/> class.
+            </summary>
+            <remarks>
+            Looks for level, showDateTime, showLogName, dateTimeFormat items from 
+            <paramref name="properties"/> for use when the GetLogger methods are called.
+            <see cref="T:Common.Logging.ConfigurationSectionHandler"/> for more information on how to use the 
+            standard .NET application configuraiton file (App.config/Web.config) 
+            to configure this adapter.
+            </remarks>
+            <param name="properties">The name value collection, typically specified by the user in 
+            a configuration section named common/logging.</param>
+        </member>
+        <member name="M:Common.Logging.Simple.ConsoleOutLoggerFactoryAdapter.#ctor(System.Collections.Specialized.NameValueCollection)">
+            <summary>
+            Constructor for binary backwards compatibility with non-portableversions
+            </summary>
+            <param name="properties">The properties.</param>
+        </member>
+        <member name="M:Common.Logging.Simple.ConsoleOutLoggerFactoryAdapter.#ctor(Common.Logging.LogLevel,System.Boolean,System.Boolean,System.Boolean,System.String)">
+            <summary>
+            Initializes a new instance of the <see cref="T:Common.Logging.Simple.AbstractSimpleLoggerFactoryAdapter"/> class with 
+            default settings for the loggers created by this factory.
+            </summary>
+        </member>
+        <member name="M:Common.Logging.Simple.ConsoleOutLoggerFactoryAdapter.#ctor(Common.Logging.LogLevel,System.Boolean,System.Boolean,System.Boolean,System.String,System.Boolean)">
+            <summary>
+            Initializes a new instance of the <see cref="T:Common.Logging.Simple.AbstractSimpleLoggerFactoryAdapter"/> class with 
+            default settings for the loggers created by this factory.
+            </summary>
+        </member>
+        <member name="M:Common.Logging.Simple.ConsoleOutLoggerFactoryAdapter.CreateLogger(System.String,Common.Logging.LogLevel,System.Boolean,System.Boolean,System.Boolean,System.String)">
+            <summary>
+            Creates a new <see cref="T:Common.Logging.Simple.ConsoleOutLogger"/> instance.
+            </summary>
+        </member>
+        <member name="T:Common.Logging.Simple.TraceLogger">
+            <summary>
+            Logger sending everything to the trace output stream using <see cref="T:System.Diagnostics.Trace"/>.
+            </summary>
+            <remarks>
+            Beware not to use <see cref="T:Common.Logging.Simple.CommonLoggingTraceListener"/> in combination with this logger as 
+            this would result in an endless loop for obvious reasons!
+            </remarks>
+            <seealso cref="P:Common.Logging.LogManager.Adapter"/>
+            <seealso cref="T:Common.Logging.ConfigurationSectionHandler"/>
+            <author>Gilles Bayon</author>
+            <author>Erich Eichinger</author>
+        </member>
+        <member name="M:Common.Logging.Simple.TraceLogger.#ctor(System.Boolean,System.String,Common.Logging.LogLevel,System.Boolean,System.Boolean,System.Boolean,System.String)">
+            <summary>
+            Creates a new TraceLogger instance.
+            </summary>
+            <param name="useTraceSource">whether to use <see cref="T:System.Diagnostics.TraceSource"/> or <see cref="T:System.Diagnostics.Trace"/> for logging.</param>
+            <param name="logName">the name of this logger</param>
+            <param name="logLevel">the default log level to use</param>
+            <param name="showLevel">Include the current log level in the log message.</param>
+            <param name="showDateTime">Include the current time in the log message.</param>
+            <param name="showLogName">Include the instance name in the log message.</param>
+            <param name="dateTimeFormat">The date and time format to use in the log message.</param>
+        </member>
+        <member name="M:Common.Logging.Simple.TraceLogger.IsLevelEnabled(Common.Logging.LogLevel)">
+            <summary>
+            Determines if the given log level is currently enabled.
+            checks <see cref="P:System.Diagnostics.TraceSource.Switch"/> if <see cref="P:Common.Logging.Simple.TraceLoggerFactoryAdapter.UseTraceSource"/> is true.
+            </summary>
+        </member>
+        <member name="M:Common.Logging.Simple.TraceLogger.WriteInternal(Common.Logging.LogLevel,System.Object,System.Exception)">
+            <summary>
+            Do the actual logging.
+            </summary>
+            <param name="level"></param>
+            <param name="message"></param>
+            <param name="e"></param>
+        </member>
+        <member name="M:Common.Logging.Simple.TraceLogger.OnDeserialization(System.Object)">
+            <summary>
+            Called after deserialization completed.
+            </summary>
+        </member>
+        <member name="T:Common.Logging.Simple.TraceLogger.FormatOutputMessage">
+            <summary>
+            Used to defer message formatting until it is really needed.
+            </summary>
+            <remarks>
+            This class also improves performance when multiple 
+            <see cref="T:System.Diagnostics.TraceListener"/>s are configured.
+            </remarks>
+        </member>
+        <member name="T:Common.Logging.Simple.TraceLoggerFactoryAdapter">
+            <summary>
+            Factory for creating <see cref="T:Common.Logging.ILog"/> instances that send 
+            everything to the <see cref="T:System.Diagnostics.Trace"/> output stream.
+            </summary>
+            <remarks>
+            Beware not to use <see cref="T:Common.Logging.Simple.CommonLoggingTraceListener"/> in combination with this logger factory
+            as this would result in an endless loop for obvious reasons!
+            <example>
+            Below is an example how to configure this adapter:
+            <code>
+            &lt;configuration&gt;
+            
+              &lt;configSections&gt;
+                &lt;sectionGroup name="common"&gt;
+                  &lt;section name="logging"
+                           type="Common.Logging.ConfigurationSectionHandler, Common.Logging"
+                           requirePermission="false" /&gt;
+                &lt;/sectionGroup&gt;
+              &lt;/configSections&gt;
+            
+              &lt;common&gt;
+                &lt;logging&gt;
+                  &lt;factoryAdapter type="Common.Logging.Simple.TraceLoggerFactoryAdapter, Common.Logging"&gt;
+                    &lt;arg key="level" value="ALL" /&gt;
+                  &lt;/factoryAdapter&gt;
+                &lt;/logging&gt;
+              &lt;/common&gt;
+            
+            &lt;/configuration&gt;
+            </code>
+            </example>
+            </remarks>
+            <seealso cref="T:Common.Logging.Simple.AbstractSimpleLoggerFactoryAdapter"/>
+            <seealso cref="P:Common.Logging.LogManager.Adapter"/>
+            <seealso cref="T:Common.Logging.ConfigurationSectionHandler"/>
+            <author>Gilles Bayon</author>
+            <author>Mark Pollack</author>
+            <author>Erich Eichinger</author>
+        </member>
+        <member name="M:Common.Logging.Simple.TraceLoggerFactoryAdapter.#ctor">
+            <summary>
+            Initializes a new instance of the <see cref="T:Common.Logging.Simple.TraceLoggerFactoryAdapter"/> class using default settings.
+            </summary>
+        </member>
+        <member name="M:Common.Logging.Simple.TraceLoggerFactoryAdapter.#ctor(Common.Logging.Configuration.NameValueCollection)">
+            <summary>
+            Initializes a new instance of the <see cref="T:Common.Logging.Simple.TraceLoggerFactoryAdapter"/> class.
+            </summary>
+            <remarks>
+            Looks for level, showDateTime, showLogName, dateTimeFormat items from 
+            <paramref name="properties"/> for use when the GetLogger methods are called.
+            <see cref="T:Common.Logging.ConfigurationSectionHandler"/> for more information on how to use the 
+            standard .NET application configuraiton file (App.config/Web.config) 
+            to configure this adapter.
+            </remarks>
+            <param name="properties">The name value collection, typically specified by the user in 
+            a configuration section named common/logging.</param>
+        </member>
+        <member name="M:Common.Logging.Simple.TraceLoggerFactoryAdapter.#ctor(Common.Logging.LogLevel,System.Boolean,System.Boolean,System.Boolean,System.String,System.Boolean)">
+            <summary>
+            Initializes a new instance of the <see cref="T:Common.Logging.Simple.AbstractSimpleLoggerFactoryAdapter"/> class with 
+            default settings for the loggers created by this factory.
+            </summary>
+        </member>
+        <member name="M:Common.Logging.Simple.TraceLoggerFactoryAdapter.CreateLogger(System.String,Common.Logging.LogLevel,System.Boolean,System.Boolean,System.Boolean,System.String)">
+            <summary>
+            Creates a new <see cref="T:Common.Logging.Simple.TraceLogger"/> instance.
+            </summary>
+        </member>
+        <member name="P:Common.Logging.Simple.TraceLoggerFactoryAdapter.UseTraceSource">
+            <summary>
+            Whether to use <see cref="T:System.Diagnostics.Trace"/>.<c>TraceXXXX(string,object[])</c> methods for logging
+            or <see cref="T:System.Diagnostics.TraceSource"/>.
+            </summary>
+        </member>
+    </members>
+</doc>
