Copying //UE4/Dev-Core to //UE4/Main

diff --git a/Engine/Source/Programs/CrashReporter/CrashReportWebSite/packages/Common.Logging.Core.3.3.1/lib/net35/Common.Logging.Core.xml b/Engine/Source/Programs/CrashReporter/CrashReportWebSite/packages/Common.Logging.Core.3.3.1/lib/net35/Common.Logging.Core.xml
new file mode 100644
--- /dev/null
+++ b/Engine/Source/Programs/CrashReporter/CrashReportWebSite/packages/Common.Logging.Core.3.3.1/lib/net35/Common.Logging.Core.xml
@@ -0,0 +1,876 @@
+<?xml version="1.0"?>
+<doc>
+    <assembly>
+        <name>Common.Logging.Core</name>
+    </assembly>
+    <members>
+        <member name="T:Common.Logging.Factory.StringFormatMethodAttribute">
+            <summary>
+            Indicates that the marked method builds string by format pattern and (optional) arguments.
+            Parameter, which contains format string, should be given in constructor. The format string
+            should be in <see cref="M:System.String.Format(System.IFormatProvider,System.String,System.Object[])"/>-like form
+            </summary>
+            <example><code>
+            [StringFormatMethod("message")]
+            public void ShowError(string message, params object[] args) { /* do something */ }
+            public void Foo() {
+              ShowError("Failed: {0}"); // Warning: Non-existing argument in format string
+            }
+            </code></example>
+        </member>
+        <member name="M:Common.Logging.Factory.StringFormatMethodAttribute.#ctor(System.String)">
+            <param name="formatParameterName">
+            Specifies which parameter of an annotated method should be treated as format-string
+            </param>
+        </member>
+        <member name="P:Common.Logging.Factory.StringFormatMethodAttribute.FormatParameterName">
+            <summary>
+            The name of the string parameter being formatted
+            </summary>
+        </member>
+        <member name="T:Common.Logging.FormatMessageHandler">
+            <summary>
+             The type of method that is passed into e.g. <see cref="M:Common.Logging.ILog.Debug(System.Action{Common.Logging.FormatMessageHandler})"/> 
+             and allows the callback method to "submit" it's message to the underlying output system.
+            </summary>
+            <param name="format">the format argument as in <see cref="M:System.String.Format(System.String,System.Object[])"/></param>
+            <param name="args">the argument list as in <see cref="M:System.String.Format(System.String,System.Object[])"/></param>
+            <seealso cref="T:Common.Logging.ILog"/>
+             <author>Erich Eichinger</author>
+        </member>
+        <member name="T:Common.Logging.IConfigurationReader">
+            <summary>
+            Interface for basic operations to read .NET application configuration information.
+            </summary>
+            <remarks>Provides a simple abstraction to handle BCL API differences between .NET 1.x and 2.0. Also
+            useful for testing scenarios.</remarks>
+            <author>Mark Pollack</author>
+        </member>
+        <member name="M:Common.Logging.IConfigurationReader.GetSection(System.String)">
+            <summary>
+            Parses the configuration section and returns the resulting object.
+            </summary>
+            <remarks>
+            <p>
+            Primary purpose of this method is to allow us to parse and 
+            load configuration sections using the same API regardless
+            of the .NET framework version.
+            </p>
+            
+            See also <c>System.Configuration.ConfigurationManager</c>
+            </remarks>
+            <param name="sectionName">Name of the configuration section.</param>
+            <returns>Object created by a corresponding IConfigurationSectionHandler.</returns>
+        </member>
+        <member name="T:Common.Logging.ILog">
+            <summary>
+            A simple logging interface abstracting logging APIs. 
+            </summary>
+            <remarks>
+            <para>
+            Implementations should defer calling a message's <see cref="M:System.Object.ToString"/> until the message really needs
+            to be logged to avoid performance penalties.
+            </para>
+            <para>
+            Each <see cref="T:Common.Logging.ILog"/> log method offers to pass in a <see cref="T:System.Action`1"/> instead of the actual message.
+            Using this style has the advantage to defer possibly expensive message argument evaluation and formatting (and formatting arguments!) until the message gets
+            actually logged. If the message is not logged at all (e.g. due to <see cref="T:Common.Logging.LogLevel"/> settings), 
+            you won't have to pay the peformance penalty of creating the message.
+            </para>
+            </remarks>
+            <example>
+            The example below demonstrates using callback style for creating the message, where the call to the 
+            <see cref="M:System.Random.NextDouble"/> and the underlying <see cref="M:System.String.Format(System.String,System.Object[])"/> only happens, if level <see cref="F:Common.Logging.LogLevel.Debug"/> is enabled:
+            <code>
+            Log.Debug( m=&gt;m("result is {0}", random.NextDouble()) );
+            Log.Debug(delegate(m) { m("result is {0}", random.NextDouble()); });
+            </code>
+            </example>
+            <seealso cref="T:System.Action`1"/>
+            <author>Mark Pollack</author>
+            <author>Bruno Baia</author>
+            <author>Erich Eichinger</author>
+        </member>
+        <member name="M:Common.Logging.ILog.Trace(System.Object)">
+            <summary>
+            Log a message object with the <see cref="F:Common.Logging.LogLevel.Trace"/> level.
+            </summary>
+            <param name="message">The message object to log.</param>
+        </member>
+        <member name="M:Common.Logging.ILog.Trace(System.Object,System.Exception)">
+            <summary>
+            Log a message object with the <see cref="F:Common.Logging.LogLevel.Trace"/> level including
+            the stack trace of the <see cref="T:System.Exception"/> passed
+            as a parameter.
+            </summary>
+            <param name="message">The message object to log.</param>
+            <param name="exception">The exception to log, including its stack trace.</param>
+        </member>
+        <member name="M:Common.Logging.ILog.TraceFormat(System.String,System.Object[])">
+            <summary>
+            Log a message with the <see cref="F:Common.Logging.LogLevel.Trace"/> level.
+            </summary>
+            <param name="format">The format of the message object to log.<see cref="M:System.String.Format(System.String,System.Object[])"/> </param>
+            <param name="args">the list of format arguments</param>
+        </member>
+        <member name="M:Common.Logging.ILog.TraceFormat(System.String,System.Exception,System.Object[])">
+            <summary>
+            Log a message with the <see cref="F:Common.Logging.LogLevel.Trace"/> level.
+            </summary>
+            <param name="format">The format of the message object to log.<see cref="M:System.String.Format(System.String,System.Object[])"/> </param>
+            <param name="exception">The exception to log.</param>
+            <param name="args">the list of format arguments</param>
+        </member>
+        <member name="M:Common.Logging.ILog.TraceFormat(System.IFormatProvider,System.String,System.Object[])">
+            <summary>
+            Log a message with the <see cref="F:Common.Logging.LogLevel.Trace"/> level.
+            </summary>
+            <param name="formatProvider">An <see cref="T:System.IFormatProvider"/> that supplies culture-specific formatting information.</param>
+            <param name="format">The format of the message object to log.<see cref="M:System.String.Format(System.String,System.Object[])"/> </param>
+            <param name="args"></param>
+        </member>
+        <member name="M:Common.Logging.ILog.TraceFormat(System.IFormatProvider,System.String,System.Exception,System.Object[])">
+            <summary>
+            Log a message with the <see cref="F:Common.Logging.LogLevel.Trace"/> level.
+            </summary>
+            <param name="formatProvider">An <see cref="T:System.IFormatProvider"/> that supplies culture-specific formatting information.</param>
+            <param name="format">The format of the message object to log.<see cref="M:System.String.Format(System.String,System.Object[])"/> </param>
+            <param name="exception">The exception to log.</param>
+            <param name="args"></param>
+        </member>
+        <member name="M:Common.Logging.ILog.Trace(System.Action{Common.Logging.FormatMessageHandler})">
+            <summary>
+            Log a message with the <see cref="F:Common.Logging.LogLevel.Trace"/> level using a callback to obtain the message
+            </summary>
+            <remarks>
+            Using this method avoids the cost of creating a message and evaluating message arguments 
+            that probably won't be logged due to loglevel settings.
+            </remarks>
+            <param name="formatMessageCallback">A callback used by the logger to obtain the message if log level is matched</param>
+        </member>
+        <member name="M:Common.Logging.ILog.Trace(System.Action{Common.Logging.FormatMessageHandler},System.Exception)">
+            <summary>
+            Log a message with the <see cref="F:Common.Logging.LogLevel.Trace"/> level using a callback to obtain the message
+            </summary>
+            <remarks>
+            Using this method avoids the cost of creating a message and evaluating message arguments 
+            that probably won't be logged due to loglevel settings.
+            </remarks>
+            <param name="formatMessageCallback">A callback used by the logger to obtain the message if log level is matched</param>
+            <param name="exception">The exception to log, including its stack trace.</param>
+        </member>
+        <member name="M:Common.Logging.ILog.Trace(System.IFormatProvider,System.Action{Common.Logging.FormatMessageHandler})">
+            <summary>
+            Log a message with the <see cref="F:Common.Logging.LogLevel.Trace"/> level using a callback to obtain the message
+            </summary>
+            <remarks>
+            Using this method avoids the cost of creating a message and evaluating message arguments 
+            that probably won't be logged due to loglevel settings.
+            </remarks>
+            <param name="formatProvider">An <see cref="T:System.IFormatProvider"/> that supplies culture-specific formatting information.</param>
+            <param name="formatMessageCallback">A callback used by the logger to obtain the message if log level is matched</param>
+        </member>
+        <member name="M:Common.Logging.ILog.Trace(System.IFormatProvider,System.Action{Common.Logging.FormatMessageHandler},System.Exception)">
+            <summary>
+            Log a message with the <see cref="F:Common.Logging.LogLevel.Trace"/> level using a callback to obtain the message
+            </summary>
+            <remarks>
+            Using this method avoids the cost of creating a message and evaluating message arguments 
+            that probably won't be logged due to loglevel settings.
+            </remarks>
+            <param name="formatProvider">An <see cref="T:System.IFormatProvider"/> that supplies culture-specific formatting information.</param>
+            <param name="formatMessageCallback">A callback used by the logger to obtain the message if log level is matched</param>
+            <param name="exception">The exception to log, including its stack trace.</param>
+        </member>
+        <member name="M:Common.Logging.ILog.Debug(System.Object)">
+            <summary>
+            Log a message object with the <see cref="F:Common.Logging.LogLevel.Debug"/> level.
+            </summary>
+            <param name="message">The message object to log.</param>
+        </member>
+        <member name="M:Common.Logging.ILog.Debug(System.Object,System.Exception)">
+            <summary>
+            Log a message object with the <see cref="F:Common.Logging.LogLevel.Debug"/> level including
+            the stack trace of the <see cref="T:System.Exception"/> passed
+            as a parameter.
+            </summary>
+            <param name="message">The message object to log.</param>
+            <param name="exception">The exception to log, including its stack trace.</param>
+        </member>
+        <member name="M:Common.Logging.ILog.DebugFormat(System.String,System.Object[])">
+            <summary>
+            Log a message with the <see cref="F:Common.Logging.LogLevel.Debug"/> level.
+            </summary>
+            <param name="format">The format of the message object to log.<see cref="M:System.String.Format(System.String,System.Object[])"/> </param>
+            <param name="args">the list of format arguments</param>
+        </member>
+        <member name="M:Common.Logging.ILog.DebugFormat(System.String,System.Exception,System.Object[])">
+            <summary>
+            Log a message with the <see cref="F:Common.Logging.LogLevel.Debug"/> level.
+            </summary>
+            <param name="format">The format of the message object to log.<see cref="M:System.String.Format(System.String,System.Object[])"/> </param>
+            <param name="exception">The exception to log.</param>
+            <param name="args">the list of format arguments</param>
+        </member>
+        <member name="M:Common.Logging.ILog.DebugFormat(System.IFormatProvider,System.String,System.Object[])">
+            <summary>
+            Log a message with the <see cref="F:Common.Logging.LogLevel.Debug"/> level.
+            </summary>
+            <param name="formatProvider">An <see cref="T:System.IFormatProvider"/> that supplies culture-specific formatting information.</param>
+            <param name="format">The format of the message object to log.<see cref="M:System.String.Format(System.String,System.Object[])"/> </param>
+            <param name="args"></param>
+        </member>
+        <member name="M:Common.Logging.ILog.DebugFormat(System.IFormatProvider,System.String,System.Exception,System.Object[])">
+            <summary>
+            Log a message with the <see cref="F:Common.Logging.LogLevel.Debug"/> level.
+            </summary>
+            <param name="formatProvider">An <see cref="T:System.IFormatProvider"/> that supplies culture-specific formatting information.</param>
+            <param name="format">The format of the message object to log.<see cref="M:System.String.Format(System.String,System.Object[])"/> </param>
+            <param name="exception">The exception to log.</param>
+            <param name="args"></param>
+        </member>
+        <member name="M:Common.Logging.ILog.Debug(System.Action{Common.Logging.FormatMessageHandler})">
+            <summary>
+            Log a message with the <see cref="F:Common.Logging.LogLevel.Debug"/> level using a callback to obtain the message
+            </summary>
+            <remarks>
+            Using this method avoids the cost of creating a message and evaluating message arguments 
+            that probably won't be logged due to loglevel settings.
+            </remarks>
+            <param name="formatMessageCallback">A callback used by the logger to obtain the message if log level is matched</param>
+        </member>
+        <member name="M:Common.Logging.ILog.Debug(System.Action{Common.Logging.FormatMessageHandler},System.Exception)">
+            <summary>
+            Log a message with the <see cref="F:Common.Logging.LogLevel.Debug"/> level using a callback to obtain the message
+            </summary>
+            <remarks>
+            Using this method avoids the cost of creating a message and evaluating message arguments 
+            that probably won't be logged due to loglevel settings.
+            </remarks>
+            <param name="formatMessageCallback">A callback used by the logger to obtain the message if log level is matched</param>
+            <param name="exception">The exception to log, including its stack trace.</param>
+        </member>
+        <member name="M:Common.Logging.ILog.Debug(System.IFormatProvider,System.Action{Common.Logging.FormatMessageHandler})">
+            <summary>
+            Log a message with the <see cref="F:Common.Logging.LogLevel.Debug"/> level using a callback to obtain the message
+            </summary>
+            <remarks>
+            Using this method avoids the cost of creating a message and evaluating message arguments 
+            that probably won't be logged due to loglevel settings.
+            </remarks>
+            <param name="formatProvider">An <see cref="T:System.IFormatProvider"/> that supplies culture-specific formatting information.</param>
+            <param name="formatMessageCallback">A callback used by the logger to obtain the message if log level is matched</param>
+        </member>
+        <member name="M:Common.Logging.ILog.Debug(System.IFormatProvider,System.Action{Common.Logging.FormatMessageHandler},System.Exception)">
+            <summary>
+            Log a message with the <see cref="F:Common.Logging.LogLevel.Debug"/> level using a callback to obtain the message
+            </summary>
+            <remarks>
+            Using this method avoids the cost of creating a message and evaluating message arguments 
+            that probably won't be logged due to loglevel settings.
+            </remarks>
+            <param name="formatProvider">An <see cref="T:System.IFormatProvider"/> that supplies culture-specific formatting information.</param>
+            <param name="formatMessageCallback">A callback used by the logger to obtain the message if log level is matched</param>
+            <param name="exception">The exception to log, including its stack Debug.</param>
+        </member>
+        <member name="M:Common.Logging.ILog.Info(System.Object)">
+            <summary>
+            Log a message object with the <see cref="F:Common.Logging.LogLevel.Info"/> level.
+            </summary>
+            <param name="message">The message object to log.</param>
+        </member>
+        <member name="M:Common.Logging.ILog.Info(System.Object,System.Exception)">
+            <summary>
+            Log a message object with the <see cref="F:Common.Logging.LogLevel.Info"/> level including
+            the stack trace of the <see cref="T:System.Exception"/> passed
+            as a parameter.
+            </summary>
+            <param name="message">The message object to log.</param>
+            <param name="exception">The exception to log, including its stack trace.</param>
+        </member>
+        <member name="M:Common.Logging.ILog.InfoFormat(System.String,System.Object[])">
+            <summary>
+            Log a message with the <see cref="F:Common.Logging.LogLevel.Info"/> level.
+            </summary>
+            <param name="format">The format of the message object to log.<see cref="M:System.String.Format(System.String,System.Object[])"/> </param>
+            <param name="args">the list of format arguments</param>
+        </member>
+        <member name="M:Common.Logging.ILog.InfoFormat(System.String,System.Exception,System.Object[])">
+            <summary>
+            Log a message with the <see cref="F:Common.Logging.LogLevel.Info"/> level.
+            </summary>
+            <param name="format">The format of the message object to log.<see cref="M:System.String.Format(System.String,System.Object[])"/> </param>
+            <param name="exception">The exception to log.</param>
+            <param name="args">the list of format arguments</param>
+        </member>
+        <member name="M:Common.Logging.ILog.InfoFormat(System.IFormatProvider,System.String,System.Object[])">
+            <summary>
+            Log a message with the <see cref="F:Common.Logging.LogLevel.Info"/> level.
+            </summary>
+            <param name="formatProvider">An <see cref="T:System.IFormatProvider"/> that supplies culture-specific formatting information.</param>
+            <param name="format">The format of the message object to log.<see cref="M:System.String.Format(System.String,System.Object[])"/> </param>
+            <param name="args"></param>
+        </member>
+        <member name="M:Common.Logging.ILog.InfoFormat(System.IFormatProvider,System.String,System.Exception,System.Object[])">
+            <summary>
+            Log a message with the <see cref="F:Common.Logging.LogLevel.Info"/> level.
+            </summary>
+            <param name="formatProvider">An <see cref="T:System.IFormatProvider"/> that supplies culture-specific formatting information.</param>
+            <param name="format">The format of the message object to log.<see cref="M:System.String.Format(System.String,System.Object[])"/> </param>
+            <param name="exception">The exception to log.</param>
+            <param name="args"></param>
+        </member>
+        <member name="M:Common.Logging.ILog.Info(System.Action{Common.Logging.FormatMessageHandler})">
+            <summary>
+            Log a message with the <see cref="F:Common.Logging.LogLevel.Info"/> level using a callback to obtain the message
+            </summary>
+            <remarks>
+            Using this method avoids the cost of creating a message and evaluating message arguments 
+            that probably won't be logged due to loglevel settings.
+            </remarks>
+            <param name="formatMessageCallback">A callback used by the logger to obtain the message if log level is matched</param>
+        </member>
+        <member name="M:Common.Logging.ILog.Info(System.Action{Common.Logging.FormatMessageHandler},System.Exception)">
+            <summary>
+            Log a message with the <see cref="F:Common.Logging.LogLevel.Info"/> level using a callback to obtain the message
+            </summary>
+            <remarks>
+            Using this method avoids the cost of creating a message and evaluating message arguments 
+            that probably won't be logged due to loglevel settings.
+            </remarks>
+            <param name="formatMessageCallback">A callback used by the logger to obtain the message if log level is matched</param>
+            <param name="exception">The exception to log, including its stack trace.</param>
+        </member>
+        <member name="M:Common.Logging.ILog.Info(System.IFormatProvider,System.Action{Common.Logging.FormatMessageHandler})">
+            <summary>
+            Log a message with the <see cref="F:Common.Logging.LogLevel.Info"/> level using a callback to obtain the message
+            </summary>
+            <remarks>
+            Using this method avoids the cost of creating a message and evaluating message arguments 
+            that probably won't be logged due to loglevel settings.
+            </remarks>
+            <param name="formatProvider">An <see cref="T:System.IFormatProvider"/> that supplies culture-specific formatting information.</param>
+            <param name="formatMessageCallback">A callback used by the logger to obtain the message if log level is matched</param>
+        </member>
+        <member name="M:Common.Logging.ILog.Info(System.IFormatProvider,System.Action{Common.Logging.FormatMessageHandler},System.Exception)">
+            <summary>
+            Log a message with the <see cref="F:Common.Logging.LogLevel.Info"/> level using a callback to obtain the message
+            </summary>
+            <remarks>
+            Using this method avoids the cost of creating a message and evaluating message arguments 
+            that probably won't be logged due to loglevel settings.
+            </remarks>
+            <param name="formatProvider">An <see cref="T:System.IFormatProvider"/> that supplies culture-specific formatting information.</param>
+            <param name="formatMessageCallback">A callback used by the logger to obtain the message if log level is matched</param>
+            <param name="exception">The exception to log, including its stack Info.</param>
+        </member>
+        <member name="M:Common.Logging.ILog.Warn(System.Object)">
+            <summary>
+            Log a message object with the <see cref="F:Common.Logging.LogLevel.Warn"/> level.
+            </summary>
+            <param name="message">The message object to log.</param>
+        </member>
+        <member name="M:Common.Logging.ILog.Warn(System.Object,System.Exception)">
+            <summary>
+            Log a message object with the <see cref="F:Common.Logging.LogLevel.Warn"/> level including
+            the stack trace of the <see cref="T:System.Exception"/> passed
+            as a parameter.
+            </summary>
+            <param name="message">The message object to log.</param>
+            <param name="exception">The exception to log, including its stack trace.</param>
+        </member>
+        <member name="M:Common.Logging.ILog.WarnFormat(System.String,System.Object[])">
+            <summary>
+            Log a message with the <see cref="F:Common.Logging.LogLevel.Warn"/> level.
+            </summary>
+            <param name="format">The format of the message object to log.<see cref="M:System.String.Format(System.String,System.Object[])"/> </param>
+            <param name="args">the list of format arguments</param>
+        </member>
+        <member name="M:Common.Logging.ILog.WarnFormat(System.String,System.Exception,System.Object[])">
+            <summary>
+            Log a message with the <see cref="F:Common.Logging.LogLevel.Warn"/> level.
+            </summary>
+            <param name="format">The format of the message object to log.<see cref="M:System.String.Format(System.String,System.Object[])"/> </param>
+            <param name="exception">The exception to log.</param>
+            <param name="args">the list of format arguments</param>
+        </member>
+        <member name="M:Common.Logging.ILog.WarnFormat(System.IFormatProvider,System.String,System.Object[])">
+            <summary>
+            Log a message with the <see cref="F:Common.Logging.LogLevel.Warn"/> level.
+            </summary>
+            <param name="formatProvider">An <see cref="T:System.IFormatProvider"/> that supplies culture-specific formatting information.</param>
+            <param name="format">The format of the message object to log.<see cref="M:System.String.Format(System.String,System.Object[])"/> </param>
+            <param name="args"></param>
+        </member>
+        <member name="M:Common.Logging.ILog.WarnFormat(System.IFormatProvider,System.String,System.Exception,System.Object[])">
+            <summary>
+            Log a message with the <see cref="F:Common.Logging.LogLevel.Warn"/> level.
+            </summary>
+            <param name="formatProvider">An <see cref="T:System.IFormatProvider"/> that supplies culture-specific formatting information.</param>
+            <param name="format">The format of the message object to log.<see cref="M:System.String.Format(System.String,System.Object[])"/> </param>
+            <param name="exception">The exception to log.</param>
+            <param name="args"></param>
+        </member>
+        <member name="M:Common.Logging.ILog.Warn(System.Action{Common.Logging.FormatMessageHandler})">
+            <summary>
+            Log a message with the <see cref="F:Common.Logging.LogLevel.Warn"/> level using a callback to obtain the message
+            </summary>
+            <remarks>
+            Using this method avoids the cost of creating a message and evaluating message arguments 
+            that probably won't be logged due to loglevel settings.
+            </remarks>
+            <param name="formatMessageCallback">A callback used by the logger to obtain the message if log level is matched</param>
+        </member>
+        <member name="M:Common.Logging.ILog.Warn(System.Action{Common.Logging.FormatMessageHandler},System.Exception)">
+            <summary>
+            Log a message with the <see cref="F:Common.Logging.LogLevel.Warn"/> level using a callback to obtain the message
+            </summary>
+            <remarks>
+            Using this method avoids the cost of creating a message and evaluating message arguments 
+            that probably won't be logged due to loglevel settings.
+            </remarks>
+            <param name="formatMessageCallback">A callback used by the logger to obtain the message if log level is matched</param>
+            <param name="exception">The exception to log, including its stack trace.</param>
+        </member>
+        <member name="M:Common.Logging.ILog.Warn(System.IFormatProvider,System.Action{Common.Logging.FormatMessageHandler})">
+            <summary>
+            Log a message with the <see cref="F:Common.Logging.LogLevel.Warn"/> level using a callback to obtain the message
+            </summary>
+            <remarks>
+            Using this method avoids the cost of creating a message and evaluating message arguments 
+            that probably won't be logged due to loglevel settings.
+            </remarks>
+            <param name="formatProvider">An <see cref="T:System.IFormatProvider"/> that supplies culture-specific formatting information.</param>
+            <param name="formatMessageCallback">A callback used by the logger to obtain the message if log level is matched</param>
+        </member>
+        <member name="M:Common.Logging.ILog.Warn(System.IFormatProvider,System.Action{Common.Logging.FormatMessageHandler},System.Exception)">
+            <summary>
+            Log a message with the <see cref="F:Common.Logging.LogLevel.Warn"/> level using a callback to obtain the message
+            </summary>
+            <remarks>
+            Using this method avoids the cost of creating a message and evaluating message arguments 
+            that probably won't be logged due to loglevel settings.
+            </remarks>
+            <param name="formatProvider">An <see cref="T:System.IFormatProvider"/> that supplies culture-specific formatting information.</param>
+            <param name="formatMessageCallback">A callback used by the logger to obtain the message if log level is matched</param>
+            <param name="exception">The exception to log, including its stack Warn.</param>
+        </member>
+        <member name="M:Common.Logging.ILog.Error(System.Object)">
+            <summary>
+            Log a message object with the <see cref="F:Common.Logging.LogLevel.Error"/> level.
+            </summary>
+            <param name="message">The message object to log.</param>
+        </member>
+        <member name="M:Common.Logging.ILog.Error(System.Object,System.Exception)">
+            <summary>
+            Log a message object with the <see cref="F:Common.Logging.LogLevel.Error"/> level including
+            the stack trace of the <see cref="T:System.Exception"/> passed
+            as a parameter.
+            </summary>
+            <param name="message">The message object to log.</param>
+            <param name="exception">The exception to log, including its stack trace.</param>
+        </member>
+        <member name="M:Common.Logging.ILog.ErrorFormat(System.String,System.Object[])">
+            <summary>
+            Log a message with the <see cref="F:Common.Logging.LogLevel.Error"/> level.
+            </summary>
+            <param name="format">The format of the message object to log.<see cref="M:System.String.Format(System.String,System.Object[])"/> </param>
+            <param name="args">the list of format arguments</param>
+        </member>
+        <member name="M:Common.Logging.ILog.ErrorFormat(System.String,System.Exception,System.Object[])">
+            <summary>
+            Log a message with the <see cref="F:Common.Logging.LogLevel.Error"/> level.
+            </summary>
+            <param name="format">The format of the message object to log.<see cref="M:System.String.Format(System.String,System.Object[])"/> </param>
+            <param name="exception">The exception to log.</param>
+            <param name="args">the list of format arguments</param>
+        </member>
+        <member name="M:Common.Logging.ILog.ErrorFormat(System.IFormatProvider,System.String,System.Object[])">
+            <summary>
+            Log a message with the <see cref="F:Common.Logging.LogLevel.Error"/> level.
+            </summary>
+            <param name="formatProvider">An <see cref="T:System.IFormatProvider"/> that supplies culture-specific formatting information.</param>
+            <param name="format">The format of the message object to log.<see cref="M:System.String.Format(System.String,System.Object[])"/> </param>
+            <param name="args"></param>
+        </member>
+        <member name="M:Common.Logging.ILog.ErrorFormat(System.IFormatProvider,System.String,System.Exception,System.Object[])">
+            <summary>
+            Log a message with the <see cref="F:Common.Logging.LogLevel.Error"/> level.
+            </summary>
+            <param name="formatProvider">An <see cref="T:System.IFormatProvider"/> that supplies culture-specific formatting information.</param>
+            <param name="format">The format of the message object to log.<see cref="M:System.String.Format(System.String,System.Object[])"/> </param>
+            <param name="exception">The exception to log.</param>
+            <param name="args"></param>
+        </member>
+        <member name="M:Common.Logging.ILog.Error(System.Action{Common.Logging.FormatMessageHandler})">
+            <summary>
+            Log a message with the <see cref="F:Common.Logging.LogLevel.Error"/> level using a callback to obtain the message
+            </summary>
+            <remarks>
+            Using this method avoids the cost of creating a message and evaluating message arguments 
+            that probably won't be logged due to loglevel settings.
+            </remarks>
+            <param name="formatMessageCallback">A callback used by the logger to obtain the message if log level is matched</param>
+        </member>
+        <member name="M:Common.Logging.ILog.Error(System.Action{Common.Logging.FormatMessageHandler},System.Exception)">
+            <summary>
+            Log a message with the <see cref="F:Common.Logging.LogLevel.Error"/> level using a callback to obtain the message
+            </summary>
+            <remarks>
+            Using this method avoids the cost of creating a message and evaluating message arguments 
+            that probably won't be logged due to loglevel settings.
+            </remarks>
+            <param name="formatMessageCallback">A callback used by the logger to obtain the message if log level is matched</param>
+            <param name="exception">The exception to log, including its stack trace.</param>
+        </member>
+        <member name="M:Common.Logging.ILog.Error(System.IFormatProvider,System.Action{Common.Logging.FormatMessageHandler})">
+            <summary>
+            Log a message with the <see cref="F:Common.Logging.LogLevel.Error"/> level using a callback to obtain the message
+            </summary>
+            <remarks>
+            Using this method avoids the cost of creating a message and evaluating message arguments 
+            that probably won't be logged due to loglevel settings.
+            </remarks>
+            <param name="formatProvider">An <see cref="T:System.IFormatProvider"/> that supplies culture-specific formatting information.</param>
+            <param name="formatMessageCallback">A callback used by the logger to obtain the message if log level is matched</param>
+        </member>
+        <member name="M:Common.Logging.ILog.Error(System.IFormatProvider,System.Action{Common.Logging.FormatMessageHandler},System.Exception)">
+            <summary>
+            Log a message with the <see cref="F:Common.Logging.LogLevel.Error"/> level using a callback to obtain the message
+            </summary>
+            <remarks>
+            Using this method avoids the cost of creating a message and evaluating message arguments 
+            that probably won't be logged due to loglevel settings.
+            </remarks>
+            <param name="formatProvider">An <see cref="T:System.IFormatProvider"/> that supplies culture-specific formatting information.</param>
+            <param name="formatMessageCallback">A callback used by the logger to obtain the message if log level is matched</param>
+            <param name="exception">The exception to log, including its stack Error.</param>
+        </member>
+        <member name="M:Common.Logging.ILog.Fatal(System.Object)">
+            <summary>
+            Log a message object with the <see cref="F:Common.Logging.LogLevel.Fatal"/> level.
+            </summary>
+            <param name="message">The message object to log.</param>
+        </member>
+        <member name="M:Common.Logging.ILog.Fatal(System.Object,System.Exception)">
+            <summary>
+            Log a message object with the <see cref="F:Common.Logging.LogLevel.Fatal"/> level including
+            the stack trace of the <see cref="T:System.Exception"/> passed
+            as a parameter.
+            </summary>
+            <param name="message">The message object to log.</param>
+            <param name="exception">The exception to log, including its stack trace.</param>
+        </member>
+        <member name="M:Common.Logging.ILog.FatalFormat(System.String,System.Object[])">
+            <summary>
+            Log a message with the <see cref="F:Common.Logging.LogLevel.Fatal"/> level.
+            </summary>
+            <param name="format">The format of the message object to log.<see cref="M:System.String.Format(System.String,System.Object[])"/> </param>
+            <param name="args">the list of format arguments</param>
+        </member>
+        <member name="M:Common.Logging.ILog.FatalFormat(System.String,System.Exception,System.Object[])">
+            <summary>
+            Log a message with the <see cref="F:Common.Logging.LogLevel.Fatal"/> level.
+            </summary>
+            <param name="format">The format of the message object to log.<see cref="M:System.String.Format(System.String,System.Object[])"/> </param>
+            <param name="exception">The exception to log.</param>
+            <param name="args">the list of format arguments</param>
+        </member>
+        <member name="M:Common.Logging.ILog.FatalFormat(System.IFormatProvider,System.String,System.Object[])">
+            <summary>
+            Log a message with the <see cref="F:Common.Logging.LogLevel.Fatal"/> level.
+            </summary>
+            <param name="formatProvider">An <see cref="T:System.IFormatProvider"/> that supplies culture-specific formatting information.</param>
+            <param name="format">The format of the message object to log.<see cref="M:System.String.Format(System.String,System.Object[])"/> </param>
+            <param name="args"></param>
+        </member>
+        <member name="M:Common.Logging.ILog.FatalFormat(System.IFormatProvider,System.String,System.Exception,System.Object[])">
+            <summary>
+            Log a message with the <see cref="F:Common.Logging.LogLevel.Fatal"/> level.
+            </summary>
+            <param name="formatProvider">An <see cref="T:System.IFormatProvider"/> that supplies culture-specific formatting information.</param>
+            <param name="format">The format of the message object to log.<see cref="M:System.String.Format(System.String,System.Object[])"/> </param>
+            <param name="exception">The exception to log.</param>
+            <param name="args"></param>
+        </member>
+        <member name="M:Common.Logging.ILog.Fatal(System.Action{Common.Logging.FormatMessageHandler})">
+            <summary>
+            Log a message with the <see cref="F:Common.Logging.LogLevel.Fatal"/> level using a callback to obtain the message
+            </summary>
+            <remarks>
+            Using this method avoids the cost of creating a message and evaluating message arguments 
+            that probably won't be logged due to loglevel settings.
+            </remarks>
+            <param name="formatMessageCallback">A callback used by the logger to obtain the message if log level is matched</param>
+        </member>
+        <member name="M:Common.Logging.ILog.Fatal(System.Action{Common.Logging.FormatMessageHandler},System.Exception)">
+            <summary>
+            Log a message with the <see cref="F:Common.Logging.LogLevel.Fatal"/> level using a callback to obtain the message
+            </summary>
+            <remarks>
+            Using this method avoids the cost of creating a message and evaluating message arguments 
+            that probably won't be logged due to loglevel settings.
+            </remarks>
+            <param name="formatMessageCallback">A callback used by the logger to obtain the message if log level is matched</param>
+            <param name="exception">The exception to log, including its stack trace.</param>
+        </member>
+        <member name="M:Common.Logging.ILog.Fatal(System.IFormatProvider,System.Action{Common.Logging.FormatMessageHandler})">
+            <summary>
+            Log a message with the <see cref="F:Common.Logging.LogLevel.Fatal"/> level using a callback to obtain the message
+            </summary>
+            <remarks>
+            Using this method avoids the cost of creating a message and evaluating message arguments 
+            that probably won't be logged due to loglevel settings.
+            </remarks>
+            <param name="formatProvider">An <see cref="T:System.IFormatProvider"/> that supplies culture-specific formatting information.</param>
+            <param name="formatMessageCallback">A callback used by the logger to obtain the message if log level is matched</param>
+        </member>
+        <member name="M:Common.Logging.ILog.Fatal(System.IFormatProvider,System.Action{Common.Logging.FormatMessageHandler},System.Exception)">
+            <summary>
+            Log a message with the <see cref="F:Common.Logging.LogLevel.Fatal"/> level using a callback to obtain the message
+            </summary>
+            <remarks>
+            Using this method avoids the cost of creating a message and evaluating message arguments 
+            that probably won't be logged due to loglevel settings.
+            </remarks>
+            <param name="formatProvider">An <see cref="T:System.IFormatProvider"/> that supplies culture-specific formatting information.</param>
+            <param name="formatMessageCallback">A callback used by the logger to obtain the message if log level is matched</param>
+            <param name="exception">The exception to log, including its stack Fatal.</param>
+        </member>
+        <member name="P:Common.Logging.ILog.IsTraceEnabled">
+            <summary>
+            Checks if this logger is enabled for the <see cref="F:Common.Logging.LogLevel.Trace"/> level.
+            </summary>
+        </member>
+        <member name="P:Common.Logging.ILog.IsDebugEnabled">
+            <summary>
+            Checks if this logger is enabled for the <see cref="F:Common.Logging.LogLevel.Debug"/> level.
+            </summary>
+        </member>
+        <member name="P:Common.Logging.ILog.IsErrorEnabled">
+            <summary>
+            Checks if this logger is enabled for the <see cref="F:Common.Logging.LogLevel.Error"/> level.
+            </summary>
+        </member>
+        <member name="P:Common.Logging.ILog.IsFatalEnabled">
+            <summary>
+            Checks if this logger is enabled for the <see cref="F:Common.Logging.LogLevel.Fatal"/> level.
+            </summary>
+        </member>
+        <member name="P:Common.Logging.ILog.IsInfoEnabled">
+            <summary>
+            Checks if this logger is enabled for the <see cref="F:Common.Logging.LogLevel.Info"/> level.
+            </summary>
+        </member>
+        <member name="P:Common.Logging.ILog.IsWarnEnabled">
+            <summary>
+            Checks if this logger is enabled for the <see cref="F:Common.Logging.LogLevel.Warn"/> level.
+            </summary>
+        </member>
+        <member name="P:Common.Logging.ILog.GlobalVariablesContext">
+            <summary>
+            Returns the global context for variables
+            </summary>
+        </member>
+        <member name="P:Common.Logging.ILog.ThreadVariablesContext">
+            <summary>
+            Returns the thread-specific context for variables
+            </summary>
+        </member>
+        <member name="T:Common.Logging.ILoggerFactoryAdapter">
+            <summary>
+            LoggerFactoryAdapter interface is used internally by LogManager
+            Only developers wishing to write new Common.Logging adapters need to
+            worry about this interface.
+            </summary>
+            <author>Gilles Bayon</author>
+        </member>
+        <member name="M:Common.Logging.ILoggerFactoryAdapter.GetLogger(System.Type)">
+            <summary>
+            Get a ILog instance by type.
+            </summary>
+            <param name="type">The type to use for the logger</param>
+            <returns></returns>
+        </member>
+        <member name="M:Common.Logging.ILoggerFactoryAdapter.GetLogger(System.String)">
+            <summary>
+            Get a ILog instance by key.
+            </summary>
+            <param name="key">The key of the logger</param>
+            <returns></returns>
+        </member>
+        <member name="T:Common.Logging.ILogManager">
+            <summary>
+            Interface for LogManager
+            </summary>
+        </member>
+        <member name="M:Common.Logging.ILogManager.Reset">
+            <summary>
+            Reset the <see cref="N:Common.Logging"/> infrastructure to its default settings. This means, that configuration settings
+            will be re-read from section <c>&lt;common/logging&gt;</c> of your <c>app.config</c>.
+            </summary>
+            <remarks>
+            This is mainly used for unit testing, you wouldn't normally use this in your applications.<br/>
+            <b>Note:</b><see cref="T:Common.Logging.ILog"/> instances already handed out from this LogManager are not(!) affected. 
+            Resetting LogManager only affects new instances being handed out.
+            </remarks>
+        </member>
+        <member name="M:Common.Logging.ILogManager.Reset(Common.Logging.IConfigurationReader)">
+            <summary>
+            Reset the <see cref="N:Common.Logging"/> infrastructure to its default settings. This means, that configuration settings
+            will be re-read from section <c>&lt;common/logging&gt;</c> of your <c>app.config</c>.
+            </summary>
+            <remarks>
+            This is mainly used for unit testing, you wouldn't normally use this in your applications.<br/>
+            <b>Note:</b><see cref="T:Common.Logging.ILog"/> instances already handed out from this LogManager are not(!) affected. 
+            Resetting LogManager only affects new instances being handed out.
+            </remarks>
+            <param name="reader">
+            the <see cref="T:Common.Logging.IConfigurationReader"/> instance to obtain settings for 
+            re-initializing the LogManager.
+            </param>
+        </member>
+        <member name="M:Common.Logging.ILogManager.GetCurrentClassLogger">
+            <summary>
+            Gets the logger by calling <see cref="M:Common.Logging.ILoggerFactoryAdapter.GetLogger(System.Type)"/>
+            on the currently configured <see cref="P:Common.Logging.ILogManager.Adapter"/> using the type of the calling class.
+            </summary>
+            <remarks>
+            This method needs to inspect the StackTrace in order to determine the calling 
+            class. This of course comes with a performance penalty, thus you shouldn't call it too
+            often in your application.
+            </remarks>
+            <seealso cref="M:Common.Logging.ILogManager.GetLogger(System.Type)"/>
+            <returns>the logger instance obtained from the current <see cref="P:Common.Logging.ILogManager.Adapter"/></returns>
+        </member>
+        <member name="M:Common.Logging.ILogManager.GetLogger``1">
+            <summary>
+            Gets the logger by calling <see cref="M:Common.Logging.ILoggerFactoryAdapter.GetLogger(System.Type)"/>
+            on the currently configured <see cref="P:Common.Logging.ILogManager.Adapter"/> using the specified type.
+            </summary>
+            <returns>the logger instance obtained from the current <see cref="P:Common.Logging.ILogManager.Adapter"/></returns>
+        </member>
+        <member name="M:Common.Logging.ILogManager.GetLogger(System.Type)">
+            <summary>
+            Gets the logger by calling <see cref="M:Common.Logging.ILoggerFactoryAdapter.GetLogger(System.Type)"/>
+            on the currently configured <see cref="P:Common.Logging.ILogManager.Adapter"/> using the specified type.
+            </summary>
+            <param name="type">The type.</param>
+            <returns>the logger instance obtained from the current <see cref="P:Common.Logging.ILogManager.Adapter"/></returns>
+        </member>
+        <member name="M:Common.Logging.ILogManager.GetLogger(System.String)">
+            <summary>
+            Gets the logger by calling <see cref="M:Common.Logging.ILoggerFactoryAdapter.GetLogger(System.String)"/>
+            on the currently configured <see cref="P:Common.Logging.ILogManager.Adapter"/> using the specified key.
+            </summary>
+            <param name="key">The key.</param>
+            <returns>the logger instance obtained from the current <see cref="P:Common.Logging.ILogManager.Adapter"/></returns>
+        </member>
+        <member name="P:Common.Logging.ILogManager.COMMON_LOGGING_SECTION">
+            <summary>
+            The key of the default configuration section to read settings from.
+            </summary>
+            <remarks>
+            You can always change the source of your configuration settings by setting another <see cref="T:Common.Logging.IConfigurationReader"/> instance
+            on <see cref="P:Common.Logging.ILogManager.ConfigurationReader"/>.
+            </remarks>
+        </member>
+        <member name="P:Common.Logging.ILogManager.ConfigurationReader">
+            <summary>
+            Gets the configuration reader used to initialize the LogManager.
+            </summary>
+            <remarks>Primarily used for testing purposes but maybe useful to obtain configuration
+            information from some place other than the .NET application configuration file.</remarks>
+            <value>The configuration reader.</value>
+        </member>
+        <member name="P:Common.Logging.ILogManager.Adapter">
+            <summary>
+            Gets or sets the adapter.
+            </summary>
+            <value>The adapter.</value>
+        </member>
+        <member name="T:Common.Logging.IVariablesContext">
+            <summary>
+            A context for logger variables
+            </summary>
+        </member>
+        <member name="M:Common.Logging.IVariablesContext.Set(System.String,System.Object)">
+            <summary>
+            Sets the value of a new or existing variable within the global context
+            </summary>
+            <param name="key">The key of the variable that is to be added</param>
+            <param name="value">The value to add</param>
+        </member>
+        <member name="M:Common.Logging.IVariablesContext.Get(System.String)">
+            <summary>
+            Gets the value of a variable within the global context
+            </summary>
+            <param name="key">The key of the variable to get</param>
+            <returns>The value or null if not found</returns>
+        </member>
+        <member name="M:Common.Logging.IVariablesContext.Contains(System.String)">
+            <summary>
+            Checks if a variable is set within the global context
+            </summary>
+            <param name="key">The key of the variable to check for</param>
+            <returns>True if the variable is set</returns>
+        </member>
+        <member name="M:Common.Logging.IVariablesContext.Remove(System.String)">
+            <summary>
+            Removes a variable from the global context by key
+            </summary>
+            <param name="key">The key of the variable to remove</param>
+        </member>
+        <member name="M:Common.Logging.IVariablesContext.Clear">
+            <summary>
+            Clears the global context variables
+            </summary>
+        </member>
+        <member name="T:Common.Logging.LogLevel">
+            <summary>
+            The 7 possible logging levels
+            </summary>
+            <author>Gilles Bayon</author>
+        </member>
+        <member name="F:Common.Logging.LogLevel.All">
+            <summary>
+            All logging levels
+            </summary>
+        </member>
+        <member name="F:Common.Logging.LogLevel.Trace">
+            <summary>
+            A trace logging level
+            </summary>
+        </member>
+        <member name="F:Common.Logging.LogLevel.Debug">
+            <summary>
+            A debug logging level
+            </summary>
+        </member>
+        <member name="F:Common.Logging.LogLevel.Info">
+            <summary>
+            A info logging level
+            </summary>
+        </member>
+        <member name="F:Common.Logging.LogLevel.Warn">
+            <summary>
+            A warn logging level
+            </summary>
+        </member>
+        <member name="F:Common.Logging.LogLevel.Error">
+            <summary>
+            An error logging level
+            </summary>
+        </member>
+        <member name="F:Common.Logging.LogLevel.Fatal">
+            <summary>
+            A fatal logging level
+            </summary>
+        </member>
+        <member name="F:Common.Logging.LogLevel.Off">
+            <summary>
+            Do not log anything.
+            </summary>
+        </member>
+    </members>
+</doc>
