Copying //UE4/Dev-Core to //UE4/Main

diff --git a/Engine/Source/Programs/CrashReporter/CrashReportWebSite/packages/HangFire.Core.1.5.3/lib/net45/Hangfire.Core.xml b/Engine/Source/Programs/CrashReporter/CrashReportWebSite/packages/HangFire.Core.1.5.3/lib/net45/Hangfire.Core.xml
new file mode 100644
--- /dev/null
+++ b/Engine/Source/Programs/CrashReporter/CrashReportWebSite/packages/HangFire.Core.1.5.3/lib/net45/Hangfire.Core.xml
@@ -0,0 +1,4012 @@
+<?xml version="1.0"?>
+<doc>
+    <assembly>
+        <name>Hangfire.Core</name>
+    </assembly>
+    <members>
+        <member name="T:Hangfire.AppBuilderExtensions">
+            <summary>
+            Provides extension methods for the <c>IAppBuilder</c> interface
+            defined in the <see href="https://www.nuget.org/packages/Owin/">Owin</see> 
+            NuGet package to simplify the integration with OWIN applications.
+            </summary>
+            
+            <remarks>
+            
+            <para>This class simplifies Hangfire configuration in OWIN applications,
+            please read <see href="http://www.asp.net/aspnet/overview/owin-and-katana/getting-started-with-owin-and-katana">
+            Getting Started with OWIN and Katana</see> if you aren't familiar with OWIN
+            and/or don't know what is the <c>Startup</c> class.
+            </para>
+            
+            <para>The methods of this class should be called from OWIN's <c>Startup</c> 
+            class.</para>
+            
+            <h3>UseHangfireDashboard</h3>
+            <para>Dashboard UI contains pages that allow you to monitor almost every
+            aspect of background processing. It is exposed as an OWIN middleware that 
+            intercepts requests to the given path.</para>
+            <para>OWIN implementation of Dashboard UI allows to use it outside of web
+            applications, including console applications and Windows Services.</para>
+            <note type="important">
+            By default, an access to the Dashboard UI is restricted <b>only to local
+            requests</b> for security reasons. Before publishing a project to
+            production, make sure you still have access to the Dashboard UI by using the
+            <see href="https://www.nuget.org/packages/Hangfire.Dashboard.Authorization/">
+            Hangfire.Dashboard.Authorization</see> package.</note>
+            
+            <h3>UseHangfireServer</h3>
+            <para>In addition to creation of a new instance of the <see cref="T:Hangfire.BackgroundJobServer"/> 
+            class, these methods also register the call to its <see cref="M:Hangfire.BackgroundJobServer.Dispose"/> 
+            method on application shutdown. This is done via registering a callback on the corresponding 
+            <see cref="T:System.Threading.CancellationToken"/> from OWIN environment (<c>"host.OnAppDisposing"</c> or 
+            <c>"server.OnDispose"</c> keys).</para>
+            <para>This enables <i>graceful shutdown</i> feature for background jobs and background processes
+            without any additional configuration.</para>
+            <para>Please see <see cref="T:Hangfire.BackgroundJobServer"/> for more details regarding
+            background processing.</para>
+            </remarks>
+            
+            <example>
+            <h3>Basic Configuration</h3> 
+            <para>Basic setup in an OWIN application looks like the following example. Please note
+            that job storage should be configured before using the methods of this class.</para>
+            
+            <code lang="cs" source="..\Samples\AppBuilderExtensions.cs" region="Basic Setup"/>
+            
+            <h3>Adding Dashboard Only</h3>
+            <para>If you want to install dashboard without starting a background job server, for example,
+            to process background jobs outside of your web application, call only the
+            <see cref="O:Hangfire.AppBuilderExtensions.UseHangfireDashboard"/>.</para>
+            
+            <code lang="cs" source="..\Samples\AppBuilderExtensions.cs" region="Dashboard Only"/>
+            
+            <h3>Change Dashboard Path</h3>
+            <para>By default, you can access Dashboard UI by hitting the <i>http(s)://&lt;app&gt;/hangfire</i>
+            URL, however you can change it as in the following example.</para>
+            
+            <code lang="cs" source="..\Samples\AppBuilderExtensions.cs" region="Change Dashboard Path"/>
+            
+            <h3>Configuring Authorization</h3>
+            <para>The following example demonstrates how to change default local-requests-only
+            authorization for Dashboard UI.</para>
+            
+            <code lang="cs" source="..\Samples\AppBuilderExtensions.cs" region="Configuring Authorization"/>
+            
+            <h3>Changing Application Path</h3>
+            <para>Have you seen the <i>Back to site</i> button in the Dashboard? By default it leads
+            you to the root of your site, but you can configure the behavior.</para>
+            
+            <code lang="cs" source="..\Samples\AppBuilderExtensions.cs" region="Change Application Path"/>
+            
+            <h3>Multiple Dashboards</h3>
+            <para>The following example demonstrates adding multiple Dashboard UI endpoints. This may
+            be useful when you are using multiple shards for your background processing needs.</para>
+            
+            <code lang="cs" source="..\Samples\AppBuilderExtensions.cs" region="Multiple Dashboards"/>
+            
+            </example>
+            
+            <seealso cref="T:Hangfire.BackgroundJobServer"/>
+            <seealso cref="N:Hangfire.Dashboard"/>
+            <seealso href="https://www.nuget.org/packages/Hangfire.Dashboard.Authorization/">
+            Hangfire.Dashboard.Authorization Package
+            </seealso>
+            
+            <threadsafety static="true" instance="false"/>
+        </member>
+        <member name="M:Hangfire.AppBuilderExtensions.UseHangfireServer(Owin.IAppBuilder)">
+            <summary>
+            Creates a new instance of the <see cref="T:Hangfire.BackgroundJobServer"/> class
+            with default options and <see cref="P:Hangfire.JobStorage.Current"/> storage and
+            registers its disposal on application shutdown.
+            </summary>
+            <param name="builder">OWIN application builder.</param>
+            
+            <exception cref="T:System.ArgumentNullException"><paramref name="builder"/> is null.</exception>
+            <exception cref="T:System.InvalidOperationException">
+            OWIN environment does not contain the application shutdown cancellation token.
+            </exception>
+            
+            <remarks>
+            Please see <see cref="T:Hangfire.AppBuilderExtensions"/> for details and examples.
+            </remarks>
+        </member>
+        <member name="M:Hangfire.AppBuilderExtensions.UseHangfireServer(Owin.IAppBuilder,Hangfire.Server.IBackgroundProcess[])">
+            <summary>
+            Creates a new instance of the <see cref="T:Hangfire.BackgroundJobServer"/> class 
+            with the given collection of additional background processes and 
+            <see cref="P:Hangfire.JobStorage.Current"/> storage, and registers its disposal
+            on application shutdown.
+            </summary>
+            <param name="builder">OWIN application builder.</param>
+            <param name="additionalProcesses">Collection of additional background processes.</param>
+            
+            <exception cref="T:System.ArgumentNullException"><paramref name="builder"/> is null.</exception>
+            <exception cref="T:System.ArgumentNullException"><paramref name="additionalProcesses"/> is null.</exception>
+            <exception cref="T:System.InvalidOperationException">
+            OWIN environment does not contain the application shutdown cancellation token.
+            </exception>
+            
+            <remarks>
+            Please see <see cref="T:Hangfire.AppBuilderExtensions"/> for details and examples.
+            </remarks>
+        </member>
+        <member name="M:Hangfire.AppBuilderExtensions.UseHangfireServer(Owin.IAppBuilder,Hangfire.BackgroundJobServerOptions)">
+            <summary>
+            Creates a new instance of the <see cref="T:Hangfire.BackgroundJobServer"/> class
+            with the specified options and <see cref="P:Hangfire.JobStorage.Current"/> storage,
+            and registers its disposal on application shutdown.
+            </summary>
+            <param name="builder">OWIN application builder.</param>
+            <param name="options">Options for background job server.</param>
+            
+            <exception cref="T:System.ArgumentNullException"><paramref name="builder"/> is null.</exception>
+            <exception cref="T:System.ArgumentNullException"><paramref name="options"/> is null.</exception>
+            <exception cref="T:System.InvalidOperationException">
+            OWIN environment does not contain the application shutdown cancellation token.
+            </exception>
+            
+            <remarks>
+            Please see <see cref="T:Hangfire.AppBuilderExtensions"/> for details and examples.
+            </remarks>
+        </member>
+        <member name="M:Hangfire.AppBuilderExtensions.UseHangfireServer(Owin.IAppBuilder,Hangfire.BackgroundJobServerOptions,Hangfire.Server.IBackgroundProcess[])">
+            <summary>
+            Creates a new instance of the <see cref="T:Hangfire.BackgroundJobServer"/> class
+            with the specified options, given collection of background processes
+            and <see cref="P:Hangfire.JobStorage.Current"/> storage, and registers its
+            disposal on application shutdown.
+            </summary>
+            <param name="builder">OWIN application builder.</param>
+            <param name="options">Options for background job server.</param>
+            <param name="additionalProcesses">Collection of additional background processes.</param>
+            
+            <exception cref="T:System.ArgumentNullException"><paramref name="builder"/> is null.</exception>
+            <exception cref="T:System.ArgumentNullException"><paramref name="options"/> is null.</exception>
+            <exception cref="T:System.ArgumentNullException"><paramref name="additionalProcesses"/> is null.</exception>
+            <exception cref="T:System.InvalidOperationException">
+            OWIN environment does not contain the application shutdown cancellation token.
+            </exception>
+            
+            <remarks>
+            Please see <see cref="T:Hangfire.AppBuilderExtensions"/> for details and examples.
+            </remarks>
+        </member>
+        <member name="M:Hangfire.AppBuilderExtensions.UseHangfireServer(Owin.IAppBuilder,Hangfire.BackgroundJobServerOptions,Hangfire.JobStorage)">
+            <summary>
+            Creates a new instance of the <see cref="T:Hangfire.BackgroundJobServer"/> class
+            with the given options and specified storage, and registers its disposal
+            on application shutdown.
+            </summary>
+            <param name="builder">OWIN application builder.</param>
+            <param name="options">Options for background job server.</param>
+            <param name="storage">Storage to use by background job server.</param>
+            
+            <exception cref="T:System.ArgumentNullException"><paramref name="builder"/> is null.</exception>
+            <exception cref="T:System.ArgumentNullException"><paramref name="storage"/> is null.</exception>
+            <exception cref="T:System.ArgumentNullException"><paramref name="options"/> is null.</exception>
+            <exception cref="T:System.InvalidOperationException">
+            OWIN environment does not contain the application shutdown cancellation token.
+            </exception>
+            
+            <remarks>
+            Please see <see cref="T:Hangfire.AppBuilderExtensions"/> for details and examples.
+            </remarks>
+        </member>
+        <member name="M:Hangfire.AppBuilderExtensions.UseHangfireServer(Owin.IAppBuilder,Hangfire.JobStorage,Hangfire.BackgroundJobServerOptions,Hangfire.Server.IBackgroundProcess[])">
+            <summary>
+            Starts a new instance of the <see cref="T:Hangfire.BackgroundJobServer"/> class with
+            the given arguments, and registers its disposal on application shutdown.
+            </summary>
+            
+            <param name="builder">OWIN application builder.</param>
+            <param name="storage">Storage to use by background job server.</param>
+            <param name="options">Options for background job server.</param>
+            <param name="additionalProcesses">Collection of additional background processes.</param>
+            
+            <exception cref="T:System.ArgumentNullException"><paramref name="builder"/> is null.</exception>
+            <exception cref="T:System.ArgumentNullException"><paramref name="storage"/> is null.</exception>
+            <exception cref="T:System.ArgumentNullException"><paramref name="options"/> is null.</exception>
+            <exception cref="T:System.ArgumentNullException"><paramref name="additionalProcesses"/> is null.</exception>
+            <exception cref="T:System.InvalidOperationException">
+            OWIN environment does not contain the application shutdown cancellation token.
+            </exception>
+            
+            <remarks>
+            Please see <see cref="T:Hangfire.AppBuilderExtensions"/> for details and examples.
+            </remarks>
+        </member>
+        <member name="M:Hangfire.AppBuilderExtensions.UseHangfireDashboard(Owin.IAppBuilder)">
+            <summary>
+            Adds Dashboard UI middleware to the OWIN request processing pipeline under 
+            the <c>/hangfire</c> path, for the <see cref="P:Hangfire.JobStorage.Current"/> storage.
+            </summary>
+            <param name="builder">OWIN application builder.</param>
+            
+            <exception cref="T:System.ArgumentNullException"><paramref name="builder"/> is null.</exception>
+            
+            <remarks>
+            Please see <see cref="T:Hangfire.AppBuilderExtensions"/> for details and examples.
+            </remarks>
+        </member>
+        <member name="M:Hangfire.AppBuilderExtensions.UseHangfireDashboard(Owin.IAppBuilder,System.String)">
+            <summary>
+            Adds Dashboard UI middleware to the OWIN request processing pipeline under
+            the given path, for the <see cref="P:Hangfire.JobStorage.Current"/> storage.
+            </summary>
+            <param name="builder">OWIN application builder.</param>
+            <param name="pathMatch">Path prefix for middleware to use, e.g. "/hangfire".</param>
+            
+            <exception cref="T:System.ArgumentNullException"><paramref name="builder"/> is null.</exception>
+            <exception cref="T:System.ArgumentNullException"><paramref name="pathMatch"/> is null.</exception>
+            
+            <remarks>
+            Please see <see cref="T:Hangfire.AppBuilderExtensions"/> for details and examples.
+            </remarks>
+        </member>
+        <member name="M:Hangfire.AppBuilderExtensions.UseHangfireDashboard(Owin.IAppBuilder,System.String,Hangfire.DashboardOptions)">
+            <summary>
+            Adds Dashboard UI middleware to the OWIN request processing pipeline under
+            the specified path and the given options, for the <see cref="P:Hangfire.JobStorage.Current"/>
+            storage.
+            </summary>
+            <param name="builder">OWIN application builder.</param>
+            <param name="pathMatch">Path prefix for middleware to use, e.g. "/hangfire".</param>
+            <param name="options">Options for Dashboard UI.</param>
+            
+            <exception cref="T:System.ArgumentNullException"><paramref name="builder"/> is null.</exception>
+            <exception cref="T:System.ArgumentNullException"><paramref name="pathMatch"/> is null.</exception>
+            <exception cref="T:System.ArgumentNullException"><paramref name="options"/> is null.</exception>
+            
+            <remarks>
+            Please see <see cref="T:Hangfire.AppBuilderExtensions"/> for details and examples.
+            </remarks>
+        </member>
+        <member name="M:Hangfire.AppBuilderExtensions.UseHangfireDashboard(Owin.IAppBuilder,System.String,Hangfire.DashboardOptions,Hangfire.JobStorage)">
+            <summary>
+            Adds Dashboard UI middleware to the OWIN request processing pipeline with the
+            specified parameters.
+            </summary>
+            <param name="builder">OWIN application builder.</param>
+            <param name="pathMatch">Path prefix for middleware to use, e.g. "/hangfire".</param>
+            <param name="options">Options for Dashboard UI.</param>
+            <param name="storage">Job storage to use by Dashboard IO.</param>
+            
+            <exception cref="T:System.ArgumentNullException"><paramref name="builder"/> is null.</exception>
+            <exception cref="T:System.ArgumentNullException"><paramref name="pathMatch"/> is null.</exception>
+            <exception cref="T:System.ArgumentNullException"><paramref name="options"/> is null.</exception>
+            <exception cref="T:System.ArgumentNullException"><paramref name="storage"/> is null.</exception>
+            
+            <remarks>
+            Please see <see cref="T:Hangfire.AppBuilderExtensions"/> for details and examples.
+            </remarks>
+        </member>
+        <member name="T:Hangfire.Logging.ILog">
+            <summary>
+            Simple interface that represent a logger.
+            </summary>
+        </member>
+        <member name="M:Hangfire.Logging.ILog.Log(Hangfire.Logging.LogLevel,System.Func{System.String},System.Exception)">
+            <summary>
+            Log a message the specified log level.
+            </summary>
+            <param name="logLevel">The log level.</param>
+            <param name="messageFunc">The message function.</param>
+            <param name="exception">An optional exception.</param>
+            <returns>true if the message was logged. Otherwise false.</returns>
+            <remarks>
+            Note to implementers: the message func should not be called if the loglevel is not enabled
+            so as not to incur performance penalties.
+            
+            To check IsEnabled call Log with only LogLevel and check the return value, no event will be written
+            </remarks>
+        </member>
+        <member name="T:Hangfire.Logging.LogLevel">
+            <summary>
+            The log level.
+            </summary>
+        </member>
+        <member name="T:Hangfire.Logging.ILogProvider">
+            <summary>
+            Represents a way to get a <see cref="T:Hangfire.Logging.ILog"/>
+            </summary>
+        </member>
+        <member name="T:Hangfire.Logging.LogProvider">
+            <summary>
+            Provides a mechanism to create instances of <see cref="T:Hangfire.Logging.ILog"/> objects.
+            </summary>
+        </member>
+        <member name="M:Hangfire.Logging.LogProvider.For``1">
+            <summary>
+            Gets a logger for the specified type.
+            </summary>
+            <typeparam name="T">The type whose name will be used for the logger.</typeparam>
+            <returns>An instance of <see cref="T:Hangfire.Logging.ILog"/></returns>
+        </member>
+        <member name="M:Hangfire.Logging.LogProvider.GetCurrentClassLogger">
+            <summary>
+            Gets a logger for the current class.
+            </summary>
+            <returns>An instance of <see cref="T:Hangfire.Logging.ILog"/></returns>
+        </member>
+        <member name="M:Hangfire.Logging.LogProvider.GetLogger(System.Type)">
+            <summary>
+            Gets a logger for the specified type.
+            </summary>
+            <param name="type">The type whose name will be used for the logger.</param>
+            <returns>An instance of <see cref="T:Hangfire.Logging.ILog"/></returns>
+        </member>
+        <member name="M:Hangfire.Logging.LogProvider.GetLogger(System.String)">
+            <summary>
+            Gets a logger with the specified name.
+            </summary>
+            <param name="name">The name.</param>
+            <returns>An instance of <see cref="T:Hangfire.Logging.ILog"/></returns>
+        </member>
+        <member name="M:Hangfire.Logging.LogProvider.SetCurrentLogProvider(Hangfire.Logging.ILogProvider)">
+            <summary>
+            Sets the current log provider.
+            </summary>
+            <param name="logProvider">The log provider.</param>
+        </member>
+        <member name="P:Hangfire.Logging.LogProviders.LoupeLogProvider.ProviderIsAvailableOverride">
+            <summary>
+            Gets or sets a value indicating whether [provider is available override]. Used in tests.
+            </summary>
+            <value>
+            <c>true</c> if [provider is available override]; otherwise, <c>false</c>.
+            </value>
+        </member>
+        <member name="T:Hangfire.Logging.LogProviders.LoupeLogProvider.WriteDelegate">
+            <summary>
+            The form of the Loupe Log.Write method we're using
+            </summary>
+        </member>
+        <member name="T:Hangfire.Logging.LogProviders.ColouredConsoleLogProvider.MessageFormatterDelegate">
+            <summary>
+            A delegate returning a formatted log message
+            </summary>
+            <param name="loggerName">The name of the Logger</param>
+            <param name="level">The Log Level</param>
+            <param name="message">The Log Message</param>
+            <param name="e">The Exception, if there is one</param>
+            <returns>A formatted Log Message string.</returns>
+        </member>
+        <member name="T:Hangfire.AttemptsExceededAction">
+            <summary>
+            Specifies a candidate state for a background job that will be chosen
+            by the <see cref="T:Hangfire.AutomaticRetryAttribute"/> filter after exceeding
+            the number of retry attempts.
+            </summary>
+        </member>
+        <member name="F:Hangfire.AttemptsExceededAction.Fail">
+            <summary>
+            Background job will be moved to the <see cref="T:Hangfire.States.FailedState"/>.
+            </summary>
+        </member>
+        <member name="F:Hangfire.AttemptsExceededAction.Delete">
+            <summary>
+            Background job will be moved to the <see cref="T:Hangfire.States.DeletedState"/>.
+            </summary>
+        </member>
+        <member name="T:Hangfire.AutomaticRetryAttribute">
+            <summary>
+            Represents a job filter that performs <i>automatic retries</i> for 
+            background jobs whose processing was failed due to an exception, with 
+            a limited number of attempts.
+            </summary>
+            
+            <remarks>
+            <para>Filter is added to the global <see cref="P:Hangfire.GlobalJobFilters.Filters"/> 
+            collection by default. Intervals between attempts are based on increasing 
+            exponential back-off multiplier in seconds.</para>
+            
+            <para>This filter works in a <i>state election</i> phase by changing the
+            candidate state from <see cref="T:Hangfire.States.FailedState"/> to the <see cref="T:Hangfire.States.ScheduledState"/>
+            when another retry should be attempted, or other state based on the value
+            of the <see cref="P:Hangfire.AutomaticRetryAttribute.OnAttemptsExceeded"/> property when attempts exceeded.
+            </para>
+            </remarks>
+            
+            <example>
+            <h3>Disabling Automatic Retries</h3>
+            <para>The following example shows how to disable automatic retries for
+            a specific job method by applying an attribute to a method.</para>
+            
+            <note>Even if you disable <see cref="T:Hangfire.AutomaticRetryAttribute"/> filter, 
+            your background jobs can still be executed several times, due to re-queue 
+            on shutdown and other compensation logic that guarantees the <i>at least
+            once</i> processing.</note>
+            
+            <code lang="cs" source="..\Samples\AutomaticRetry.cs" region="Disable Retries"/>
+            
+            <h3>Overriding Defaults</h3>
+            <para>The following example shows how to override the default number of
+            retry attempts for all of the background jobs by modifying the global
+            <see cref="P:Hangfire.GlobalJobFilters.Filters"/> collection.</para>
+            
+            <code lang="cs" source="..\Samples\AutomaticRetry.cs" region="Override Default"/>
+            
+            <h3>Specifying Attempts Exceeded Action</h3>
+            <para>The following example shows how to ignore a background job when
+            number of retry attempts exceed using the <see cref="P:Hangfire.AutomaticRetryAttribute.OnAttemptsExceeded"/>
+            property.</para>
+            
+            <note type="tip">Choose <see cref="F:Hangfire.AttemptsExceededAction.Delete"/> action 
+            when you aren't interested in processing background job that failed several
+            times.</note>
+            
+            <code lang="cs" source="..\Samples\AutomaticRetry.cs" region="Attempts Exceeded"/>
+            </example>
+            
+            <threadsafety static="true" instance="true"/>
+        </member>
+        <member name="T:Hangfire.Common.JobFilterAttribute">
+            <summary>
+            Represents the base class for job filter attributes.
+            </summary>
+        </member>
+        <member name="T:Hangfire.Common.IJobFilter">
+            <summary>
+            Defines members that specify the order of filters and 
+            whether multiple filters are allowed.
+            </summary>
+        </member>
+        <member name="P:Hangfire.Common.IJobFilter.AllowMultiple">
+            <summary>
+            When implemented in a class, gets or sets a value 
+            that indicates whether multiple filters are allowed.
+            </summary>
+        </member>
+        <member name="P:Hangfire.Common.IJobFilter.Order">
+            <summary>
+            When implemented in a class, gets the filter order.
+            </summary>
+        </member>
+        <member name="T:Hangfire.States.IElectStateFilter">
+            <summary>
+            Defines methods that are required for a state changing filter.
+            </summary>
+        </member>
+        <member name="M:Hangfire.States.IElectStateFilter.OnStateElection(Hangfire.States.ElectStateContext)">
+            <summary>
+            Called when the current state of the job is being changed to the
+            specified candidate state.
+            This state change could be intercepted and the final state could
+            be changed through setting the different state in the context 
+            in an implementation of this method.
+            </summary>
+        </member>
+        <member name="T:Hangfire.States.IApplyStateFilter">
+            <summary>
+            Provides methods that are required for a state changed filter.
+            </summary>
+        </member>
+        <member name="M:Hangfire.States.IApplyStateFilter.OnStateApplied(Hangfire.States.ApplyStateContext,Hangfire.Storage.IWriteOnlyTransaction)">
+            <summary>
+            Called after the specified state was applied
+            to the job within the given transaction.
+            </summary>
+        </member>
+        <member name="M:Hangfire.States.IApplyStateFilter.OnStateUnapplied(Hangfire.States.ApplyStateContext,Hangfire.Storage.IWriteOnlyTransaction)">
+            <summary>
+            Called when the state with specified state was 
+            unapplied from the job within the given transaction.
+            </summary>
+        </member>
+        <member name="F:Hangfire.AutomaticRetryAttribute.DefaultRetryAttempts">
+            <summary>
+            Represents the default number of retry attempts. This field is read-only.
+            </summary>
+            <remarks>
+            The value of this field is <c>10</c>.
+            </remarks>
+        </member>
+        <member name="M:Hangfire.AutomaticRetryAttribute.#ctor">
+            <summary>
+            Initializes a new instance of the <see cref="T:Hangfire.AutomaticRetryAttribute"/>
+            class with <see cref="F:Hangfire.AutomaticRetryAttribute.DefaultRetryAttempts"/> number.
+            </summary>
+        </member>
+        <member name="M:Hangfire.AutomaticRetryAttribute.OnStateElection(Hangfire.States.ElectStateContext)">
+            <inheritdoc />
+        </member>
+        <member name="M:Hangfire.AutomaticRetryAttribute.OnStateApplied(Hangfire.States.ApplyStateContext,Hangfire.Storage.IWriteOnlyTransaction)">
+            <inheritdoc />
+        </member>
+        <member name="M:Hangfire.AutomaticRetryAttribute.OnStateUnapplied(Hangfire.States.ApplyStateContext,Hangfire.Storage.IWriteOnlyTransaction)">
+            <inheritdoc />
+        </member>
+        <member name="M:Hangfire.AutomaticRetryAttribute.ScheduleAgainLater(Hangfire.States.ElectStateContext,System.Int32,Hangfire.States.FailedState)">
+            <summary>
+            Schedules the job to run again later. See <see cref="M:Hangfire.AutomaticRetryAttribute.SecondsToDelay(System.Int64)"/>.
+            </summary>
+            <param name="context">The state context.</param>
+            <param name="retryAttempt">The count of retry attempts made so far.</param>
+            <param name="failedState">Object which contains details about the current failed state.</param>
+        </member>
+        <member name="M:Hangfire.AutomaticRetryAttribute.TransitionToDeleted(Hangfire.States.ElectStateContext,Hangfire.States.FailedState)">
+            <summary>
+            Transition the candidate state to the deleted state.
+            </summary>
+            <param name="context">The state context.</param>
+            <param name="failedState">Object which contains details about the current failed state.</param>
+        </member>
+        <member name="P:Hangfire.AutomaticRetryAttribute.Attempts">
+            <summary>
+            Gets or sets the maximum number of automatic retry attempts.
+            </summary>
+            <value>Any non-negative number.</value>
+            <exception cref="T:System.ArgumentOutOfRangeException">The value in a set operation is less than zero.</exception>
+        </member>
+        <member name="P:Hangfire.AutomaticRetryAttribute.OnAttemptsExceeded">
+            <summary>
+            Gets or sets a candidate state for a background job that 
+            will be chosen when number of retry attempts exceeded.
+            </summary>
+        </member>
+        <member name="P:Hangfire.AutomaticRetryAttribute.LogEvents">
+            <summary>
+            Gets or sets whether to produce log messages on retry attempts.
+            </summary>
+        </member>
+        <member name="T:Hangfire.BackgroundJobClientException">
+            <summary>
+            The exception that is thrown when an instance of the class that 
+            implements the <see cref="T:Hangfire.IBackgroundJobClient"/> interface is unable
+            to perform an operation due to an error.
+            </summary>
+        </member>
+        <member name="T:Hangfire.Client.CreateJobFailedException">
+            <summary>
+            The exception that is thrown when a <see cref="T:Hangfire.BackgroundJobClient"/> class instance
+            could not create a job due to another exception was thrown.
+            </summary>
+        </member>
+        <member name="M:Hangfire.Client.CreateJobFailedException.#ctor(System.String,System.Exception)">
+            <summary>
+            Initializes a new instance of the <see cref="T:Hangfire.Client.CreateJobFailedException"/>
+            class with a specified error message and a reference to the
+            inner exception that is the cause of this exception.
+            </summary>
+            <param name="message">The error message that explains the reason for the exception.</param>
+            <param name="inner">The exception that is the cause of this exception, not null.</param>
+        </member>
+        <member name="M:Hangfire.BackgroundJobClientException.#ctor(System.String,System.Exception)">
+            <summary>
+            Initializes a new instance of the <see cref="T:Hangfire.BackgroundJobClientException"/>
+            class with a specified error message and a reference to the inner exception
+            that is the cause of this exception.
+            </summary>
+            <param name="message">The error message that explains the reason for the exception.</param>
+            <param name="inner">The exception that is the cause of this exception, not null.</param>
+        </member>
+        <member name="M:Hangfire.BackgroundJobServer.#ctor">
+            <summary>
+            Initializes a new instance of the <see cref="T:Hangfire.BackgroundJobServer"/> class
+            with default options and <see cref="P:Hangfire.JobStorage.Current"/> storage.
+            </summary>
+        </member>
+        <member name="M:Hangfire.BackgroundJobServer.#ctor(Hangfire.JobStorage)">
+            <summary>
+            Initializes a new instance of the <see cref="T:Hangfire.BackgroundJobServer"/> class
+            with default options and the given storage.
+            </summary>
+            <param name="storage">The storage</param>
+        </member>
+        <member name="M:Hangfire.BackgroundJobServer.#ctor(Hangfire.BackgroundJobServerOptions)">
+            <summary>
+            Initializes a new instance of the <see cref="T:Hangfire.BackgroundJobServer"/> class
+            with the given options and <see cref="P:Hangfire.JobStorage.Current"/> storage.
+            </summary>
+            <param name="options">Server options</param>
+        </member>
+        <member name="M:Hangfire.BackgroundJobServer.#ctor(Hangfire.BackgroundJobServerOptions,Hangfire.JobStorage)">
+            <summary>
+            Initializes a new instance of the <see cref="T:Hangfire.BackgroundJobServer"/> class
+            with the specified options and the given storage.
+            </summary>
+            <param name="options">Server options</param>
+            <param name="storage">The storage</param>
+        </member>
+        <member name="T:Hangfire.Client.IBackgroundJobFactory">
+            <summary>
+            This interface acts as extensibility point for the process
+            of job creation. See the default implementation in the
+            <see cref="T:Hangfire.Client.BackgroundJobFactory"/> class.
+            </summary>
+        </member>
+        <member name="M:Hangfire.Client.IBackgroundJobFactory.Create(Hangfire.Client.CreateContext)">
+            <summary>
+            Runs the process of job creation with the specified context.
+            </summary>
+        </member>
+        <member name="T:Hangfire.Common.Job">
+            <summary>
+            Represents an action that can be marshalled to another process to
+            be performed.
+            </summary>
+            
+            <remarks>
+            <para>The ability to serialize an action is the cornerstone of 
+            marshalling it outside of a current process boundaries. We are leaving 
+            behind all the tricky features, e.g. serializing lambdas with their
+            closures or so, and considering a simple method call information as 
+            a such an action, and using reflection to perform it.</para>
+            
+            <para>Reflection-based method invocation requires an instance of
+            the <see cref="T:System.Reflection.MethodInfo"/> class, the arguments and an instance of 
+            the type on which to invoke the method (unless it is static). Since the
+            same <see cref="T:System.Reflection.MethodInfo"/> instance can be shared across multiple 
+            types (especially when they are defined in interfaces), we also allow 
+            to specify a <see cref="P:Hangfire.Common.Job.Type"/> that contains the defined method 
+            explicitly for better flexibility.</para>
+            
+            <para>Marshalling imposes restrictions on a method that should be 
+            performed:</para>
+            
+            <list type="bullet">
+                <item>Method should be public.</item>
+                <item>Method should not contain <see langword="out"/> and <see langword="ref"/> parameters.</item>
+                <item>Method should not contain open generic parameters.</item>
+            </list>
+            </remarks>
+            
+            <example>
+            <para>The following example demonstrates the creation of a <see cref="T:Hangfire.Common.Job"/>
+            type instances using expression trees. This is the recommended way of
+            creating jobs.</para>
+            
+            <code lang="cs" source="..\Samples\Job.cs" region="Supported Methods"/>
+            
+            <para>The next example demonstrates unsupported methods. Any attempt
+            to create a job based on these methods fails with 
+            <see cref="T:System.NotSupportedException"/>.</para>
+            
+            <code lang="cs" source="..\Samples\Job.cs" region="Unsupported Methods"/>
+            </example>
+            
+            <seealso cref="T:Hangfire.IBackgroundJobClient"/>
+            <seealso cref="T:Hangfire.Server.IBackgroundJobPerformer"/>
+            
+            <threadsafety static="true" instance="false"/>
+        </member>
+        <member name="M:Hangfire.Common.Job.Perform(Hangfire.JobActivator,Hangfire.IJobCancellationToken)">
+            <exclude />
+        </member>
+        <member name="M:Hangfire.Common.Job.#ctor(System.Reflection.MethodInfo)">
+            <summary>
+            Initializes a new instance of the <see cref="T:Hangfire.Common.Job"/> class with the
+            metadata of a method with no arguments.
+            </summary>
+            
+            <param name="method">Method that should be invoked.</param>
+            
+            <exception cref="T:System.ArgumentNullException"><paramref name="method"/> argument is null.</exception>
+            <exception cref="T:System.NotSupportedException"><paramref name="method"/> is not supported.</exception>
+        </member>
+        <member name="M:Hangfire.Common.Job.#ctor(System.Reflection.MethodInfo,System.Object[])">
+            <summary>
+            Initializes a new instance of the <see cref="T:Hangfire.Common.Job"/> class with the
+            metadata of a method and the given list of arguments.
+            </summary>
+            
+            <param name="method">Method that should be invoked.</param>
+            <param name="args">Arguments that will be passed to a method invocation.</param>
+            
+            <exception cref="T:System.ArgumentNullException"><paramref name="method"/> argument is null.</exception>
+            <exception cref="T:System.ArgumentException">Parameter/argument count mismatch.</exception>
+            <exception cref="T:System.NotSupportedException"><paramref name="method"/> is not supported.</exception>
+        </member>
+        <member name="M:Hangfire.Common.Job.#ctor(System.Type,System.Reflection.MethodInfo)">
+            <summary>
+            Initializes a new instance of the <see cref="T:Hangfire.Common.Job"/> class with the
+            type, metadata of a method with no arguments.
+            </summary>
+            
+            <param name="type">Type that contains the given method.</param>
+            <param name="method">Method that should be invoked.</param>
+            
+            <exception cref="T:System.ArgumentNullException"><paramref name="type"/> argument is null.</exception>
+            <exception cref="T:System.ArgumentNullException"><paramref name="method"/> argument is null.</exception>
+            <exception cref="T:System.ArgumentException">
+            <paramref name="type"/> does not contain the given <paramref name="method"/>.
+            </exception>
+            <exception cref="T:System.ArgumentException">Parameter/argument count mismatch.</exception>
+            <exception cref="T:System.NotSupportedException"><paramref name="method"/> is not supported.</exception>
+        </member>
+        <member name="M:Hangfire.Common.Job.#ctor(System.Type,System.Reflection.MethodInfo,System.Object[])">
+            <summary>
+            Initializes a new instance of the <see cref="T:Hangfire.Common.Job"/> class with the 
+            type, metadata of a method and the given list of arguments.
+            </summary>
+            
+            <param name="type">Type that contains the given method.</param>
+            <param name="method">Method that should be invoked.</param>
+            <param name="args">Arguments that should be passed during the method call.</param>
+            
+            <exception cref="T:System.ArgumentNullException"><paramref name="type"/> argument is null.</exception>
+            <exception cref="T:System.ArgumentNullException"><paramref name="method"/> argument is null.</exception>
+            <exception cref="T:System.ArgumentNullException"><paramref name="args"/> argument is null.</exception>
+            <exception cref="T:System.ArgumentException">
+            <paramref name="type"/> does not contain the given <paramref name="method"/>.
+            </exception>
+            <exception cref="T:System.ArgumentException">Parameter/argument count mismatch.</exception>
+            <exception cref="T:System.NotSupportedException"><paramref name="method"/> is not supported.</exception>
+        </member>
+        <member name="M:Hangfire.Common.Job.FromExpression(System.Linq.Expressions.Expression{System.Action})">
+            <summary>
+            Gets a new instance of the <see cref="T:Hangfire.Common.Job"/> class based on the
+            given expression tree of a method call.
+            </summary>
+            
+            <param name="methodCall">Expression tree of a method call.</param>
+            
+            <exception cref="T:System.ArgumentNullException"><paramref name="methodCall"/> is null.</exception>
+            <exception cref="T:System.ArgumentException">
+            <paramref name="methodCall"/> expression body is not of type 
+            <see cref="T:System.Linq.Expressions.MethodCallExpression"/>.</exception>
+            <exception cref="T:System.NotSupportedException"><paramref name="methodCall"/> 
+            expression contains a method that is not supported.</exception>
+            <exception cref="T:System.InvalidOperationException"><paramref name="methodCall"/>
+            instance object of a given expression points to <see langword="null"/>.
+            </exception>
+            
+            <remarks>
+            <para>The <see cref="P:Hangfire.Common.Job.Type"/> property of a returning job will 
+            point to the type of a given instance object when it is specified, 
+            or to the declaring type otherwise. All the arguments are evaluated 
+            using the expression compiler that uses caching where possible to 
+            decrease the performance penalty.</para>
+            
+            <note>Instance object (e.g. <c>() =&gt; instance.Method()</c>) is 
+            <b>only used to obtain the type</b> for a job. It is not
+            serialized and not passed across the process boundaries.</note>
+            </remarks>
+        </member>
+        <member name="M:Hangfire.Common.Job.FromExpression``1(System.Linq.Expressions.Expression{System.Action{``0}})">
+            <summary>
+            Gets a new instance of the <see cref="T:Hangfire.Common.Job"/> class based on the
+            given expression tree of an instance method call with explicit
+            type specification.
+            </summary>
+            <typeparam name="TType">Explicit type that should be used on method call.</typeparam>
+            <param name="methodCall">Expression tree of a method call on <typeparamref name="TType"/>.</param>
+            
+            <exception cref="T:System.ArgumentNullException"><paramref name="methodCall"/> is null.</exception>
+            <exception cref="T:System.ArgumentException">
+            <paramref name="methodCall"/> expression body is not of type 
+            <see cref="T:System.Linq.Expressions.MethodCallExpression"/>.</exception>
+            <exception cref="T:System.NotSupportedException"><paramref name="methodCall"/> 
+            expression contains a method that is not supported.</exception>
+            
+            <remarks>
+            <para>All the arguments are evaluated using the expression compiler
+            that uses caching where possible to decrease the performance 
+            penalty.</para>
+            </remarks>
+        </member>
+        <member name="P:Hangfire.Common.Job.Arguments">
+            <exclude />
+        </member>
+        <member name="P:Hangfire.Common.Job.Type">
+            <summary>
+            Gets the metadata of a type that contains a method that should be 
+            invoked during the performance.
+            </summary>
+        </member>
+        <member name="P:Hangfire.Common.Job.Method">
+            <summary>
+            Gets the metadata of a method that should be invoked during the 
+            performance.
+            </summary>
+        </member>
+        <member name="P:Hangfire.Common.Job.Args">
+            <summary>
+            Gets a read-only collection of arguments that Should be passed to a 
+            method invocation during the performance.
+            </summary>
+        </member>
+        <member name="M:Hangfire.Dashboard.OwinRequestExtensions.ReadFormSafeAsync(Microsoft.Owin.IOwinContext)">
+            <summary>
+            Hack to prevent "Unable to cast object of type 'Microsoft.Owin.FormCollection' 
+            to type 'Microsoft.Owin.IFormCollection'" exception, when internalized version
+            does not match the current project's one.
+            </summary>
+        </member>
+        <member name="M:Hangfire.Dashboard.RazorPage.Execute">
+            <exclude />
+        </member>
+        <member name="M:Hangfire.Dashboard.RazorPage.Assign(Hangfire.Dashboard.RazorPage)">
+            <exclude />
+        </member>
+        <member name="M:Hangfire.Dashboard.RazorPage.WriteLiteral(System.String)">
+            <exclude />
+        </member>
+        <member name="M:Hangfire.Dashboard.RazorPage.Write(System.Object)">
+            <exclude />
+        </member>
+        <member name="T:Hangfire.NamespaceDoc">
+            <summary>
+            The <see cref="N:Hangfire"/> namespace contains high-level types for configuring,
+            creating and processing background jobs, such as <see cref="T:Hangfire.GlobalConfiguration"/>,
+            <see cref="T:Hangfire.BackgroundJob"/> and <see cref="T:Hangfire.BackgroundJobServer"/>.
+            </summary>
+        </member>
+        <member name="T:Hangfire.Annotations.NamespaceDoc">
+            <summary>
+            The <see cref="N:Hangfire.Annotations"/> namespace contains attributes that enable
+            additional code inspections in design time with JetBrains ReSharper.
+            </summary>
+            <remarks>
+            To enable annotations, open ReSharper options → Code Inspections → Code Annotations 
+            and add the <see cref="N:Hangfire.Annotations"/> namespace to the corresponding list.
+            </remarks>
+        </member>
+        <member name="T:Hangfire.Client.NamespaceDoc">
+            <summary>
+            The <see cref="N:Hangfire.Client"/> namespace contains types that allow you to
+            customize the background job creation pipeline using the <see cref="T:Hangfire.Client.IClientFilter"/>,
+            or define your own creation process by implementing the <see cref="T:Hangfire.Client.IBackgroundJobFactory"/>
+            interface.
+            </summary>
+        </member>
+        <member name="T:Hangfire.Common.NamespaceDoc">
+            <summary>
+            The <see cref="N:Hangfire.Common"/> namespace provides base types for background
+            job filters, such as <see cref="T:Hangfire.Common.JobFilterAttribute"/>, and some helper classes.
+            </summary>
+        </member>
+        <member name="T:Hangfire.Dashboard.NamespaceGroupDoc">
+            <summary>
+            The Hangfire.Dashboard namespaces contain types that allow you to
+            restrict an access to the Dashboard UI by implementing the <see cref="T:Hangfire.Dashboard.IAuthorizationFilter"/>
+            interface, as well as customize it by adding new pages, menu items, metrics, routes.
+            </summary>
+        </member>
+        <member name="T:Hangfire.Dashboard.NamespaceDoc">
+            <summary>
+            The <see cref="N:Hangfire.Dashboard"/> namespace contains types that allow you to
+            restrict an access to the Dashboard UI by implementing the <see cref="T:Hangfire.Dashboard.IAuthorizationFilter"/>
+            interface, as well as customize it by adding new pages, menu items, metrics, routes.
+            </summary>
+        </member>
+        <member name="T:Hangfire.Dashboard.Pages.NamespaceDoc">
+            <summary>
+            The <see cref="N:Hangfire.Dashboard.Pages"/> namespace contains the <see cref="T:Hangfire.Dashboard.Pages.LayoutPage"/>
+            class, layout for all the Dashboard UI pages.
+            </summary>
+        </member>
+        <member name="T:Hangfire.Logging.NamespaceGroupDoc">
+            <summary>
+            The Hangfire.Logging namespaces contain types that allow you to 
+            integrate Hangfire's logging with your projects as well as use it 
+            to log custom messages.
+            </summary>
+        </member>
+        <member name="T:Hangfire.Logging.NamespaceDoc">
+            <summary>
+            The <see cref="N:Hangfire.Logging"/> namespace contains types that allow you to 
+            integrate Hangfire's logging with your projects as well as use it 
+            to log custom messages.
+            </summary>
+        </member>
+        <member name="T:Hangfire.Logging.LogProviders.NamespaceDoc">
+            <summary>
+            The <see cref="N:Hangfire.Logging.LogProviders"/> namespace contains types for 
+            supporting most popular logging frameworks to simplify the logging integration 
+            with your projects.
+            </summary>
+        </member>
+        <member name="T:Hangfire.Server.NamespaceDoc">
+            <summary>
+            The <see cref="N:Hangfire.Server"/> namespace contains types that are responsible
+            for background processing. You may use them to customize your processing pipeline
+            by implementing the <see cref="T:Hangfire.Server.IServerFilter"/> interface or define your own 
+            continuously-running background processes by implementing the <see cref="T:Hangfire.Server.IBackgroundProcess"/> 
+            as well as create completely custom instances of <see cref="T:Hangfire.Server.BackgroundProcessingServer"/>.
+            </summary>
+        </member>
+        <member name="T:Hangfire.States.NamespaceDoc">
+            <summary>
+            The <see cref="N:Hangfire.States"/> namespace contains types that describe
+            background job states and the transitions between them. You can implement
+            custom <see cref="T:Hangfire.States.IElectStateFilter"/> or <see cref="T:Hangfire.States.IApplyStateFilter"/>
+            to customize the state changing pipeline, or define your own state by 
+            implementing the  <see cref="T:Hangfire.States.IState"/> interface.
+            </summary>
+        </member>
+        <member name="T:Hangfire.Storage.NamespaceGroupDoc">
+            <summary>
+            The Hangfire.Storage namespaces contain abstract types like <see cref="T:Hangfire.JobStorage"/>,
+            <see cref="T:Hangfire.Storage.IStorageConnection"/> and <see cref="T:Hangfire.Storage.IWriteOnlyTransaction"/> for
+            querying and modifying the underlying background job storage. 
+            These types are also used to implement support for other persistent storages.
+            </summary>
+        </member>
+        <member name="T:Hangfire.Storage.NamespaceDoc">
+            <summary>
+            The Hangfire.Storage namespaces contain abstract types like <see cref="T:Hangfire.JobStorage"/>,
+            <see cref="T:Hangfire.Storage.IStorageConnection"/> and <see cref="T:Hangfire.Storage.IWriteOnlyTransaction"/> for
+            querying and modifying the underlying background job storage. 
+            These types are also used to implement support for other persistent storages.
+            </summary>
+        </member>
+        <member name="T:Hangfire.Storage.Monitoring.NamespaceDoc">
+            <summary>
+            The <see cref="N:Hangfire.Storage.Monitoring"/> provides data transfer objects 
+            for the <see cref="T:Hangfire.Storage.IMonitoringApi"/> interface. 
+            </summary>
+            <remarks>
+            I have no idea why I placed these types to a separate namespace, they should 
+            be moved to the parent <see cref="N:Hangfire.Storage"/> namespace in version 2.0.
+            </remarks>
+        </member>
+        <member name="T:Hangfire.Server.IServerProcess">
+            <exclude />
+        </member>
+        <member name="T:Hangfire.Server.IBackgroundProcess">
+            <summary>
+            Provides methods for defining processes that will be executed in a
+            background thread by <see cref="T:Hangfire.Server.BackgroundProcessingServer"/>.
+            </summary>
+            
+            <remarks>
+            Needs a wait.
+            Cancellation token
+            Connection disposal
+            </remarks>
+            
+            <seealso cref="T:Hangfire.Server.BackgroundProcessingServer"/>
+        </member>
+        <member name="M:Hangfire.Server.IBackgroundProcess.Execute(Hangfire.Server.BackgroundProcessContext)">
+            <summary>
+            
+            </summary>
+            <param name="context">Context for a background process.</param>
+            <exception cref="T:System.ArgumentNullException"><paramref name="context"/> is null.</exception>
+        </member>
+        <member name="T:Hangfire.Server.BackgroundProcessingServer">
+            <summary>
+            Responsible for running the given collection background processes.
+            </summary>
+            
+            <remarks>
+            Immediately starts the processes in a background thread.
+            Responsible for announcing/removing a server, bound to a storage.
+            Wraps all the processes with a infinite loop and automatic retry.
+            Executes all the processes in a single context.
+            Uses timeout in dispose method, waits for all the components, cancel signals shutdown
+            Contains some required processes and uses storage processes.
+            Generates unique id.
+            Properties are still bad.
+            </remarks>
+        </member>
+        <member name="M:Hangfire.Server.BackgroundProcessingServer.#ctor(Hangfire.JobStorage,System.Collections.Generic.IEnumerable{Hangfire.Server.IBackgroundProcess},System.Collections.Generic.IDictionary{System.String,System.Object},Hangfire.Server.BackgroundProcessingServerOptions)">
+            <summary>
+            Initializes a new instance of the <see cref="T:Hangfire.Server.BackgroundProcessingServer"/>
+            class and immediately starts all the given background processes.
+            </summary>
+            <param name="storage"></param>
+            <param name="processes"></param>
+            <param name="properties"></param>
+            <param name="options"></param>
+        </member>
+        <member name="M:MoreLinq.MoreEnumerable.Pairwise``2(System.Collections.Generic.IEnumerable{``0},System.Func{``0,``0,``1})">
+            <summary>
+            Returns a sequence resulting from applying a function to each 
+            element in the source sequence and its 
+            predecessor, with the exception of the first element which is 
+            only returned as the predecessor of the second element.
+            </summary>
+            <typeparam name="TSource">The type of the elements of <paramref name="source"/>.</typeparam>
+            <typeparam name="TResult">The type of the element of the returned sequence.</typeparam>
+            <param name="source">The source sequence.</param>
+            <param name="resultSelector">A transform function to apply to 
+            each pair of sequence.</param>
+            <returns>
+            Returns the resulting sequence.
+            </returns>
+            <remarks>
+            This operator uses deferred execution and streams its results.
+            </remarks>
+            <example>
+            <code>
+            int[] numbers = { 123, 456, 789 };
+            IEnumerable&lt;int&gt; result = numbers.Pairwise(5, (a, b) => a + b);
+            </code>
+            The <c>result</c> variable, when iterated over, will yield 
+            579 and 1245, in turn.
+            </example>
+        </member>
+        <member name="T:Hangfire.States.AwaitingState">
+            <summary>
+            Defines the <i>intermediate</i> state of a background job when it is waiting
+            for a parent background job to be finished before it is moved to the
+            <see cref="T:Hangfire.States.EnqueuedState"/> by the <see cref="T:Hangfire.ContinuationsSupportAttribute"/>
+            filter.
+            </summary>
+            
+            <remarks>
+            <para>Background job in <see cref="T:Hangfire.States.AwaitingState"/> is referred as a
+            <b>continuation</b> of a background job with <see cref="P:Hangfire.States.AwaitingState.ParentId"/>.</para>
+            </remarks>
+            
+            <threadsafety static="true" instance="false"/>
+        </member>
+        <member name="T:Hangfire.States.IState">
+            <summary>
+            Provides the essential members for describing a background job state.
+            </summary>
+            <remarks>
+            <para>Background job processing in Hangfire is all about moving a background job
+            from one state to another. States are used to clearly decide what to do
+            with a background job. For example, <see cref="T:Hangfire.States.EnqueuedState"/> tells
+            Hangfire that a job should be processed by a <see cref="T:Hangfire.Server.Worker"/>,
+            and <see cref="T:Hangfire.States.FailedState"/> tells Hangfire that a job should be investigated 
+            by a developer.</para>
+            
+            <para>Each state have some essential properties like <see cref="P:Hangfire.States.IState.Name"/>,
+            <see cref="P:Hangfire.States.IState.IsFinal"/> and a custom ones that are exposed through
+            the <see cref="M:Hangfire.States.IState.SerializeData"/> method. Serialized data may be used during
+            the processing stage.</para>
+            
+            <para>Hangfire allows you to define custom states to extend the processing
+            pipeline. <see cref="T:Hangfire.States.IStateHandler"/> interface implementation can be used
+            to define additional work for a state transition, and 
+            <see cref="T:Hangfire.Server.IBackgroundProcess"/> interface implementation can be
+            used to process background jobs in a new state. For example, delayed jobs
+            and their <see cref="T:Hangfire.States.ScheduledState"/>, continuations and their 
+            <see cref="T:Hangfire.States.AwaitingState"/> can be simply moved to an extension package.</para>
+            </remarks>
+            
+            <example>
+            <para>Let's create a new state. Consider you haves background jobs that
+            throw a transient exception from time to time, and you want to simply
+            ignore those exceptions. By default, Hangfire will move a job that throwed
+            an exception to the <see cref="T:Hangfire.States.FailedState"/>, however a job in the <i>failed</i>
+            state will live in a Failed jobs page forever, unless we use <see cref="T:Hangfire.AutomaticRetryAttribute"/>,
+            delete or retry it manually, because the <see cref="T:Hangfire.States.FailedState"/> is not
+            a <i>final</i> state.</para>
+            
+            <para>Our new state will look like a <see cref="T:Hangfire.States.FailedState"/>, but we
+            define the state as a <i>final</i> one, letting Hangfire to expire faulted
+            jobs. Please refer to the <see cref="T:Hangfire.States.IState"/> interface properties to learn
+            about their details.</para>
+            
+            <para>In articles related to <see cref="T:Hangfire.States.IStateHandler"/> and <see cref="T:Hangfire.States.IElectStateFilter"/>
+            interfaces we'll discuss how to use this new state.</para>
+            
+            <code lang="cs" source="..\Samples\States.cs" region="FaultedState"/>
+            </example>
+            
+            <seealso cref="T:Hangfire.States.IBackgroundJobStateChanger"/>
+            <seealso cref="T:Hangfire.States.IStateHandler"/>
+            <seealso cref="T:Hangfire.States.IElectStateFilter"/>
+            <seealso cref="T:Hangfire.States.IApplyStateFilter"/>
+        </member>
+        <member name="M:Hangfire.States.IState.SerializeData">
+            <summary>
+            Gets a serialized representation of the current state. 
+            </summary>
+            <remarks>
+            Returning dictionary contains the serialized properties of a state. You can obtain 
+            the state data by using the <see cref="M:Hangfire.Storage.IStorageConnection.GetStateData(System.String)"/>
+            method. Please refer to documentation for this method in implementors to learn
+            which key/value pairs are available.
+            </remarks>
+            <returns>A dictionary with serialized properties of the current state.</returns>
+        </member>
+        <member name="P:Hangfire.States.IState.Name">
+            <summary>
+            Gets the unique name of the state.
+            </summary>
+            
+            <value>Unique among other states string, that is ready for 
+            ordinal comparisons.</value>
+            
+            <remarks>
+            <para>The state name is used to differentiate one state from another
+            during the state change process. So all the implemented states
+            should have a <b>unique</b> state name. Please use one-word names 
+            that start with a capital letter, in a past tense in English for 
+            your state names, for example:</para>
+            <list type="bullet">
+                <item><c>Succeeded</c></item>
+                <item><c>Enqueued</c></item>
+                <item><c>Deleted</c></item>
+                <item><c>Failed</c></item>
+            </list>
+            
+            <note type="implement">
+            The returning value should be hard-coded, no modifications of
+            this property should be allowed to a user. Implementors should
+            not add a public setter on this property.
+            </note>
+            </remarks>
+        </member>
+        <member name="P:Hangfire.States.IState.Reason">
+            <summary>
+            Gets the human-readable reason of a state transition.
+            </summary>
+            
+            <value>Any string with a reasonable length to fit dashboard elements.</value>
+            
+            <remarks>
+            <para>The reason is usually displayed in the Dashboard UI to simplify 
+            the understanding of a background job lifecycle by providing a 
+            human-readable text that explains why a background job is moved
+            to the corresponding state. Here are some examples:</para>
+            <list type="bullet">
+                <item>
+                    <i>Can not change the state to 'Enqueued': target 
+                    method was not found</i>
+                </item>
+                <item><i>Exceeded the maximum number of retry attempts</i></item>
+            </list>
+            <note type="implement">
+            The reason value is usually not hard-coded in a state implementation,
+            allowing users to change it when creating an instance of a state 
+            through the public setter.
+            </note>
+            </remarks>
+        </member>
+        <member name="P:Hangfire.States.IState.IsFinal">
+            <summary>
+            Gets if the current state is a <i>final</i> one.
+            </summary>
+            
+            <value><see langword="false"/> for <i>intermediate states</i>,
+            and <see langword="true"/> for the <i>final</i> ones.</value>
+            
+            <remarks>
+            <para>Final states define a termination stage of a background job 
+            processing pipeline. Background jobs in a final state is considered 
+            as finished with no further processing required.</para>
+            
+            <para>The <see cref="T:Hangfire.States.IBackgroundJobStateChanger">state machine</see> marks
+            finished background jobs to be expired within an interval that
+            is defined in the <see cref="P:Hangfire.States.ApplyStateContext.JobExpirationTimeout"/>
+            property that is available from a state changing filter that 
+            implements the <see cref="T:Hangfire.States.IApplyStateFilter"/> interface.</para>
+            
+            <note type="implement">
+            When implementing this property, always hard-code this property to
+            <see langword="true"/> or <see langword="false"/>. Hangfire does
+            not work with states that can be both <i>intermediate</i> and
+            <i>final</i> yet. Don't define a public setter for this property.
+            </note>
+            </remarks>
+            
+            <seealso cref="T:Hangfire.States.SucceededState"/>
+            <seealso cref="T:Hangfire.States.FailedState"/>
+            <seealso cref="T:Hangfire.States.DeletedState"/>
+        </member>
+        <member name="P:Hangfire.States.IState.IgnoreJobLoadException">
+            <summary>
+            Gets whether transition to this state should ignore job de-serialization 
+            exceptions.
+            </summary>
+            
+            <value><see langword="false"/> to move to the <see cref="T:Hangfire.States.FailedState"/> on 
+            deserialization exceptions, <see langword="true"/> to ignore them.</value>
+            
+            <remarks>
+            <para>During a state transition, an instance of the <see cref="T:Hangfire.Common.Job"/> class
+            is deserialized to get state changing filters, and to allow <see cref="T:Hangfire.States.IStateHandler">
+            state handlers</see> to perform additional work related to the state.</para>
+            
+            <para>However we cannot always deserialize a job, for example, when job method was
+            removed from the code base or its assembly reference is missing. Since background
+            processing is impossible anyway, the <see cref="T:Hangfire.States.IBackgroundJobStateChanger">state machine</see>
+            moves such a background job to the <see cref="T:Hangfire.States.FailedState"/> in this case to
+            highlight a problem to the developers (because deserialization exception may
+            occur due to bad refactorings or other programming mistakes).</para>
+            
+            <para>However, in some exceptional cases we can ignore deserialization exceptions,
+            and allow a state transition for some states that does not require a <see cref="T:Hangfire.Common.Job"/>
+            instance. <see cref="T:Hangfire.States.FailedState"/> itself and <see cref="T:Hangfire.States.DeletedState"/> are
+            examples of such a behavior.</para>
+            
+            <note type="implement">
+            In general, implementers should return <see langword="false"/> when implementing 
+            this property.
+            </note>
+            </remarks>
+            
+            <seealso cref="T:Hangfire.States.FailedState"/>
+            <seealso cref="T:Hangfire.States.DeletedState"/>
+        </member>
+        <member name="F:Hangfire.States.AwaitingState.StateName">
+            <summary>
+            Represents the name of the <i>Awaiting</i> state. This field is read-only.
+            </summary>
+            <remarks>
+            The value of this field is <c>"Awaiting"</c>.
+            </remarks>
+        </member>
+        <member name="M:Hangfire.States.AwaitingState.#ctor(System.String)">
+            <summary>
+            Initializes a new instance of the <see cref="T:Hangfire.States.AwaitingState"/> class with
+            the specified parent background job id and with an instance of the 
+            <see cref="T:Hangfire.States.EnqueuedState"/> class as a next state.
+            </summary>
+            <param name="parentId">The identifier of a background job to wait for.</param>
+        </member>
+        <member name="M:Hangfire.States.AwaitingState.#ctor(System.String,Hangfire.States.IState)">
+            <summary>
+            Initializes a new instance of the <see cref="T:Hangfire.States.AwaitingState"/> class with
+            the specified parent job id and next state.
+            </summary>
+            <param name="parentId">The identifier of a background job to wait for.</param>
+            <param name="nextState">The next state for the continuation.</param>
+        </member>
+        <member name="M:Hangfire.States.AwaitingState.#ctor(System.String,Hangfire.States.IState,Hangfire.JobContinuationOptions)">
+            <summary>
+            Initializes a new instance of the <see cref="T:Hangfire.States.AwaitingState"/> class with
+            the given <i>options</i> along with other parameters.
+            </summary>
+            <param name="parentId">The identifier of a background job to wait for.</param>
+            <param name="nextState">The next state for the continuation.</param>
+            <param name="options">Options to configure a continuation.</param>
+        </member>
+        <member name="M:Hangfire.States.AwaitingState.#ctor(System.String,Hangfire.States.IState,Hangfire.JobContinuationOptions,System.TimeSpan)">
+            <summary>
+            Initializes a new instance of the <see cref="T:Hangfire.States.AwaitingState"/> class with
+            the specified <i>expiration time</i> along with other parameters.
+            </summary>
+            <param name="parentId">The identifier of a background job to wait for.</param>
+            <param name="nextState">The next state for the continuation.</param>
+            <param name="options">Options to configure the continuation.</param>
+            <param name="expiration">The expiration time for the continuation.</param>
+        </member>
+        <member name="M:Hangfire.States.AwaitingState.SerializeData">
+            <inheritdoc/>
+            <remarks>
+            <para>Returning dictionary contains the following keys. You can obtain 
+            the state data by using the <see cref="M:Hangfire.Storage.IStorageConnection.GetStateData(System.String)"/>
+            method.</para>
+            <list type="table">
+                <listheader>
+                    <term>Key</term>
+                    <term>Type</term>
+                    <term>Deserialize Method</term>
+                    <description>Notes</description>
+                </listheader>
+                <item>
+                    <term><c>ParentId</c></term>
+                    <term><see cref="T:System.String"/></term>
+                    <term><i>Not required</i></term>
+                    <description>Please see the <see cref="P:Hangfire.States.AwaitingState.ParentId"/> property.</description>
+                </item>
+                <item>
+                    <term><c>NextState</c></term>
+                    <term><see cref="T:Hangfire.States.IState"/></term>
+                    <term>
+                        <see cref="M:Newtonsoft.Json.JsonConvert.DeserializeObject(System.String,Newtonsoft.Json.JsonSerializerSettings)"/> with 
+                        <see cref="F:Newtonsoft.Json.TypeNameHandling.Objects"/>
+                    </term>
+                    <description>Please see the <see cref="P:Hangfire.States.AwaitingState.NextState"/> property.</description>
+                </item>
+                <item>
+                    <term><c>Options</c></term>
+                    <term><see cref="T:Hangfire.JobContinuationOptions"/></term>
+                    <term>
+                        <see cref="M:System.Enum.Parse(System.Type,System.String)"/> with <see cref="T:Hangfire.JobContinuationOptions"/>
+                    </term>
+                    <description>Please see the <see cref="P:Hangfire.States.AwaitingState.Options"/> property.</description>
+                </item>
+            </list>
+            </remarks>
+        </member>
+        <member name="P:Hangfire.States.AwaitingState.ParentId">
+            <summary>
+            Gets the identifier of a parent background job.
+            </summary>
+        </member>
+        <member name="P:Hangfire.States.AwaitingState.NextState">
+            <summary>
+            Gets the next state, to which a background job will be moved.
+            </summary>
+        </member>
+        <member name="P:Hangfire.States.AwaitingState.Options">
+            <summary>
+            Gets the continuation options associated with the current state.
+            </summary>
+        </member>
+        <member name="P:Hangfire.States.AwaitingState.Expiration">
+            <summary>
+            Gets the expiration time of a background job continuation.
+            </summary>
+        </member>
+        <member name="P:Hangfire.States.AwaitingState.Name">
+            <inheritdoc/>
+            <remarks>
+            Always equals to <see cref="F:Hangfire.States.AwaitingState.StateName"/> for the <see cref="T:Hangfire.States.AwaitingState"/>.
+            Please see the remarks section of the <see cref="P:Hangfire.States.IState.Name">IState.Name</see>
+            article for the details.
+            </remarks>
+        </member>
+        <member name="P:Hangfire.States.AwaitingState.Reason">
+            <inheritdoc />
+        </member>
+        <member name="P:Hangfire.States.AwaitingState.IsFinal">
+            <inheritdoc/>
+            <remarks>
+            Always returns <see langword="false"/> for the <see cref="T:Hangfire.States.AwaitingState"/>.
+            Please refer to the <see cref="P:Hangfire.States.IState.IsFinal">IState.IsFinal</see> documentation
+            for the details.
+            </remarks>
+        </member>
+        <member name="P:Hangfire.States.AwaitingState.IgnoreJobLoadException">
+            <inheritdoc/>
+            <remarks>
+            Always returns <see langword="false"/> for the <see cref="T:Hangfire.States.AwaitingState"/>.
+            Please see the description of this property in the
+            <see cref="P:Hangfire.States.IState.IgnoreJobLoadException">IState.IgnoreJobLoadException</see>
+            article.
+            </remarks>
+        </member>
+        <member name="T:Hangfire.States.IStateHandler">
+            <summary>
+            Provides a mechanism for performing custom actions when applying or
+            unapplying the state of a background job by <see cref="T:Hangfire.States.StateMachine"/>.
+            </summary>
+        </member>
+        <member name="M:Hangfire.States.IStateHandler.Apply(Hangfire.States.ApplyStateContext,Hangfire.Storage.IWriteOnlyTransaction)">
+            <summary>
+            Performs additional actions when applying a state whose name is
+            equal to the <see cref="P:Hangfire.States.IStateHandler.StateName"/> property.
+            </summary>
+            <param name="context">The context of a state applying process.</param>
+            <param name="transaction">The current transaction of a state applying process.</param>
+        </member>
+        <member name="M:Hangfire.States.IStateHandler.Unapply(Hangfire.States.ApplyStateContext,Hangfire.Storage.IWriteOnlyTransaction)">
+            <summary>
+            Performs additional actions when unapplying a state whose name
+            is equal to the <see cref="P:Hangfire.States.IStateHandler.StateName"/> property.
+            </summary>
+            <param name="context">The context of a state applying process.</param>
+            <param name="transaction">The current transaction of a state applying process.</param>
+        </member>
+        <member name="P:Hangfire.States.IStateHandler.StateName">
+            <summary>
+            Gets the name of a state, for which custom actions will be
+            performed.
+            </summary>
+        </member>
+        <member name="T:Hangfire.Cron">
+            <summary>
+            Helper class that provides common values for the cron expressions.
+            </summary>
+        </member>
+        <member name="M:Hangfire.Cron.Minutely">
+            <summary>
+            Returns cron expression that fires every minute.
+            </summary>
+        </member>
+        <member name="M:Hangfire.Cron.Hourly">
+            <summary>
+            Returns cron expression that fires every hour at the first minute.
+            </summary>
+        </member>
+        <member name="M:Hangfire.Cron.Hourly(System.Int32)">
+            <summary>
+            Returns cron expression that fires every hour at the specified minute.
+            </summary>
+            <param name="minute">The minute in which the schedule will be activated (0-59).</param>
+        </member>
+        <member name="M:Hangfire.Cron.Daily">
+            <summary>
+            Returns cron expression that fires every day at 00:00 UTC.
+            </summary>
+        </member>
+        <member name="M:Hangfire.Cron.Daily(System.Int32)">
+            <summary>
+            Returns cron expression that fires every day at the first minute of
+            the specified hour in UTC.
+            </summary>
+            <param name="hour">The hour in which the schedule will be activated (0-23).</param>
+        </member>
+        <member name="M:Hangfire.Cron.Daily(System.Int32,System.Int32)">
+            <summary>
+            Returns cron expression that fires every day at the specified hour and minute
+            in UTC.
+            </summary>
+            <param name="hour">The hour in which the schedule will be activated (0-23).</param>
+            <param name="minute">The minute in which the schedule will be activated (0-59).</param>
+        </member>
+        <member name="M:Hangfire.Cron.Weekly">
+            <summary>
+            Returns cron expression that fires every week at Monday, 00:00 UTC.
+            </summary>
+        </member>
+        <member name="M:Hangfire.Cron.Weekly(System.DayOfWeek)">
+            <summary>
+            Returns cron expression that fires every week at 00:00 UTC of the specified
+            day of the week.
+            </summary>
+            <param name="dayOfWeek">The day of week in which the schedule will be activated.</param>
+        </member>
+        <member name="M:Hangfire.Cron.Weekly(System.DayOfWeek,System.Int32)">
+            <summary>
+            Returns cron expression that fires every week at the first minute
+            of the specified day of week and hour in UTC.
+            </summary>
+            <param name="dayOfWeek">The day of week in which the schedule will be activated.</param>
+            <param name="hour">The hour in which the schedule will be activated (0-23).</param>
+        </member>
+        <member name="M:Hangfire.Cron.Weekly(System.DayOfWeek,System.Int32,System.Int32)">
+            <summary>
+            Returns cron expression that fires every week at the specified day
+            of week, hour and minute in UTC.
+            </summary>
+            <param name="dayOfWeek">The day of week in which the schedule will be activated.</param>
+            <param name="hour">The hour in which the schedule will be activated (0-23).</param>
+            <param name="minute">The minute in which the schedule will be activated (0-59).</param>
+        </member>
+        <member name="M:Hangfire.Cron.Monthly">
+            <summary>
+            Returns cron expression that fires every month at 00:00 UTC of the first
+            day of month.
+            </summary>
+        </member>
+        <member name="M:Hangfire.Cron.Monthly(System.Int32)">
+            <summary>
+            Returns cron expression that fires every month at 00:00 UTC of the specified
+            day of month.
+            </summary>
+            <param name="day">The day of month in which the schedule will be activated (1-31).</param>
+        </member>
+        <member name="M:Hangfire.Cron.Monthly(System.Int32,System.Int32)">
+            <summary>
+            Returns cron expression that fires every month at the first minute of the
+            specified day of month and hour in UTC.
+            </summary>
+            <param name="day">The day of month in which the schedule will be activated (1-31).</param>
+            <param name="hour">The hour in which the schedule will be activated (0-23).</param>
+        </member>
+        <member name="M:Hangfire.Cron.Monthly(System.Int32,System.Int32,System.Int32)">
+            <summary>
+            Returns cron expression that fires every month at the specified day of month,
+            hour and minute in UTC.
+            </summary>
+            <param name="day">The day of month in which the schedule will be activated (1-31).</param>
+            <param name="hour">The hour in which the schedule will be activated (0-23).</param>
+            <param name="minute">The minute in which the schedule will be activated (0-59).</param>
+        </member>
+        <member name="M:Hangfire.Cron.Yearly">
+            <summary>
+            Returns cron expression that fires every year on Jan, 1st at 00:00 UTC.
+            </summary>
+        </member>
+        <member name="M:Hangfire.Cron.Yearly(System.Int32)">
+            <summary>
+            Returns cron expression that fires every year in the first day at 00:00 UTC
+            of the specified month.
+            </summary>
+            <param name="month">The month in which the schedule will be activated (1-12).</param>
+        </member>
+        <member name="M:Hangfire.Cron.Yearly(System.Int32,System.Int32)">
+            <summary>
+            Returns cron expression that fires every year at 00:00 UTC of the specified
+            month and day of month.
+            </summary>
+            <param name="month">The month in which the schedule will be activated (1-12).</param>
+            <param name="day">The day of month in which the schedule will be activated (1-31).</param>
+        </member>
+        <member name="M:Hangfire.Cron.Yearly(System.Int32,System.Int32,System.Int32)">
+            <summary>
+            Returns cron expression that fires every year at the first minute of the
+            specified month, day and hour in UTC.
+            </summary>
+            <param name="month">The month in which the schedule will be activated (1-12).</param>
+            <param name="day">The day of month in which the schedule will be activated (1-31).</param>
+            <param name="hour">The hour in which the schedule will be activated (0-23).</param>
+        </member>
+        <member name="M:Hangfire.Cron.Yearly(System.Int32,System.Int32,System.Int32,System.Int32)">
+            <summary>
+            Returns cron expression that fires every year at the specified month, day,
+            hour and minute in UTC.
+            </summary>
+            <param name="month">The month in which the schedule will be activated (1-12).</param>
+            <param name="day">The day of month in which the schedule will be activated (1-31).</param>
+            <param name="hour">The hour in which the schedule will be activated (0-23).</param>
+            <param name="minute">The minute in which the schedule will be activated (0-59).</param>
+        </member>
+        <member name="M:Hangfire.Cron.GetDescription(System.String)">
+            <summary>
+            Converts a Cron expression string into a description.
+            </summary>
+            <param name="cronExpression">A Cron expression string.</param>
+            <returns>English description.</returns>
+        </member>
+        <member name="P:Hangfire.DashboardOptions.AppPath">
+            <summary>
+            The path for the Back To Site link. Set to <see langword="null" /> in order to hide the Back To Site link.
+            </summary>
+        </member>
+        <member name="T:Hangfire.Dashboard.DashboardOwinExtensions">
+            <exclude />
+        </member>
+        <member name="M:Hangfire.Dashboard.DashboardOwinExtensions.MapHangfireDashboard(Owin.IAppBuilder)">
+            <summary>
+            Maps dashboard to the app builder pipeline at "/hangfire"
+            with authorization filter that blocks all remote requests
+            and <see cref="P:Hangfire.JobStorage.Current"/> storage instance.
+            </summary>
+            <param name="app">The app builder</param>
+        </member>
+        <member name="M:Hangfire.Dashboard.DashboardOwinExtensions.MapHangfireDashboard(Owin.IAppBuilder,System.String)">
+            <summary>
+            Maps dashboard to the app builder pipeline at the specified
+            path with authorization filter that blocks all remote requests
+            and <see cref="P:Hangfire.JobStorage.Current"/> storage instance.
+            </summary>
+            <param name="app">The app builder</param>
+            <param name="dashboardPath">The path to map dashboard</param>
+        </member>
+        <member name="M:Hangfire.Dashboard.DashboardOwinExtensions.MapHangfireDashboard(Owin.IAppBuilder,System.String,System.String)">
+            <summary>
+            Maps dashboard to the app builder pipeline at the specified
+            path with authorization filter that blocks all remote requests
+            and <see cref="P:Hangfire.JobStorage.Current"/> storage instance.
+            </summary>
+            <param name="app">The app builder</param>
+            <param name="dashboardPath">The path to map dashboard</param>
+            <param name="appPath">The application path on Back To Site link. Pass null in order to hide the Back To Site link.</param>
+        </member>
+        <member name="M:Hangfire.Dashboard.DashboardOwinExtensions.MapHangfireDashboard(Owin.IAppBuilder,System.String,System.String,System.Collections.Generic.IEnumerable{Hangfire.Dashboard.IAuthorizationFilter})">
+            <summary>
+            Maps dashboard to the app builder pipeline at the specified
+            path with given authorization filters that apply to any request
+            and <see cref="P:Hangfire.JobStorage.Current"/> storage instance.
+            </summary>
+            <param name="app">The app builder</param>
+            <param name="dashboardPath">The path to map dashboard</param>
+            <param name="appPath">The application path on Back To Site link</param>
+            <param name="authorizationFilters">Array of authorization filters</param>
+        </member>
+        <member name="M:Hangfire.Dashboard.DashboardOwinExtensions.MapHangfireDashboard(Owin.IAppBuilder,System.String,System.String,System.Collections.Generic.IEnumerable{Hangfire.Dashboard.IAuthorizationFilter},Hangfire.JobStorage)">
+            <summary>
+            Maps dashboard to the app builder pipeline at the specified path
+            with given authorization filters that apply to any request and
+            storage instance that is used to query the information.
+            </summary>
+            <param name="app">The app builder</param>
+            <param name="dashboardPath">The path to map dashboard</param>
+            <param name="appPath">The application path on Back To Site link</param>
+            <param name="authorizationFilters">Array of authorization filters</param>
+            <param name="storage">The storage instance</param>
+        </member>
+        <member name="T:Hangfire.IBootstrapperConfiguration">
+            <exclude/>
+            <summary>
+            Represents a configuration class for Hangfire components that
+            is used by the <see cref="T:Hangfire.OwinBootstrapper"/> class.
+            </summary>
+        </member>
+        <member name="M:Hangfire.IBootstrapperConfiguration.UseAuthorizationFilters(Hangfire.Dashboard.IAuthorizationFilter[])">
+            <summary>
+            Tells bootstrapper to pass the given collection of filters
+            to the dashboard middleware to authorize dashboard requests. 
+            Previous calls to this method are ignored. Empty array 
+            enables access for all users.
+            </summary>
+            <param name="filters">Authorization filters</param>
+        </member>
+        <member name="M:Hangfire.IBootstrapperConfiguration.UseFilter(System.Object)">
+            <summary>
+            Tells bootstrapper to register the given job filter globally.
+            </summary>
+            <param name="filter">Job filter instance</param>
+        </member>
+        <member name="M:Hangfire.IBootstrapperConfiguration.UseDashboardPath(System.String)">
+            <summary>
+            Tells bootstrapper to map the dashboard middleware to the
+            given path in the OWIN pipeline. 
+            </summary>
+            <param name="path">Dashboard path, '/hangfire' by default</param>
+        </member>
+        <member name="M:Hangfire.IBootstrapperConfiguration.UseAppPath(System.String)">
+            <summary>
+            Tells bootstrapper to use the given path on Back To Site link in the dashboard.
+            </summary>
+            <param name="path">Back To Site path, '/' by default</param>
+        </member>
+        <member name="M:Hangfire.IBootstrapperConfiguration.UseStorage(Hangfire.JobStorage)">
+            <summary>
+            Tells bootstrapper to register the given instance of the
+            <see cref="T:Hangfire.JobStorage"/> class globally.
+            </summary>
+            <param name="storage">Job storage</param>
+        </member>
+        <member name="M:Hangfire.IBootstrapperConfiguration.UseActivator(Hangfire.JobActivator)">
+            <summary>
+            Tells bootstrapper to register the given instance of the
+            <see cref="T:Hangfire.JobActivator"/> class globally.
+            </summary>
+            <param name="activator">Job storage</param>
+        </member>
+        <member name="M:Hangfire.IBootstrapperConfiguration.UseServer(System.Func{Hangfire.BackgroundJobServer})">
+            <summary>
+            Tells bootstrapper to start the given job server on application
+            start, and stop it automatically on application shutdown request.
+            </summary>
+            <param name="server">Job server</param>
+        </member>
+        <member name="T:Hangfire.OwinBootstrapper">
+            <exclude />
+        </member>
+        <member name="M:Hangfire.OwinBootstrapper.UseHangfire(Owin.IAppBuilder,System.Action{Hangfire.IBootstrapperConfiguration})">
+            <summary>
+            Bootstraps Hangfire components using the given configuration
+            action and maps Hangfire Dashboard to the app builder pipeline
+            at the configured path ('/hangfire' by default).
+            </summary>
+            <param name="app">The app builder</param>
+            <param name="configurationAction">Configuration action</param>
+        </member>
+        <member name="T:Hangfire.Annotations.CanBeNullAttribute">
+            <summary>
+            Indicates that the value of the marked element could be <c>null</c> sometimes,
+            so the check for <c>null</c> is necessary before its usage
+            </summary>
+            <example><code>
+            [CanBeNull] 
+            public object Test() { return null; }
+            
+            public void UseTest() 
+            {
+                var p = Test();
+                var s = p.ToString(); // Warning: Possible 'System.NullReferenceException'
+            }
+            </code></example>
+        </member>
+        <member name="T:Hangfire.Annotations.NotNullAttribute">
+            <summary>
+            Indicates that the value of the marked element could never be <c>null</c>
+            </summary>
+            <example><code>
+            [NotNull] 
+            public object Foo() 
+            {
+                return null; // Warning: Possible 'null' assignment
+            }
+            </code></example>
+        </member>
+        <member name="T:Hangfire.Annotations.StringFormatMethodAttribute">
+            <summary>
+            Indicates that the marked method builds string by format pattern and (optional) arguments.
+            Parameter, which contains format string, should be given in constructor. The format string
+            should be in <see cref="M:System.String.Format(System.IFormatProvider,System.String,System.Object[])"/>-like form
+            </summary>
+            <example><code>
+            [StringFormatMethod("message")]
+            public void ShowError(string message, params object[] args) { /* do something */ }
+            
+            public void Foo() 
+            {
+                ShowError("Failed: {0}"); // Warning: Non-existing argument in format string
+            }
+            </code></example>
+        </member>
+        <member name="M:Hangfire.Annotations.StringFormatMethodAttribute.#ctor(System.String)">
+            <param name="formatParameterName">
+            Specifies which parameter of an annotated method should be treated as format-string
+            </param>
+        </member>
+        <member name="T:Hangfire.Annotations.InvokerParameterNameAttribute">
+            <summary>
+            Indicates that the function argument should be string literal and match one
+            of the parameters of the caller function. For example, ReSharper annotates
+            the parameter of <see cref="T:System.ArgumentNullException"/>
+            </summary>
+            <example><code>
+            public void Foo(string param) 
+            {
+                if (param == null)
+                    throw new ArgumentNullException("par"); // Warning: Cannot resolve symbol
+            }
+            </code></example>
+        </member>
+        <member name="T:Hangfire.Annotations.NotifyPropertyChangedInvocatorAttribute">
+             <summary>
+             Indicates that the method is contained in a type that implements
+             <see cref="T:System.ComponentModel.INotifyPropertyChanged"/> interface
+             and this method is used to notify that some property value changed
+             </summary>
+             <remarks>
+             The method should be non-static and conform to one of the supported signatures:
+             <list>
+             <item><c>NotifyChanged(string)</c></item>
+             <item><c>NotifyChanged(params string[])</c></item>
+             <item><c>NotifyChanged{T}(Expression{Func{T}})</c></item>
+             <item><c>NotifyChanged{T,U}(Expression{Func{T,U}})</c></item>
+             <item><c>SetProperty{T}(ref T, T, string)</c></item>
+             </list>
+             </remarks>
+             <example><code>
+             public class Foo : INotifyPropertyChanged 
+             {
+                 public event PropertyChangedEventHandler PropertyChanged;
+                 [NotifyPropertyChangedInvocator]
+                 protected virtual void NotifyChanged(string propertyName) { ... }
+            
+                 private string _name;
+                 public string Name {
+                     get { return _name; }
+                     set { _name = value; NotifyChanged("LastName"); /* Warning */ }
+                 }
+             }
+             </code>
+             Examples of generated notifications:
+             <list>
+             <item><c>NotifyChanged("Property")</c></item>
+             <item><c>NotifyChanged(() =&gt; Property)</c></item>
+             <item><c>NotifyChanged((VM x) =&gt; x.Property)</c></item>
+             <item><c>SetProperty(ref myField, value, "Property")</c></item>
+             </list>
+             </example>
+        </member>
+        <member name="T:Hangfire.Annotations.ContractAnnotationAttribute">
+            <summary>
+            Describes dependency between method input and output
+            </summary>
+            <syntax>
+            <p>Function Definition Table syntax:</p>
+            <list>
+            <item>FDT      ::= FDTRow [;FDTRow]*</item>
+            <item>FDTRow   ::= Input =&gt; Output | Output &lt;= Input</item>
+            <item>Input    ::= ParameterName: Value [, Input]*</item>
+            <item>Output   ::= [ParameterName: Value]* {halt|stop|void|nothing|Value}</item>
+            <item>Value    ::= true | false | null | notnull | canbenull</item>
+            </list>
+            If method has single input parameter, it's name could be omitted.<br/>
+            Using <c>halt</c> (or <c>void</c>/<c>nothing</c>, which is the same)
+            for method output means that the methos doesn't return normally.<br/>
+            <c>canbenull</c> annotation is only applicable for output parameters.<br/>
+            You can use multiple <c>[ContractAnnotation]</c> for each FDT row,
+            or use single attribute with rows separated by semicolon.<br/>
+            </syntax>
+            <examples><list>
+            <item><code>
+            [ContractAnnotation("=> halt")]
+            public void TerminationMethod()
+            </code></item>
+            <item><code>
+            [ContractAnnotation("halt &lt;= condition: false")]
+            public void Assert(bool condition, string text) // regular assertion method
+            </code></item>
+            <item><code>
+            [ContractAnnotation("s:null => true")]
+            public bool IsNullOrEmpty(string s) // string.IsNullOrEmpty()
+            </code></item>
+            <item><code>
+            // A method that returns null if the parameter is null, and not null if the parameter is not null
+            [ContractAnnotation("null => null; notnull => notnull")]
+            public object Transform(object data) 
+            </code></item>
+            <item><code>
+            [ContractAnnotation("s:null=>false; =>true,result:notnull; =>false, result:null")]
+            public bool TryParse(string s, out Person result)
+            </code></item>
+            </list></examples>
+        </member>
+        <member name="T:Hangfire.Annotations.LocalizationRequiredAttribute">
+            <summary>
+            Indicates that marked element should be localized or not
+            </summary>
+            <example><code>
+            [LocalizationRequiredAttribute(true)]
+            public class Foo 
+            {
+                private string str = "my string"; // Warning: Localizable string
+            }
+            </code></example>
+        </member>
+        <member name="T:Hangfire.Annotations.CannotApplyEqualityOperatorAttribute">
+            <summary>
+            Indicates that the value of the marked type (or its derivatives)
+            cannot be compared using '==' or '!=' operators and <c>Equals()</c>
+            should be used instead. However, using '==' or '!=' for comparison
+            with <c>null</c> is always permitted.
+            </summary>
+            <example><code>
+            [CannotApplyEqualityOperator]
+            class NoEquality { }
+            class UsesNoEquality 
+            {
+                public void Test() 
+                {
+                    var ca1 = new NoEquality();
+                    var ca2 = new NoEquality();
+                    if (ca1 != null) // OK
+                    { 
+                        bool condition = ca1 == ca2; // Warning
+                    }
+                }
+            }
+            </code></example>
+        </member>
+        <member name="T:Hangfire.Annotations.BaseTypeRequiredAttribute">
+            <summary>
+            When applied to a target attribute, specifies a requirement for any type marked
+            with the target attribute to implement or inherit specific type or types.
+            </summary>
+            <example><code>
+            [BaseTypeRequired(typeof(IComponent)] // Specify requirement
+            public class ComponentAttribute : Attribute { }
+            [Component] // ComponentAttribute requires implementing IComponent interface
+            public class MyComponent : IComponent { }
+            </code></example>
+        </member>
+        <member name="T:Hangfire.Annotations.UsedImplicitlyAttribute">
+            <summary>
+            Indicates that the marked symbol is used implicitly
+            (e.g. via reflection, in external library), so this symbol
+            will not be marked as unused (as well as by other usage inspections)
+            </summary>
+        </member>
+        <member name="T:Hangfire.Annotations.MeansImplicitUseAttribute">
+            <summary>
+            Should be used on attributes and causes ReSharper
+            to not mark symbols marked with such attributes as unused
+            (as well as by other usage inspections)
+            </summary>
+        </member>
+        <member name="F:Hangfire.Annotations.ImplicitUseKindFlags.Access">
+            <summary>Only entity marked with attribute considered used</summary>
+        </member>
+        <member name="F:Hangfire.Annotations.ImplicitUseKindFlags.Assign">
+            <summary>Indicates implicit assignment to a member</summary>
+        </member>
+        <member name="F:Hangfire.Annotations.ImplicitUseKindFlags.InstantiatedWithFixedConstructorSignature">
+            <summary>
+            Indicates implicit instantiation of a type with fixed constructor signature.
+            That means any unused constructor parameters won't be reported as such.
+            </summary>
+        </member>
+        <member name="F:Hangfire.Annotations.ImplicitUseKindFlags.InstantiatedNoFixedConstructorSignature">
+            <summary>Indicates implicit instantiation of a type</summary>
+        </member>
+        <member name="T:Hangfire.Annotations.ImplicitUseTargetFlags">
+            <summary>
+            Specify what is considered used implicitly
+            when marked with <see cref="T:Hangfire.Annotations.MeansImplicitUseAttribute"/>
+            or <see cref="T:Hangfire.Annotations.UsedImplicitlyAttribute"/>
+            </summary>
+        </member>
+        <member name="F:Hangfire.Annotations.ImplicitUseTargetFlags.Members">
+            <summary>Members of entity marked with attribute are considered used</summary>
+        </member>
+        <member name="F:Hangfire.Annotations.ImplicitUseTargetFlags.WithMembers">
+            <summary>Entity marked with attribute and all its members considered used</summary>
+        </member>
+        <member name="T:Hangfire.Annotations.PublicAPIAttribute">
+            <summary>
+            This attribute is intended to mark publicly available API
+            which should not be removed and so is treated as used
+            </summary>
+        </member>
+        <member name="T:Hangfire.Annotations.InstantHandleAttribute">
+            <summary>
+            Tells code analysis engine if the parameter is completely handled
+            when the invoked method is on stack. If the parameter is a delegate,
+            indicates that delegate is executed while the method is executed.
+            If the parameter is an enumerable, indicates that it is enumerated
+            while the method is executed
+            </summary>
+        </member>
+        <member name="T:Hangfire.Annotations.PureAttribute">
+            <summary>
+            Indicates that a method does not make any observable state changes.
+            The same as <c>System.Diagnostics.Contracts.PureAttribute</c>
+            </summary>
+            <example><code>
+            [Pure] private int Multiply(int x, int y) { return x * y; }
+            public void Foo() {
+              const int a = 2, b = 2;
+              Multiply(a, b); // Waring: Return value of pure method is not used
+            }
+            </code></example>
+        </member>
+        <member name="T:Hangfire.RecurringJobManager">
+            <summary>
+            Represents a recurring job manager that allows to create, update
+            or delete recurring jobs.
+            </summary>
+        </member>
+        <member name="T:Hangfire.BackgroundJobClientExtensions">
+            <summary>
+            Provides extension methods for the <see cref="T:Hangfire.IBackgroundJobClient"/>
+            interface to simplify the creation of fire-and-forget jobs, delayed 
+            jobs, continuations and other background jobs in well-known states.
+            Also allows to re-queue and delete existing background jobs.
+            </summary>
+        </member>
+        <member name="M:Hangfire.BackgroundJobClientExtensions.Enqueue(Hangfire.IBackgroundJobClient,System.Linq.Expressions.Expression{System.Action})">
+            <summary>
+            Creates a background job based on a specified static method 
+            call expression and places it into its actual queue. 
+            Please, see the <see cref="T:Hangfire.QueueAttribute"/> to learn how to 
+            place the job on a non-default queue.
+            </summary>
+            
+            <param name="client">A job client instance.</param>
+            <param name="methodCall">Static method call expression that will be marshalled to the Server.</param>
+            <returns>Unique identifier of the created job.</returns>
+        </member>
+        <member name="M:Hangfire.BackgroundJobClientExtensions.Enqueue``1(Hangfire.IBackgroundJobClient,System.Linq.Expressions.Expression{System.Action{``0}})">
+            <summary>
+            Creates a background job based on a specified instance method 
+            call expression and places it into its actual queue. 
+            Please, see the <see cref="T:Hangfire.QueueAttribute"/> to learn how to 
+            place the job on a non-default queue.
+            </summary>
+            
+            <typeparam name="T">Type whose method will be invoked during job processing.</typeparam>
+            <param name="client">A job client instance.</param>
+            <param name="methodCall">Instance method call expression that will be marshalled to the Server.</param>
+            <returns>Unique identifier of the created job.</returns>
+        </member>
+        <member name="M:Hangfire.BackgroundJobClientExtensions.Schedule(Hangfire.IBackgroundJobClient,System.Linq.Expressions.Expression{System.Action},System.TimeSpan)">
+            <summary>
+            Creates a new background job based on a specified static method
+            call expression and schedules it to be enqueued after a given delay.
+            </summary>
+            <param name="client">A job client instance.</param>
+            <param name="methodCall">Instance method call expression that will be marshalled to the Server.</param>
+            <param name="delay">Delay, after which the job will be enqueued.</param>
+            <returns>Unique identifier of the created job.</returns>
+        </member>
+        <member name="M:Hangfire.BackgroundJobClientExtensions.Schedule(Hangfire.IBackgroundJobClient,System.Linq.Expressions.Expression{System.Action},System.DateTimeOffset)">
+            <summary>
+            Creates a new background job based on a specified method call expression
+            and schedules it to be enqueued at the specified moment of time.
+            </summary>
+            <param name="client">A job client instance.</param>
+            <param name="methodCall">Method call expression that will be marshalled to the Server.</param>
+            <param name="enqueueAt">Moment of time at which the job will be enqueued.</param>
+            <returns>Unique identifier or a created job.</returns>
+        </member>
+        <member name="M:Hangfire.BackgroundJobClientExtensions.Schedule``1(Hangfire.IBackgroundJobClient,System.Linq.Expressions.Expression{System.Action{``0}},System.TimeSpan)">
+            <summary>
+            Creates a new background job based on a specified instance method
+            call expression and schedules it to be enqueued after a given delay.
+            </summary>
+            
+            <typeparam name="T">Type whose method will be invoked during job processing.</typeparam>
+            <param name="client">A job client instance.</param>
+            <param name="methodCall">Instance method call expression that will be marshalled to the Server.</param>
+            <param name="delay">Delay, after which the job will be enqueued.</param>
+            <returns>Unique identifier of the created job.</returns>
+        </member>
+        <member name="M:Hangfire.BackgroundJobClientExtensions.Schedule``1(Hangfire.IBackgroundJobClient,System.Linq.Expressions.Expression{System.Action{``0}},System.DateTimeOffset)">
+            <summary>
+            Creates a new background job based on a specified method expression and schedules
+            it to be enqueued at the specified moment.
+            </summary>
+            <typeparam name="T">Type whose method will be invoked during job processing.</typeparam>
+            <param name="client">A job client instance.</param>
+            <param name="methodCall">Method call expression that will be marshalled to the Server.</param>
+            <param name="enqueueAt">Moment at which the job will be enqueued.</param>
+            <returns>Unique identifier of a created job.</returns>
+        </member>
+        <member name="M:Hangfire.BackgroundJobClientExtensions.Create(Hangfire.IBackgroundJobClient,System.Linq.Expressions.Expression{System.Action},Hangfire.States.IState)">
+            <summary>
+            Creates a new background job based on a specified static method
+            within a given state.
+            </summary>
+            <param name="client">A job client instance.</param>
+            <param name="methodCall">Static method call expression that will be marshalled to the Server.</param>
+            <param name="state">Initial state of a job.</param>
+            <returns>Unique identifier of the created job.</returns>
+        </member>
+        <member name="M:Hangfire.BackgroundJobClientExtensions.Create``1(Hangfire.IBackgroundJobClient,System.Linq.Expressions.Expression{System.Action{``0}},Hangfire.States.IState)">
+            <summary>
+            Creates a new background job based on a specified instance method
+            within a given state.
+            </summary>
+            
+            <typeparam name="T">Type whose method will be invoked during job processing.</typeparam>
+            <param name="client">A job client instance.</param>
+            <param name="methodCall">Instance method call expression that will be marshalled to the Server.</param>
+            <param name="state">Initial state of a job.</param>
+            <returns>Unique identifier of the created job.</returns>
+        </member>
+        <member name="M:Hangfire.BackgroundJobClientExtensions.ChangeState(Hangfire.IBackgroundJobClient,System.String,Hangfire.States.IState)">
+            <summary>
+            Changes state of a job with the given <paramref name="jobId"/> to
+            the specified one. 
+            </summary>
+            
+            <param name="client">An instance of <see cref="T:Hangfire.IBackgroundJobClient"/> implementation.</param>
+            <param name="jobId">A job, whose state is being changed.</param>
+            <param name="state">New state for a job.</param>
+            <returns>True, if state change succeeded, otherwise false.</returns>
+        </member>
+        <member name="M:Hangfire.BackgroundJobClientExtensions.Delete(Hangfire.IBackgroundJobClient,System.String)">
+            <summary>
+            Changes state of a job with the specified <paramref name="jobId"/>
+            to the <see cref="T:Hangfire.States.DeletedState"/>.
+            </summary>
+            
+            <remarks>
+            The job is not actually being deleted, this method changes only
+            its state.
+            
+            This operation does not provides guarantee that the job will not be 
+            performed. If you deleting a job that is performing right now, it 
+            will be performed anyway, despite of this call.
+            
+            The method returns result of a state transition. It can be false
+            if a job was expired, its method does not exist or there was an
+            exception during the state change process.
+            </remarks>
+            
+            <param name="client">An instance of <see cref="T:Hangfire.IBackgroundJobClient"/> implementation.</param>
+            <param name="jobId">Identifier of job, whose state is being changed.</param>
+            <returns>True, if state change succeeded, otherwise false.</returns>
+        </member>
+        <member name="M:Hangfire.BackgroundJobClientExtensions.Delete(Hangfire.IBackgroundJobClient,System.String,System.String)">
+            <summary>
+            Changes state of a job with the specified <paramref name="jobId"/>
+            to the <see cref="T:Hangfire.States.DeletedState"/>. If <paramref name="fromState"/> value 
+            is not null, state change will be performed only if the current state name 
+            of a job equal to the given value.
+            </summary>
+            
+            <remarks>
+            The job is not actually being deleted, this method changes only
+            its state.
+            
+            This operation does not provides guarantee that the job will not be 
+            performed. If you deleting a job that is performing right now, it 
+            will be performed anyway, despite of this call.
+            
+            The method returns result of a state transition. It can be false
+            if a job was expired, its method does not exist or there was an
+            exception during the state change process.
+            </remarks>
+            
+            <param name="client">An instance of <see cref="T:Hangfire.IBackgroundJobClient"/> implementation.</param>
+            <param name="jobId">Identifier of job, whose state is being changed.</param>
+            <param name="fromState">Current state assertion, or null if unneeded.</param>
+            <returns>True, if state change succeeded, otherwise false.</returns>
+        </member>
+        <member name="M:Hangfire.BackgroundJobClientExtensions.Requeue(Hangfire.IBackgroundJobClient,System.String)">
+            <summary>
+            Changes state of a job with the specified <paramref name="jobId"/>
+            to the <see cref="T:Hangfire.States.EnqueuedState"/>.
+            </summary>
+            
+            <param name="client">An instance of <see cref="T:Hangfire.IBackgroundJobClient"/> implementation.</param>
+            <param name="jobId">Identifier of job, whose state is being changed.</param>
+            <returns>True, if state change succeeded, otherwise false.</returns>
+        </member>
+        <member name="M:Hangfire.BackgroundJobClientExtensions.Requeue(Hangfire.IBackgroundJobClient,System.String,System.String)">
+            <summary>
+            Changes state of a job with the specified <paramref name="jobId"/>
+            to the <see cref="T:Hangfire.States.EnqueuedState"/>. If <paramref name="fromState"/> value 
+            is not null, state change will be performed only if the current state name 
+            of a job equal to the given value.
+            </summary>
+            
+            <param name="client">An instance of <see cref="T:Hangfire.IBackgroundJobClient"/> implementation.</param>
+            <param name="jobId">Identifier of job, whose state is being changed.</param>
+            <param name="fromState">Current state assertion, or null if unneeded.</param>
+            <returns>True, if state change succeeded, otherwise false.</returns>
+        </member>
+        <member name="T:Hangfire.Common.CachedExpressionCompiler">
+             <summary>
+             The caching expression tree compiler was copied from MVC core to MVC Futures so that Futures code could benefit
+             from it and so that it could be exposed as a public API. This is the only public entry point into the system.
+             See the comments in the ExpressionUtil namespace for more information.
+            
+             The unit tests for the ExpressionUtil.* types are in the System.Web.Mvc.Test project.
+             </summary>
+        </member>
+        <member name="M:Hangfire.Common.CachedExpressionCompiler.Evaluate(System.Linq.Expressions.Expression)">
+            <summary>
+            Evaluates an expression (not a LambdaExpression), e.g. 2 + 2.
+            </summary>
+            <param name="arg"></param>
+            <returns>Expression result.</returns>
+        </member>
+        <member name="T:Hangfire.Server.IServerFilter">
+            <summary>
+            Defines methods that are required for a server filter.
+            </summary>
+        </member>
+        <member name="M:Hangfire.Server.IServerFilter.OnPerforming(Hangfire.Server.PerformingContext)">
+            <summary>
+            Called before the performance of the job.
+            </summary>
+            <param name="filterContext">The filter context.</param>
+        </member>
+        <member name="M:Hangfire.Server.IServerFilter.OnPerformed(Hangfire.Server.PerformedContext)">
+            <summary>
+            Called after the performance of the job.
+            </summary>
+            <param name="filterContext">The filter context.</param>
+        </member>
+        <member name="T:Hangfire.IBackgroundJobClient">
+            <summary>
+            Provides methods for creating all the types of background jobs and 
+            changing their states.
+            </summary>
+            
+            <remarks>
+            <para>Please see the <see cref="T:Hangfire.BackgroundJobClient"/> class for
+            details regarding the implementation.</para>
+            </remarks>
+        </member>
+        <member name="M:Hangfire.IBackgroundJobClient.Create(Hangfire.Common.Job,Hangfire.States.IState)">
+            <summary>
+            Creates a new background job in a specified state.
+            </summary>
+            
+            <param name="job">Job that should be processed in background.</param>
+            <param name="state">Initial state for a background job.</param>
+            <returns>Unique identifier of a created background job <i>-or-</i> 
+             <see langword="null"/>, if it was not created.</returns>
+            
+            <exception cref="T:System.ArgumentNullException"><paramref name="job"/> is null.</exception>
+            <exception cref="T:System.ArgumentNullException"><paramref name="state"/> is null.</exception>
+            <exception cref="T:Hangfire.BackgroundJobClientException">Creation failed due to an exception.</exception>
+            
+            <remarks>
+            <para>The interface allows implementations to return <see langword="null"/> 
+            value for this method when background job creation has been canceled
+            by an implementation under the normal circumstances (not due to an
+            exception). For example, the <see cref="T:Hangfire.Client.CreatingContext"/> class
+            contains the <see cref="P:Hangfire.Client.CreatingContext.Canceled"/> property that
+            may be used by a client filter to cancel a background job creation.
+            </para>
+            
+            <para>The interface allows implementations to create a background 
+            job in a state other than specified. The given state instance also 
+            may be modified. For example, <see cref="T:Hangfire.States.ElectStateContext"/> class
+            contains public setter for the <see cref="P:Hangfire.States.ElectStateContext.CandidateState"/>
+            property allowing to choose completely different state by state
+            election filters.</para>
+            </remarks>
+        </member>
+        <member name="M:Hangfire.IBackgroundJobClient.ChangeState(System.String,Hangfire.States.IState,System.String)">
+            <summary>
+            Attempts to change a state of a background job with a given
+            identifier to a specified one.
+            </summary>
+            
+            <param name="jobId">Identifier of background job, whose state should be changed.</param>
+            <param name="state">New state for a background job.</param>
+            <param name="expectedState">Expected state assertion, or <see langword="null"/> if unneeded.</param>
+            
+            <returns><see langword="true"/>, if a <b>given</b> state was applied
+            successfully otherwise <see langword="false"/>.</returns>
+            
+            <exception cref="T:System.ArgumentNullException"><paramref name="jobId"/> is null.</exception>
+            <exception cref="T:System.ArgumentNullException"><paramref name="state"/> is null.</exception>
+            <exception cref="T:Hangfire.BackgroundJobClientException">State change failed due to an exception.</exception>
+            
+            <remarks>
+            <para>If <paramref name="expectedState"/> value is not null, state 
+            change will be performed only if the current state name of a job 
+            equal to the given value.</para>
+            
+            <para>The interface allows implementations to change a state of a 
+            background job to other than specified. The given state instance also
+            may be modified. For example, <see cref="T:Hangfire.States.ElectStateContext"/> class
+            contains public setter for the <see cref="P:Hangfire.States.ElectStateContext.CandidateState"/>
+            property allowing to choose completely different state by state
+            election filters. If a state was changed, <see langword="false"/> 
+            value will be returned.</para>
+            </remarks>
+        </member>
+        <member name="T:Hangfire.Client.CreateContext">
+            <summary>
+            Provides information about the context in which the job is created.
+            </summary>
+        </member>
+        <member name="P:Hangfire.Client.CreateContext.Items">
+            <summary>
+            Gets an instance of the key-value storage. You can use it
+            to pass additional information between different client filters
+            or just between different methods.
+            </summary>
+        </member>
+        <member name="P:Hangfire.Client.CreateContext.InitialState">
+            <summary>
+            Gets the initial state of the creating job. Note, that
+            the final state of the created job could be changed after 
+            the registered instances of the <see cref="T:Hangfire.States.IElectStateFilter"/>
+            class are doing their job.
+            </summary>
+        </member>
+        <member name="T:Hangfire.Common.JobLoadException">
+            <summary>
+            The exception that is thrown when a job could not
+            be loaded from the storage due to missing or incorrect 
+            information about its type or method.
+            </summary>
+        </member>
+        <member name="M:Hangfire.Common.JobLoadException.#ctor(System.String,System.Exception)">
+            <summary>
+            Initializes a new instance of the <see cref="T:Hangfire.Common.JobLoadException"/>
+            class with a given message and information about inner exception.
+            </summary>
+        </member>
+        <member name="T:Hangfire.Common.JobFilter">
+            <summary>
+            Represents a metadata class that contains a reference to the 
+            implementation of one or more of the filter interfaces, the filter's 
+            order, and the filter's scope.
+            </summary>
+        </member>
+        <member name="F:Hangfire.Common.JobFilter.DefaultOrder">
+            <summary>
+            Represents a constant that is used to specify the default ordering of filters.
+            </summary>
+        </member>
+        <member name="M:Hangfire.Common.JobFilter.#ctor(System.Object,Hangfire.Common.JobFilterScope,System.Nullable{System.Int32})">
+            <summary>
+            Initializes a new instance of the Filter class.
+            </summary>
+            <param name="instance">Filter instance.</param>
+            <param name="scope">Filter scope.</param>
+            <param name="order">The run order.</param>
+        </member>
+        <member name="P:Hangfire.Common.JobFilter.Instance">
+            <summary>
+            Gets the instance of the filter.
+            </summary>
+        </member>
+        <member name="P:Hangfire.Common.JobFilter.Order">
+            <summary>
+            Gets the order in which the filter is applied.
+            </summary>
+        </member>
+        <member name="P:Hangfire.Common.JobFilter.Scope">
+            <summary>
+            Gets the scope ordering of the filter.
+            </summary>
+        </member>
+        <member name="T:Hangfire.Common.JobFilterProviderCollection">
+            <summary>
+            Represents the collection of filter providers for the application.
+            </summary>
+        </member>
+        <member name="T:Hangfire.Common.IJobFilterProvider">
+            <summary>
+            Provides an interface for finding filters.
+            </summary>
+        </member>
+        <member name="M:Hangfire.Common.IJobFilterProvider.GetFilters(Hangfire.Common.Job)">
+            <summary>
+            Returns an enumerator that contains all the <see cref="T:Hangfire.Common.IJobFilterProvider"/>.
+            </summary>
+             
+            <returns>
+            The enumerator that contains all the <see cref="T:Hangfire.Common.IJobFilterProvider"/>.
+            </returns>
+        </member>
+        <member name="M:Hangfire.Common.JobFilterProviderCollection.#ctor">
+            <summary>
+            Initializes a new instance of the <see cref="T:Hangfire.Common.JobFilterProviderCollection"/> 
+            class.
+            </summary>
+        </member>
+        <member name="M:Hangfire.Common.JobFilterProviderCollection.GetFilters(Hangfire.Common.Job)">
+            <summary>
+            Returns the collection of filter providers.
+            </summary>
+            <param name="job">Job, can be null.</param>
+            <returns>The collection of filter providers.</returns>
+        </member>
+        <member name="T:Hangfire.Common.JobFilterProviders">
+            <summary>
+            Provides a registration point for filters.
+            </summary>
+        </member>
+        <member name="P:Hangfire.Common.JobFilterProviders.Providers">
+            <summary>
+            Provides a registration point for filters.
+            </summary>
+        </member>
+        <member name="T:Hangfire.Common.JobFilterScope">
+            <summary>
+            Defines values that specify the order in which Hangfire filters 
+            run within the same filter type and filter order.
+            </summary>
+            
+            <remarks>
+            Hangfire supports the following types of filters:
+            
+            <list type="number">
+                <item>
+                    <description>
+                        Client / Server filters, which implement
+                        <see cref="T:Hangfire.Client.IClientFilter"/> and <see cref="T:Hangfire.Server.IServerFilter"/>
+                        interfaces respectively.
+                    </description>
+                </item>
+                <item>
+                    <description>
+                        State changing filters, which implement the
+                        <see cref="T:Hangfire.States.IElectStateFilter"/> interface.
+                    </description>
+                </item>
+                <item>
+                    <description>
+                        State changed filters, which implement the
+                        <see cref="T:Hangfire.States.IApplyStateFilter"/> interface.
+                    </description>
+                </item>
+                <item>
+                    <description>
+                        Client / Server exception filters, which implement
+                        <see cref="T:Hangfire.Client.IClientExceptionFilter"/> or 
+                        <see cref="T:Hangfire.Server.IServerExceptionFilter"/> interfaces
+                        respectively.
+                    </description>
+                </item>
+            </list>
+            
+            Порядок запуска указанных типов фильтров строго фиксирован, например,
+            фильтры исключений всегда выполняются после всех остальных фильтров,
+            а фильтры состояний всегда запускаются внутри клиентских и серверных
+            фильтров.
+            
+            Внутри же одного типа фильтров, порядок выполнения сначала определяется
+            значением Order, а затем значением Scope. Перечисление <see cref="T:Hangfire.Common.JobFilterScope"/> 
+            определяет следующие значения (в порядке, в котором они будут выполнены):
+            
+            <list type="number">
+                <item>
+                    <description>
+                        <see cref="F:Hangfire.Common.JobFilterScope.Global"/>.
+                    </description>
+                </item>
+                <item>
+                    <description>
+                        <see cref="F:Hangfire.Common.JobFilterScope.Type"/>.
+                    </description>
+                </item>
+                <item>
+                    <description>
+                        <see cref="F:Hangfire.Common.JobFilterScope.Method"/>.
+                    </description>
+                </item>
+            </list>
+            
+            Для примера, клиентский фильтр, у которого свойство Order имеет значение 0,
+            а значение filter scope равно <see cref="F:Hangfire.Common.JobFilterScope.Global"/>,
+            будет выполнен раньше фильтра с тем же самым значением Order,
+            но c filter scope, равным <see cref="F:Hangfire.Common.JobFilterScope.Type"/>.
+            
+            Значения Scope задаются, в основном, в реализациях интерфейса
+            <see cref="T:Hangfire.Common.IJobFilterProvider"/>. Так, класс <see cref="T:Hangfire.Common.JobFilterCollection"/>
+            определяет значение Scope как <see cref="F:Hangfire.Common.JobFilterScope.Global"/>.
+            
+            Порядок выполнения фильтров одинакового типа, с одинаковым значением
+            Order и с одинаковым scope, не оговаривается.
+            </remarks>
+        </member>
+        <member name="F:Hangfire.Common.JobFilterScope.Global">
+            <summary>
+            Specifies an order before the <see cref="F:Hangfire.Common.JobFilterScope.Type"/>.
+            </summary>
+        </member>
+        <member name="F:Hangfire.Common.JobFilterScope.Type">
+            <summary>
+            Specifies an order after the <see cref="F:Hangfire.Common.JobFilterScope.Global"/> and
+            before the <see cref="F:Hangfire.Common.JobFilterScope.Method"/>.
+            </summary>
+        </member>
+        <member name="F:Hangfire.Common.JobFilterScope.Method">
+            <summary>
+            Specifies an order after the <see cref="F:Hangfire.Common.JobFilterScope.Type"/>.
+            </summary>
+        </member>
+        <member name="T:Hangfire.Client.ClientExceptionContext">
+            <summary>
+            Provides the context for the <see cref="M:Hangfire.Client.IClientExceptionFilter.OnClientException(Hangfire.Client.ClientExceptionContext)"/>
+            method of the <see cref="T:Hangfire.Client.IClientExceptionFilter"/> interface.
+            </summary>
+        </member>
+        <member name="P:Hangfire.Client.ClientExceptionContext.Exception">
+            <summary>
+            Gets an exception that occurred during the creation of the job.
+            </summary>
+        </member>
+        <member name="P:Hangfire.Client.ClientExceptionContext.ExceptionHandled">
+            <summary>
+            Gets or sets a value that indicates that this <see cref="T:Hangfire.Client.ClientExceptionContext"/>
+            object handles an exception occurred during the creation of the job.
+            </summary>
+        </member>
+        <member name="T:Hangfire.Common.JobFilterCollection">
+            <summary>
+            Represents a class that contains the job filters.
+            </summary>
+            <remarks>
+            Job filters run for every create, perform and state change
+            of every job. All the instances in the filters collection
+            should be thread-safe.
+            
+            You can register a filter using the 
+            <see cref="P:Hangfire.GlobalJobFilters.Filters"/> registration endpoint.
+            </remarks>
+        </member>
+        <member name="M:Hangfire.Common.JobFilterCollection.Add(System.Object)">
+            <summary>
+            Adds the specified filter to the global filter collection.
+            </summary>
+            <param name="filter">The filter instance.</param>
+        </member>
+        <member name="M:Hangfire.Common.JobFilterCollection.Add(System.Object,System.Int32)">
+            <summary>
+            Adds the specified filter to the global filter collection 
+            using the specified filter run order.
+            </summary>
+            <param name="filter">The filter instance.</param>
+            <param name="order">The run order.</param>
+        </member>
+        <member name="M:Hangfire.Common.JobFilterCollection.Clear">
+            <summary>
+            Removes all filters from the global filter collection.
+            </summary>
+        </member>
+        <member name="M:Hangfire.Common.JobFilterCollection.Contains(System.Object)">
+            <summary>
+            Determines wheither a filter is in the global filter collection.
+            </summary>
+            <param name="filter">The filter instance.</param>
+            <returns>True if the global filter collection contains the filter, otherwise false.</returns>
+        </member>
+        <member name="M:Hangfire.Common.JobFilterCollection.Remove(System.Object)">
+            <summary>
+            Removes all filters that match the specified filter.
+            </summary>
+            <param name="filter">The filter instance.</param>
+        </member>
+        <member name="P:Hangfire.Common.JobFilterCollection.Count">
+            <summary>
+            Gets the number of filters in the global job filter collection.
+            </summary>
+        </member>
+        <member name="T:Hangfire.Client.IClientExceptionFilter">
+            <summary>
+            Defines methods that are required for the client exception filter.
+            </summary>
+        </member>
+        <member name="M:Hangfire.Client.IClientExceptionFilter.OnClientException(Hangfire.Client.ClientExceptionContext)">
+            <summary>
+            Called when an exception occurred during the creation of the job.
+            </summary>
+            <param name="filterContext">The filter context.</param>
+        </member>
+        <member name="T:Hangfire.Server.RecurringJobScheduler">
+            <summary>
+            Represents a background process responsible for <i>enqueueing recurring 
+            jobs</i>.
+            </summary>
+            
+            <remarks>
+            <para>This background process polls the <i>recurring job schedule</i>
+            for recurring jobs ready to be enqueued. Interval between scheduler
+            polls is hard-coded to <b>1 minute</b> as a compromise between
+            frequency and additional stress on job storage.</para>
+            
+            <note type="tip">
+            Use custom background processes if you need to schedule recurring jobs
+            with frequency less than one minute. Please see the 
+            <see cref="T:Hangfire.Server.IBackgroundProcess"/> interface for details.
+            </note>
+            
+            <para>Recurring job schedule is based on Set and Hash data structures
+            of a job storage, so you can use this background process as an example 
+            of a custom extension.</para>
+            
+            <para>Multiple instances of this background process can be used in
+            separate threads/processes without additional configuration (distributed
+            locks are used). However, this only adds support for fail-over, and does 
+            not increase the performance.</para>
+            
+            <note type="important">
+            If you are using <b>custom filter providers</b>, you need to pass a 
+            custom <see cref="T:Hangfire.Client.IBackgroundJobFactory"/> instance to make this 
+            process respect your filters when enqueueing background jobs.
+            </note>
+            </remarks>
+            
+            <threadsafety static="true" instance="true"/>
+            
+            <seealso cref="T:Hangfire.RecurringJobManager"/>
+        </member>
+        <member name="M:Hangfire.Server.RecurringJobScheduler.#ctor">
+            <summary>
+            Initializes a new instance of the <see cref="T:Hangfire.Server.RecurringJobScheduler"/>
+            class with default background job factory.
+            </summary>
+        </member>
+        <member name="M:Hangfire.Server.RecurringJobScheduler.#ctor(Hangfire.Client.IBackgroundJobFactory)">
+            <summary>
+            Initializes a new instance of the <see cref="T:Hangfire.Server.RecurringJobScheduler"/>
+            class with custom background job factory.
+            </summary>
+            <param name="factory">Factory that will be used to create background jobs.</param>
+            
+            <exception cref="T:System.ArgumentNullException"><paramref name="factory"/> is null.</exception>
+        </member>
+        <member name="M:Hangfire.Server.RecurringJobScheduler.Execute(Hangfire.Server.BackgroundProcessContext)">
+            <inheritdoc />
+        </member>
+        <member name="M:Hangfire.Server.RecurringJobScheduler.ToString">
+            <inheritdoc />
+        </member>
+        <member name="T:Hangfire.Server.ServerOwinExtensions">
+            <exclude />
+        </member>
+        <member name="M:Hangfire.Server.ServerOwinExtensions.RunHangfireServer(Owin.IAppBuilder,Hangfire.BackgroundJobServer)">
+            <summary>
+            Starts the specified background job server and registers the call
+            to its `Dispose` method at OWIN application's shutdown event.
+            </summary>
+            <param name="app">The app builder</param>
+            <param name="server">The background job server to start</param>
+        </member>
+        <member name="T:Hangfire.BootstrapperConfiguration">
+            <exclude />
+        </member>
+        <member name="T:Hangfire.BootstrapperConfigurationExtensions">
+            <exclude />
+        </member>
+        <member name="M:Hangfire.BootstrapperConfigurationExtensions.UseServer(Hangfire.IBootstrapperConfiguration)">
+            <summary>
+            Tells bootstrapper to start a job server with default options
+            on application start and stop it automatically on application
+            shutdown request. Global job storage is used.
+            </summary>
+            <param name="configuration">Configuration</param>
+        </member>
+        <member name="M:Hangfire.BootstrapperConfigurationExtensions.UseServer(Hangfire.IBootstrapperConfiguration,System.Int32)">
+            <summary>
+            Tells bootstrapper to start a job server with the given
+            amount of workers on application start and stop it automatically
+            on application shutdown request. Global job storage is used.
+            </summary>
+            <param name="configuration">Configuration</param>
+            <param name="workerCount">Worker count</param>
+        </member>
+        <member name="M:Hangfire.BootstrapperConfigurationExtensions.UseServer(Hangfire.IBootstrapperConfiguration,System.String[])">
+            <summary>
+            Tells bootstrapper to start a job server with the given
+            queues array on application start and stop it automatically
+            on application shutdown request. Global job storage is used.
+            </summary>
+            <param name="configuration">Configuration</param>
+            <param name="queues">Queues to listen</param>
+        </member>
+        <member name="M:Hangfire.BootstrapperConfigurationExtensions.UseServer(Hangfire.IBootstrapperConfiguration,System.Int32,System.String[])">
+            <summary>
+            Tells bootstrapper to start a job server with the given
+            queues array and specified amount of workers on application
+            start and stop it automatically on application shutdown request.
+            Global job storage is used.
+            </summary>
+            <param name="configuration">Configuration</param>
+            <param name="workerCount">Worker count</param>
+            <param name="queues">Queues to listen</param>
+        </member>
+        <member name="M:Hangfire.BootstrapperConfigurationExtensions.UseServer(Hangfire.IBootstrapperConfiguration,Hangfire.BackgroundJobServerOptions)">
+            <summary>
+            Tells bootstrapper to start a job server with the given
+            options on application start and stop it automatically
+            on application shutdown request. Global job storage is used.
+            </summary>
+            <param name="configuration">Configuration</param>
+            <param name="options">Job server options</param>
+        </member>
+        <member name="M:Hangfire.BootstrapperConfigurationExtensions.UseServer(Hangfire.IBootstrapperConfiguration,Hangfire.JobStorage)">
+            <summary>
+            Tells bootstrapper to start a job server, that uses
+            the given job storage, on application start and stop
+            it automatically on application shutdown request.
+            </summary>
+            <param name="configuration">Configuration</param>
+            <param name="storage">Job storage to use</param>
+        </member>
+        <member name="M:Hangfire.BootstrapperConfigurationExtensions.UseServer(Hangfire.IBootstrapperConfiguration,Hangfire.JobStorage,Hangfire.BackgroundJobServerOptions)">
+            <summary>
+            Tells bootstrapper to start a job server with the given
+            options that use the specified storage (not the global one) on
+            application start and stop it automatically on application
+            shutdown request.
+            </summary>
+            <param name="configuration">Configuration</param>
+            <param name="storage">Job storage to use</param>
+            <param name="options">Job server options</param>
+        </member>
+        <member name="T:Hangfire.States.IStateMachine">
+            <summary>
+            Provides a mechanism for running state election and state applying processes.
+            </summary>
+            
+            <seealso cref="T:Hangfire.States.StateMachine"/>
+        </member>
+        <member name="M:Hangfire.States.IStateMachine.ApplyState(Hangfire.States.ApplyStateContext)">
+            <summary>
+            Performs the state applying process, where a current background job
+            will be moved to the elected state.
+            </summary>
+            <param name="context">The context of a state applying process.</param>
+        </member>
+        <member name="T:Hangfire.States.DeletedState">
+            <summary>
+            Defines the <i>final</i> state of a background job when nobody
+            is interested whether it was performed or not.
+            </summary>
+            <remarks>
+            <para>Deleted state is used when you are not interested in a processing
+            of a background job. This state isn't backed by any background process,
+            so when you change a state of the job to the <i>Deleted</i>, only
+            expiration time will be set on a job without any additional processing.</para>
+            </remarks>
+            
+            <example>
+            <para>The following example demonstrates how to cancel an <i>enqueued</i> background
+            job. Please note that this job may be processed before you change its state.</para>
+            <para>This example shows how to create an instance of the <see cref="T:Hangfire.States.DeletedState"/>
+            class and use the <see cref="M:Hangfire.IBackgroundJobClient.ChangeState(System.String,Hangfire.States.IState,System.String)"/> method. Please see
+            <see cref="O:Hangfire.BackgroundJob.Delete">BackgroundJob.Delete</see>
+            and <see cref="O:Hangfire.BackgroundJobClientExtensions.Delete">BackgroundJobClientExtensions.Delete</see>
+            method overloads for simpler API.</para> 
+            
+            <code lang="cs" source="..\Samples\States.cs" region="DeletedState"/>
+            
+            </example>
+            
+            <seealso cref="O:Hangfire.BackgroundJob.Delete">BackgroundJob.Delete Overload</seealso>
+            <seealso cref="O:Hangfire.BackgroundJobClientExtensions.Delete">BackgroundJobClientExtensions.Delete Overload</seealso>
+            <seealso cref="M:Hangfire.IBackgroundJobClient.ChangeState(System.String,Hangfire.States.IState,System.String)"/>
+            
+            <threadsafety static="true" instance="false"/>
+        </member>
+        <member name="F:Hangfire.States.DeletedState.StateName">
+            <summary>
+            Represents the name of the <i>Deleted</i> state. This field is read-only.
+            </summary>
+            <remarks>
+            The value of this field is <c>"Deleted"</c>.
+            </remarks>
+        </member>
+        <member name="M:Hangfire.States.DeletedState.#ctor">
+            <summary>
+            Initializes a new instance of the <see cref="T:Hangfire.States.DeletedState"/> class.
+            </summary>
+        </member>
+        <member name="M:Hangfire.States.DeletedState.SerializeData">
+            <inheritdoc/>
+            <remarks>
+            <para>Returning dictionary contains the following keys. You can obtain 
+            the state data by using the <see cref="M:Hangfire.Storage.IStorageConnection.GetStateData(System.String)"/>
+            method.</para>
+            <list type="table">
+                <listheader>
+                    <term>Key</term>
+                    <term>Type</term>
+                    <term>Deserialize Method</term>
+                    <description>Notes</description>
+                </listheader>
+                <item>
+                    <term><c>DeletedAt</c></term>
+                    <term><see cref="T:System.DateTime"/></term>
+                    <term><see cref="M:Hangfire.Common.JobHelper.DeserializeDateTime(System.String)"/></term>
+                    <description>Please see the <see cref="P:Hangfire.States.DeletedState.DeletedAt"/> property.</description>
+                </item>
+            </list>
+            </remarks>
+        </member>
+        <member name="P:Hangfire.States.DeletedState.Name">
+            <inheritdoc/>
+            <remarks>
+            Always equals to <see cref="F:Hangfire.States.DeletedState.StateName"/> for the <see cref="T:Hangfire.States.DeletedState"/>.
+            Please see the remarks section of the <see cref="P:Hangfire.States.IState.Name">IState.Name</see>
+            article for the details.
+            </remarks>
+        </member>
+        <member name="P:Hangfire.States.DeletedState.Reason">
+            <inheritdoc />
+        </member>
+        <member name="P:Hangfire.States.DeletedState.IsFinal">
+            <inheritdoc/>
+            <remarks>
+            Always returns <see langword="true"/> for the <see cref="T:Hangfire.States.DeletedState"/>.
+            Please refer to the <see cref="P:Hangfire.States.IState.IsFinal">IState.IsFinal</see> documentation
+            for the details.
+            </remarks>
+        </member>
+        <member name="P:Hangfire.States.DeletedState.IgnoreJobLoadException">
+            <inheritdoc/>
+            <remarks>
+            Always returns <see langword="true"/> for the <see cref="T:Hangfire.States.DeletedState"/>.
+            Please see the description of this property in the
+            <see cref="P:Hangfire.States.IState.IgnoreJobLoadException">IState.IgnoreJobLoadException</see>
+            article.
+            </remarks>
+        </member>
+        <member name="P:Hangfire.States.DeletedState.DeletedAt">
+            <summary>
+            Gets a date/time when the current state instance was created.
+            </summary>
+        </member>
+        <member name="T:Hangfire.Common.JobFilterAttributeFilterProvider">
+            <summary>
+            Defines a filter provider for filter attributes.
+            </summary>
+        </member>
+        <member name="M:Hangfire.Common.JobFilterAttributeFilterProvider.#ctor">
+            <summary>
+            Initializes a new instance of the  <see cref="T:Hangfire.Common.JobFilterAttributeFilterProvider"/>
+            class with the attribute instance caching enabled.
+            </summary>
+        </member>
+        <member name="M:Hangfire.Common.JobFilterAttributeFilterProvider.#ctor(System.Boolean)">
+            <summary>
+            Initializes a new instance of the <see cref="T:Hangfire.Common.JobFilterAttributeFilterProvider"/> 
+            class and optionally caches attribute instances.
+            </summary>
+            <param name="cacheAttributeInstances"></param>
+        </member>
+        <member name="T:Hangfire.Common.JobFilterInfo">
+            <summary>
+            Encapsulates information about the available job filters.
+            </summary>
+        </member>
+        <member name="M:Hangfire.Common.JobFilterInfo.#ctor(System.Collections.Generic.IEnumerable{Hangfire.Common.JobFilter})">
+            <summary>
+            Initializes a new instance of the <see cref="T:Hangfire.Common.JobFilterInfo"/> class using the specified filters collection.
+            </summary>
+            <param name="filters">The filters collection.</param>
+        </member>
+        <member name="P:Hangfire.Common.JobFilterInfo.ClientFilters">
+            <summary>
+            Gets all the client filters in the application.
+            </summary>
+            
+            <returns>
+            The client filters.
+            </returns>
+        </member>
+        <member name="P:Hangfire.Common.JobFilterInfo.ServerFilters">
+            <summary>
+            Gets all the server filters in the application.
+            </summary>
+            
+            <returns>
+            The server filters.
+            </returns>
+        </member>
+        <member name="P:Hangfire.Common.JobFilterInfo.ElectStateFilters">
+            <summary>
+            Gets all the stat changing filters in the application.
+            </summary>
+            
+            <returns>
+            The state changing filters.
+            </returns>
+        </member>
+        <member name="P:Hangfire.Common.JobFilterInfo.ApplyStateFilters">
+            <summary>
+            Gets all the state changed filters in the application.
+            </summary>
+            
+            <returns>
+            The state changed filters.
+            </returns>
+        </member>
+        <member name="P:Hangfire.Common.JobFilterInfo.ClientExceptionFilters">
+            <summary>
+            Gets all the client exception filters in the application.
+            </summary>
+            
+            <returns>
+            The client exception filters.
+            </returns>
+        </member>
+        <member name="P:Hangfire.Common.JobFilterInfo.ServerExceptionFilters">
+            <summary>
+            Gets all the server exception filters in the application.
+            </summary>
+            
+            <returns>
+            The server exception filters.
+            </returns>
+        </member>
+        <member name="T:Hangfire.Server.DelayedJobScheduler">
+            <summary>
+            Represents a background process responsible for <i>enqueueing delayed
+            jobs</i>.
+            </summary>
+            
+            <remarks>
+            <para>This background process polls the <i>delayed job schedule</i> for 
+            delayed jobs that are ready to be enqueued. To prevent a stress load
+            on a job storage, the configurable delay is used between scheduler 
+            runs. Delay is used only when there are no more background jobs to be
+            enqueued.</para>
+            
+            <para>When a background job is ready to be enqueued, it is simply
+            moved from <see cref="T:Hangfire.States.ScheduledState"/> to the <see cref="T:Hangfire.States.EnqueuedState"/>
+            by using <see cref="T:Hangfire.States.IBackgroundJobStateChanger"/>.</para>
+            
+            <para>Delayed job schedule is based on a Set data structure of a job
+            storage, so you can use this background process as an example of a
+            custom extension.</para>
+             
+            <para>Multiple instances of this background process can be used in
+            separate threads/processes without additional configuration (distributed
+            locks are used). However, this only adds support for fail-over, and does 
+            not increase the performance.</para>
+            
+            <note type="important">
+            If you are using <b>custom filter providers</b>, you need to pass a custom
+            <see cref="T:Hangfire.States.IBackgroundJobStateChanger"/> instance to make this process
+            respect your filters when enqueueing background jobs.
+            </note>
+            </remarks>
+            
+            <threadsafety static="true" instance="true"/>
+            
+            <seealso cref="T:Hangfire.States.ScheduledState"/>
+        </member>
+        <member name="F:Hangfire.Server.DelayedJobScheduler.DefaultPollingDelay">
+            <summary>
+            Represents a default polling interval for delayed job scheduler. 
+            This field is read-only.
+            </summary>
+            <remarks>
+            The value of this field is <c>TimeSpan.FromSeconds(15)</c>.
+            </remarks>
+        </member>
+        <member name="M:Hangfire.Server.DelayedJobScheduler.#ctor">
+            <summary>
+            Initializes a new instance of the <see cref="T:Hangfire.Server.DelayedJobScheduler"/>
+            class with the <see cref="F:Hangfire.Server.DelayedJobScheduler.DefaultPollingDelay"/> value as a
+            delay between runs.
+            </summary>
+        </member>
+        <member name="M:Hangfire.Server.DelayedJobScheduler.#ctor(System.TimeSpan)">
+            <summary>
+            Initializes a new instance of the <see cref="T:Hangfire.Server.DelayedJobScheduler"/>
+            class with a specified polling interval.
+            </summary>
+            <param name="pollingDelay">Delay between scheduler runs.</param>
+        </member>
+        <member name="M:Hangfire.Server.DelayedJobScheduler.#ctor(System.TimeSpan,Hangfire.States.IBackgroundJobStateChanger)">
+            <summary>
+            Initializes a new instance of the <see cref="T:Hangfire.Server.DelayedJobScheduler"/>
+            class with a specified polling interval and given state changer.
+            </summary>
+            <param name="pollingDelay">Delay between scheduler runs.</param>
+            <param name="stateChanger">State changer to use for background jobs.</param>
+            
+            <exception cref="T:System.ArgumentNullException"><paramref name="stateChanger"/> is null.</exception>
+        </member>
+        <member name="M:Hangfire.Server.DelayedJobScheduler.Execute(Hangfire.Server.BackgroundProcessContext)">
+            <inheritdoc />
+        </member>
+        <member name="M:Hangfire.Server.DelayedJobScheduler.ToString">
+            <inheritdoc />
+        </member>
+        <member name="T:Hangfire.Server.ServerWatchdogOptions">
+            <exclude />
+        </member>
+        <member name="T:Hangfire.Server.IServerExceptionFilter">
+            <summary>
+            Defines methods that are required for the server exception filter.
+            </summary>
+        </member>
+        <member name="M:Hangfire.Server.IServerExceptionFilter.OnServerException(Hangfire.Server.ServerExceptionContext)">
+            <summary>
+            Called when an exception occurred during the performance of the job.
+            </summary>
+            <param name="filterContext">The filter context.</param>
+        </member>
+        <member name="T:Hangfire.Client.CreatedContext">
+            <summary>
+            Provides the context for the <see cref="M:Hangfire.Client.IClientFilter.OnCreated(Hangfire.Client.CreatedContext)"/> 
+            method of the <see cref="T:Hangfire.Client.IClientFilter"/> interface.
+            </summary>
+        </member>
+        <member name="P:Hangfire.Client.CreatedContext.Exception">
+            <summary>
+            Gets an exception that occurred during the creation of the job.
+            </summary>
+        </member>
+        <member name="P:Hangfire.Client.CreatedContext.Canceled">
+            <summary>
+            Gets a value that indicates that this <see cref="T:Hangfire.Client.CreatedContext"/>
+            object was canceled.
+            </summary>
+        </member>
+        <member name="P:Hangfire.Client.CreatedContext.ExceptionHandled">
+            <summary>
+            Gets or sets a value that indicates that this <see cref="T:Hangfire.Client.CreatedContext"/>
+            object handles an exception occurred during the creation of the job.
+            </summary>
+        </member>
+        <member name="T:Hangfire.Client.CreatingContext">
+            <summary>
+            Provides the context for the <see cref="M:Hangfire.Client.IClientFilter.OnCreating(Hangfire.Client.CreatingContext)"/>
+            method of the <see cref="T:Hangfire.Client.IClientFilter"/> interface.
+            </summary>
+        </member>
+        <member name="M:Hangfire.Client.CreatingContext.SetJobParameter(System.String,System.Object)">
+            <summary>
+            Sets the job parameter of the specified <paramref name="name"/>
+            to the corresponding <paramref name="value"/>. The value of the
+            parameter is serialized to a JSON string.
+            </summary>
+            
+            <param name="name">The name of the parameter.</param>
+            <param name="value">The value of the parameter.</param>
+            
+            <exception cref="T:System.ArgumentNullException">The <paramref name="name"/> is null or empty.</exception>
+        </member>
+        <member name="M:Hangfire.Client.CreatingContext.GetJobParameter``1(System.String)">
+            <summary>
+            Gets the job parameter of the specified <paramref name="name"/>
+            if it exists. The parameter is deserialized from a JSON 
+            string value to the given type <typeparamref name="T"/>.
+            </summary>
+            
+            <typeparam name="T">The type of the parameter.</typeparam>
+            <param name="name">The name of the parameter.</param>
+            <returns>The value of the given parameter if it exists or null otherwise.</returns>
+            
+            <exception cref="T:System.ArgumentNullException">The <paramref name="name"/> is null or empty.</exception>
+            <exception cref="T:System.InvalidOperationException">Could not deserialize the parameter value to the type <typeparamref name="T"/>.</exception>
+        </member>
+        <member name="P:Hangfire.Client.CreatingContext.Canceled">
+            <summary>
+            Gets or sets a value that indicates that this <see cref="T:Hangfire.Client.CreatingContext"/>
+            object was canceled.
+            </summary>
+        </member>
+        <member name="T:Hangfire.Server.IServerComponent">
+            <exclude />
+        </member>
+        <member name="T:Hangfire.Server.PerformContext">
+            <summary>
+            Provides information about the context in which the job
+            is performed.
+            </summary>
+        </member>
+        <member name="P:Hangfire.Server.PerformContext.Items">
+            <summary>
+            Gets an instance of the key-value storage. You can use it
+            to pass additional information between different client filters
+            or just between different methods.
+            </summary>
+        </member>
+        <member name="T:Hangfire.Server.ServerExceptionContext">
+            <summary>
+            Provides the context for the <see cref="M:Hangfire.Server.IServerExceptionFilter.OnServerException(Hangfire.Server.ServerExceptionContext)"/>
+            method of the <see cref="T:Hangfire.Server.IServerExceptionFilter"/> interface.
+            </summary>
+        </member>
+        <member name="P:Hangfire.Server.ServerExceptionContext.Exception">
+            <summary>
+            Gets an exception that occurred during the performance of the job.
+            </summary>
+        </member>
+        <member name="P:Hangfire.Server.ServerExceptionContext.ExceptionHandled">
+            <summary>
+            Gets or sets a value that indicates that this <see cref="T:Hangfire.Server.ServerExceptionContext"/>
+            object handles an exception occurred during the performance of the job.
+            </summary>
+        </member>
+        <member name="T:Hangfire.Server.PerformedContext">
+            <summary>
+            Provides the context for the <see cref="M:Hangfire.Server.IServerFilter.OnPerformed(Hangfire.Server.PerformedContext)"/>
+            method of the <see cref="T:Hangfire.Server.IServerFilter"/> interface.
+            </summary>
+        </member>
+        <member name="P:Hangfire.Server.PerformedContext.Result">
+            <summary>
+            Gets a value that was returned by the job.
+            </summary>
+        </member>
+        <member name="P:Hangfire.Server.PerformedContext.Canceled">
+            <summary>
+            Gets a value that indicates that this <see cref="T:Hangfire.Server.PerformedContext"/>
+            object was canceled.
+            </summary>
+        </member>
+        <member name="P:Hangfire.Server.PerformedContext.Exception">
+            <summary>
+            Gets an exception that occurred during the performance of the job.
+            </summary>
+        </member>
+        <member name="P:Hangfire.Server.PerformedContext.ExceptionHandled">
+            <summary>
+            Gets or sets a value that indicates that this <see cref="T:Hangfire.Server.PerformedContext"/>
+            object handles an exception occurred during the performance of the job.
+            </summary>
+        </member>
+        <member name="T:Hangfire.Server.PerformingContext">
+            <summary>
+            Provides the context for the <see cref="M:Hangfire.Server.IServerFilter.OnPerforming(Hangfire.Server.PerformingContext)"/>
+            method of the <see cref="T:Hangfire.Server.IServerFilter"/> interface.
+            </summary>
+        </member>
+        <member name="P:Hangfire.Server.PerformingContext.Canceled">
+            <summary>
+            Gets or sets a value that indicates that this <see cref="T:Hangfire.Server.PerformingContext"/>
+            object was canceled.
+            </summary>
+        </member>
+        <member name="T:Hangfire.GlobalJobFilters">
+            <summary>
+            Represents the global filter collection.
+            </summary>
+        </member>
+        <member name="P:Hangfire.GlobalJobFilters.Filters">
+            <summary>
+            Gets the global filter collection.
+            </summary>
+        </member>
+        <member name="T:Hangfire.Server.Worker">
+            <summary>
+            Represents a background process responsible for <i>processing 
+            fire-and-forget jobs</i>.
+            </summary>
+            
+            <remarks>
+            <para>This is the heart of background processing in Hangfire</para>
+            </remarks>
+            
+            <threadsafety static="true" instance="true"/>
+            
+            <seealso cref="T:Hangfire.States.EnqueuedState"/>
+        </member>
+        <member name="M:Hangfire.Server.Worker.Execute(Hangfire.Server.BackgroundProcessContext)">
+            <inheritdoc />
+        </member>
+        <member name="M:Hangfire.Server.Worker.ToString">
+            <inheritdoc />
+        </member>
+        <member name="M:Hangfire.States.IBackgroundJobStateChanger.ChangeState(Hangfire.States.StateChangeContext)">
+            <summary>
+            Attempts to change the state of a job, respecting any applicable job filters and state handlers.
+            </summary>
+            <returns><c>Null</c> if a constraint has failed, otherwise the final applied state</returns>
+            <remarks>Also ensures that the job data can be loaded for this job</remarks>
+        </member>
+        <member name="T:Hangfire.States.EnqueuedState">
+            <summary>
+            Defines the <i>intermediate</i> state of a background job when it is placed 
+            on a message queue to be processed by the <see cref="T:Hangfire.Server.Worker"/> 
+            background process <i>as soon as possible</i>.
+            </summary>
+            <remarks>
+            <para>Background job in <see cref="T:Hangfire.States.EnqueuedState"/> is referred as
+            <b>fire-and-forget job</b>.</para>
+            <para>Background job identifier is placed on a queue with the given name. When
+            a queue name wasn't specified, the <see cref="F:Hangfire.States.EnqueuedState.DefaultQueue"/> name will
+            be used. Message queue implementation depends on a current <see cref="T:Hangfire.JobStorage"/>
+            instance.</para>
+            </remarks> 
+            <example>
+            The following example demonstrates the creation of a background job in
+            <see cref="T:Hangfire.States.EnqueuedState"/>. Please see 
+            <see cref="O:Hangfire.BackgroundJob.Enqueue">BackgroundJob.Enqueue</see>
+            and <see cref="O:Hangfire.BackgroundJobClientExtensions.Enqueue">BackgroundJobClientExtensions.Enqueue</see>
+            method overloads for simpler API.
+            
+            <code lang="cs" source="..\Samples\States.cs" region="EnqueuedState #1"/>
+            <code lang="vb" source="..\VBSamples\States.vb" region="EnqueuedState #1"/>
+            
+            The code below implements the retry action for a failed background job.
+            
+            <code lang="cs" source="..\Samples\States.cs" region="EnqueuedState #2"/>
+            <code lang="vb" source="..\VBSamples\States.vb" region="EnqueuedState #2"/>
+             
+            </example>
+            
+            <seealso cref="O:Hangfire.BackgroundJob.Enqueue">BackgroundJob.Enqueue Overload</seealso>
+            <seealso cref="O:Hangfire.BackgroundJobClientExtensions.Enqueue">BackgroundJobClientExtensions.Enqueue Overload</seealso>
+            <seealso cref="O:Hangfire.BackgroundJobClientExtensions.Create">BackgroundJobClientExtensions.Create Overload</seealso>
+            <seealso cref="M:Hangfire.IBackgroundJobClient.Create(Hangfire.Common.Job,Hangfire.States.IState)"/>
+            <seealso cref="M:Hangfire.IBackgroundJobClient.ChangeState(System.String,Hangfire.States.IState,System.String)"/>
+            <seealso cref="T:Hangfire.Server.Worker"/>
+            
+            <threadsafety static="true" instance="false"/>
+        </member>
+        <member name="F:Hangfire.States.EnqueuedState.DefaultQueue">
+            <summary>
+            Represents the default queue name. This field is constant.
+            </summary>
+            <remarks>
+            The value of this field is <c>"default"</c>.
+            </remarks>
+        </member>
+        <member name="F:Hangfire.States.EnqueuedState.StateName">
+            <summary>
+            Represents the name of the <i>Enqueued</i> state. This field is read-only.
+            </summary>
+            <remarks>
+            The value of this field is <c>"Enqueued"</c>.
+            </remarks>
+        </member>
+        <member name="M:Hangfire.States.EnqueuedState.#ctor">
+            <summary>
+            Initializes a new instance of the <see cref="T:Hangfire.States.EnqueuedState"/> class 
+            with the <see cref="F:Hangfire.States.EnqueuedState.DefaultQueue">default</see> queue name.
+            </summary>
+        </member>
+        <member name="M:Hangfire.States.EnqueuedState.#ctor(System.String)">
+            <summary>
+            Initializes a new instance of the <see cref="T:Hangfire.States.EnqueuedState"/> class
+            with the specified queue name.
+            </summary>
+            <param name="queue">The queue name to which a background job identifier will be added.</param>
+            
+            <seealso cref="P:Hangfire.States.EnqueuedState.Queue"/>
+            
+            <exception cref="T:System.ArgumentNullException">
+            The <paramref name="queue"/> argument is <see langword="null"/>,  empty or consist only of 
+            white-space characters.
+            </exception>
+            <exception cref="T:System.ArgumentException">
+            The <paramref name="queue"/> argument is not a valid queue name.
+            </exception>
+        </member>
+        <member name="M:Hangfire.States.EnqueuedState.SerializeData">
+            <inheritdoc/>
+            <remarks>
+            <para>Returning dictionary contains the following keys. You can obtain 
+            the state data by using the <see cref="M:Hangfire.Storage.IStorageConnection.GetStateData(System.String)"/>
+            method.</para>
+            <list type="table">
+                <listheader>
+                    <term>Key</term>
+                    <term>Type</term>
+                    <term>Deserialize Method</term>
+                    <description>Notes</description>
+                </listheader>
+                <item>
+                    <term><c>EnqueuedAt</c></term>
+                    <term><see cref="T:System.DateTime"/></term>
+                    <term><see cref="M:Hangfire.Common.JobHelper.DeserializeDateTime(System.String)"/></term>
+                    <description>Please see the <see cref="P:Hangfire.States.EnqueuedState.EnqueuedAt"/> property.</description>
+                </item>
+                <item>
+                    <term><c>Queue</c></term>
+                    <term><see cref="T:System.String"/></term>
+                    <term><i>Not required</i></term>
+                    <description>Please see the <see cref="P:Hangfire.States.EnqueuedState.Queue"/> property.</description>
+                </item>
+            </list>
+            </remarks>
+        </member>
+        <member name="P:Hangfire.States.EnqueuedState.Queue">
+            <summary>
+            Gets or sets a queue name to which a background job identifier
+            will be added.
+            </summary>
+            <value>A queue name that consist only of lowercase letters, digits and
+            underscores.</value>
+            <remarks>
+            <para>Queue name must consist only of lowercase letters, digits and
+            underscores, other characters aren't permitted. Some examples:</para>
+            <list type="bullet">
+                <item><c>"critical"</c> (good)</item>
+                <item><c>"worker_1"</c> (good)</item>
+                <item><c>"documents queue"</c> (bad, whitespace)</item>
+                <item><c>"MyQueue"</c> (bad, capital letters)</item>
+            </list>
+            </remarks>
+            
+            <exception cref="T:System.ArgumentNullException">
+            The value specified for a set operation is <see langword="null"/>, 
+            empty or consist only of white-space characters.
+            </exception>
+            <exception cref="T:System.ArgumentException">
+            The value specified for a set operation is not a valid queue name.
+            </exception>
+        </member>
+        <member name="P:Hangfire.States.EnqueuedState.EnqueuedAt">
+            <summary>
+            Gets a date/time when the current state instance was created.
+            </summary>
+        </member>
+        <member name="P:Hangfire.States.EnqueuedState.Name">
+            <inheritdoc/>
+            <remarks>
+            Always equals to <see cref="F:Hangfire.States.EnqueuedState.StateName"/> for the <see cref="T:Hangfire.States.EnqueuedState"/>.
+            Please see the remarks section of the <see cref="P:Hangfire.States.IState.Name">IState.Name</see>
+            article for the details.
+            </remarks>
+        </member>
+        <member name="P:Hangfire.States.EnqueuedState.Reason">
+            <inheritdoc />
+        </member>
+        <member name="P:Hangfire.States.EnqueuedState.IsFinal">
+            <inheritdoc/>
+            <remarks>
+            Always returns <see langword="false"/> for the <see cref="T:Hangfire.States.EnqueuedState"/>.
+            Please refer to the <see cref="P:Hangfire.States.IState.IsFinal">IState.IsFinal</see> documentation
+            for the details.
+            </remarks>
+        </member>
+        <member name="P:Hangfire.States.EnqueuedState.IgnoreJobLoadException">
+            <inheritdoc/>
+            <remarks>
+            Always returns <see langword="false"/> for the <see cref="T:Hangfire.States.EnqueuedState"/>.
+            Please see the description of this property in the
+            <see cref="P:Hangfire.States.IState.IgnoreJobLoadException">IState.IgnoreJobLoadException</see>
+            article.
+            </remarks>
+        </member>
+        <member name="T:Hangfire.States.FailedState">
+            <summary>
+            Defines the <i>intermediate</i> state of a background job when its processing 
+            was interrupted by an exception and it is a developer's responsibility
+            to decide what to do with it next.
+            </summary>
+            <remarks>
+            <para>Failed state is used in Hangfire when something went wrong and an exception
+            occurred during the background job processing. The primary reason for this state
+            is to notify the developers that something went wrong. By default background job 
+            is moved to the <i>Failed</i> state only after some automatic retries, because the 
+            <see cref="T:Hangfire.AutomaticRetryAttribute"/> filter is enabled by default.</para>
+            <note type="important">
+            Failed jobs are <b>not expiring</b> and will stay in your current job storage 
+            forever, increasing its size until you retry or delete them manually. If you 
+            expect some exceptions, please use the following rules.
+            <list type="bullet">
+                <item>Ignore, move to <i>Succeeded</i> state – use the <c>catch</c>
+                statement in your code without re-throwing the exception.</item>
+                <item>Ignore, move to <i>Deleted</i> state – use the <see cref="T:Hangfire.AutomaticRetryAttribute"/>
+                with <see cref="F:Hangfire.AttemptsExceededAction.Delete"/> option.</item>
+                <item>Re-queue a job – use the <see cref="T:Hangfire.AutomaticRetryAttribute"/> with
+                <see cref="F:Hangfire.AttemptsExceededAction.Fail"/> option.</item>
+            </list>
+            </note>
+            <para>It is not supposed to use the <see cref="T:Hangfire.States.FailedState"/> class in a user
+            code unless you are writing state changing filters or new background processing
+            rules.</para>
+            </remarks>
+            
+            <seealso cref="T:Hangfire.AutomaticRetryAttribute"/>
+            <seealso cref="T:Hangfire.States.IBackgroundJobStateChanger"/>
+            <seealso cref="T:Hangfire.Server.Worker"/>
+            
+            <threadsafety static="true" instance="false"/>
+        </member>
+        <member name="F:Hangfire.States.FailedState.StateName">
+            <summary>
+            Represents the name of the <i>Failed</i> state. This field is read-only.
+            </summary>
+            <remarks>
+            The value of this field is <c>"Failed"</c>.
+            </remarks>
+        </member>
+        <member name="M:Hangfire.States.FailedState.#ctor(System.Exception)">
+            <summary>
+            Initializes a new instace of the <see cref="T:Hangfire.States.FailedState"/> class
+            with the given exception.
+            </summary>
+            <param name="exception">Exception that occurred during the background 
+            job processing.</param>
+            
+            <exception cref="T:System.ArgumentNullException">The <paramref name="exception"/> 
+            argument is <see langword="null"/></exception>
+        </member>
+        <member name="M:Hangfire.States.FailedState.SerializeData">
+            <inheritdoc/>
+            <remarks>
+            <para>Returning dictionary contains the following keys. You can obtain 
+            the state data by using the <see cref="M:Hangfire.Storage.IStorageConnection.GetStateData(System.String)"/>
+            method.</para>
+            <list type="table">
+                <listheader>
+                    <term>Key</term>
+                    <term>Type</term>
+                    <term>Deserialize Method</term>
+                    <description>Notes</description>
+                </listheader>
+                <item>
+                    <term><c>FailedAt</c></term>
+                    <term><see cref="T:System.DateTime"/></term>
+                    <term><see cref="M:Hangfire.Common.JobHelper.DeserializeDateTime(System.String)"/></term>
+                    <description>Please see the <see cref="P:Hangfire.States.FailedState.FailedAt"/> property.</description>
+                </item>
+                <item>
+                    <term><c>ExceptionType</c></term>
+                    <term><see cref="T:System.String"/></term>
+                    <term><i>Not required</i></term>
+                    <description>The full name of the current exception type.</description>
+                </item>
+                <item>
+                    <term><c>ExceptionMessage</c></term>
+                    <term><see cref="T:System.String"/></term>
+                    <term><i>Not required</i></term>
+                    <description>Message that describes the current exception.</description>
+                </item>
+                <item>
+                    <term><c>ExceptionDetails</c></term>
+                    <term><see cref="T:System.String"/></term>
+                    <term><i>Not required</i></term>
+                    <description>String representation of the current exception.</description>
+                </item>
+            </list>
+            </remarks>
+        </member>
+        <member name="P:Hangfire.States.FailedState.FailedAt">
+            <summary>
+            Gets a date/time when the current state instance was created.
+            </summary>
+        </member>
+        <member name="P:Hangfire.States.FailedState.Exception">
+            <summary>
+            Gets the exception that occurred during the background job processing.
+            </summary>
+        </member>
+        <member name="P:Hangfire.States.FailedState.Name">
+            <inheritdoc/>
+            <remarks>
+            Always equals to <see cref="F:Hangfire.States.FailedState.StateName"/> for the <see cref="T:Hangfire.States.FailedState"/>.
+            Please see the remarks section of the <see cref="P:Hangfire.States.IState.Name">IState.Name</see>
+            article for the details.
+            </remarks>
+        </member>
+        <member name="P:Hangfire.States.FailedState.Reason">
+            <inheritdoc />
+        </member>
+        <member name="P:Hangfire.States.FailedState.IsFinal">
+            <inheritdoc/>
+            <remarks>
+            Always returns <see langword="false"/> for the <see cref="T:Hangfire.States.FailedState"/>.
+            Please refer to the <see cref="P:Hangfire.States.IState.IsFinal">IState.IsFinal</see> documentation
+            for the details.
+            </remarks>
+        </member>
+        <member name="P:Hangfire.States.FailedState.IgnoreJobLoadException">
+            <inheritdoc/>
+            <remarks>
+            Always returns <see langword="false"/> for the <see cref="T:Hangfire.States.FailedState"/>.
+            Please see the description of this property in the
+            <see cref="P:Hangfire.States.IState.IgnoreJobLoadException">IState.IgnoreJobLoadException</see>
+            article.
+            </remarks>
+        </member>
+        <member name="T:Hangfire.Client.IClientFilter">
+            <summary>
+            Defines methods that are required for a client filter.
+            </summary>
+        </member>
+        <member name="M:Hangfire.Client.IClientFilter.OnCreating(Hangfire.Client.CreatingContext)">
+            <summary>
+            Called before the creation of the job. 
+            </summary>
+            <param name="filterContext">The filter context.</param>
+        </member>
+        <member name="M:Hangfire.Client.IClientFilter.OnCreated(Hangfire.Client.CreatedContext)">
+            <summary>
+            Called after the creation of the job.
+            </summary>
+            <param name="filterContext">The filter context.</param>
+        </member>
+        <member name="T:Hangfire.BackgroundJob">
+            <summary>
+            Provides static methods for creating <i>fire-and-forget</i>, <i>delayed</i>
+            jobs and <i>continuations</i> as well as re-queue and delete existing
+            background jobs.
+            </summary>
+            
+            <remarks>
+            <para>This class is a wrapper for the <see cref="T:Hangfire.IBackgroundJobClient"/> 
+            interface and its default implementation, <see cref="T:Hangfire.BackgroundJobClient"/>
+            class, that was created for the most simple scenarios. Please consider 
+            using the types above in real world applications.</para>
+            <para>This class also contains undocumented constructor and instance 
+            members. They are hidden to not to confuse new users. You can freely 
+            use them in low-level API.</para>
+            </remarks>
+            
+            <seealso cref="T:Hangfire.IBackgroundJobClient"/>
+            <seealso cref="T:Hangfire.BackgroundJobClient"/>
+            
+            <threadsafety static="true" instance="false"/>
+        </member>
+        <member name="M:Hangfire.BackgroundJob.Enqueue(System.Linq.Expressions.Expression{System.Action})">
+            <summary>
+            Creates a new fire-and-forget job based on a given method call expression.
+            </summary>
+            <param name="methodCall">Method call expression that will be marshalled to a server.</param>
+            <returns>Unique identifier of a background job.</returns>
+            
+            <exception cref="T:System.ArgumentNullException">
+            <paramref name="methodCall"/> is <see langword="null"/>.
+            </exception>
+            
+            <seealso cref="T:Hangfire.States.EnqueuedState"/>
+            <seealso cref="O:Hangfire.IBackgroundJobClient.Enqueue"/>
+        </member>
+        <member name="M:Hangfire.BackgroundJob.Enqueue``1(System.Linq.Expressions.Expression{System.Action{``0}})">
+            <summary>
+            Creates a new fire-and-forget job based on a given method call expression.
+            </summary>
+            <param name="methodCall">Method call expression that will be marshalled to a server.</param>
+            <returns>Unique identifier of a background job.</returns>
+            
+            <exception cref="T:System.ArgumentNullException">
+            <paramref name="methodCall"/> is <see langword="null"/>.
+            </exception>
+            
+            <seealso cref="T:Hangfire.States.EnqueuedState"/>
+            <seealso cref="O:Hangfire.IBackgroundJobClient.Enqueue"/>
+        </member>
+        <member name="M:Hangfire.BackgroundJob.Schedule(System.Linq.Expressions.Expression{System.Action},System.TimeSpan)">
+            <summary>
+            Creates a new background job based on a specified static method
+            call expression and schedules it to be enqueued after a given delay.
+            </summary>
+            
+            <param name="methodCall">Instance method call expression that will be marshalled to the Server.</param>
+            <param name="delay">Delay, after which the job will be enqueued.</param>
+            <returns>Unique identifier of the created job.</returns>
+        </member>
+        <member name="M:Hangfire.BackgroundJob.Schedule(System.Linq.Expressions.Expression{System.Action},System.DateTimeOffset)">
+            <summary>
+            Creates a new background job based on a specified method call expression
+            and schedules it to be enqueued at the given moment of time.
+            </summary>
+            
+            <param name="methodCall">Method call expression that will be marshalled to the Server.</param>
+            <param name="enqueueAt">The moment of time at which the job will be enqueued.</param>
+            <returns>Unique identifier of a created job.</returns>
+        </member>
+        <member name="M:Hangfire.BackgroundJob.Schedule``1(System.Linq.Expressions.Expression{System.Action{``0}},System.TimeSpan)">
+            <summary>
+            Creates a new background job based on a specified instance method
+            call expression and schedules it to be enqueued after a given delay.
+            </summary>
+            
+            <typeparam name="T">Type whose method will be invoked during job processing.</typeparam>
+            <param name="methodCall">Instance method call expression that will be marshalled to the Server.</param>
+            <param name="delay">Delay, after which the job will be enqueued.</param>
+            <returns>Unique identifier of the created job.</returns>
+        </member>
+        <member name="M:Hangfire.BackgroundJob.Schedule``1(System.Linq.Expressions.Expression{System.Action{``0}},System.DateTimeOffset)">
+            <summary>
+            Creates a new background job based on a specified method call expression
+            and schedules it to be enqueued at the given moment of time.
+            </summary>
+            
+            <typeparam name="T">The type whose method will be invoked during the job processing.</typeparam>
+            <param name="methodCall">Method call expression that will be marshalled to the Server.</param>
+            <param name="enqueueAt">The moment of time at which the job will be enqueued.</param>
+            <returns>Unique identifier of a created job.</returns>
+        </member>
+        <member name="M:Hangfire.BackgroundJob.Delete(System.String)">
+            <summary>
+            Changes state of a job with the specified <paramref name="jobId"/>
+            to the <see cref="T:Hangfire.States.DeletedState"/>. 
+            <seealso cref="M:Hangfire.BackgroundJobClientExtensions.Delete(Hangfire.IBackgroundJobClient,System.String)"/>
+            </summary>
+            
+            <param name="jobId">An identifier, that will be used to find a job.</param>
+            <returns>True on a successfull state transition, false otherwise.</returns>
+        </member>
+        <member name="M:Hangfire.BackgroundJob.Delete(System.String,System.String)">
+            <summary>
+            Changes state of a job with the specified <paramref name="jobId"/>
+            to the <see cref="T:Hangfire.States.DeletedState"/>. State change is only performed
+            if current job state is equal to the <paramref name="fromState"/> value.
+            <seealso cref="M:Hangfire.BackgroundJobClientExtensions.Delete(Hangfire.IBackgroundJobClient,System.String,System.String)"/>
+            </summary>
+            
+            <param name="jobId">Identifier of job, whose state is being changed.</param>
+            <param name="fromState">Current state assertion, or null if unneeded.</param>
+            <returns>True, if state change succeeded, otherwise false.</returns>
+        </member>
+        <member name="M:Hangfire.BackgroundJob.Requeue(System.String)">
+            <summary>
+            Changes state of a job with the specified <paramref name="jobId"/>
+            to the <see cref="T:Hangfire.States.EnqueuedState"/>.
+            </summary>
+            
+            <param name="jobId">Identifier of job, whose state is being changed.</param>
+            <returns>True, if state change succeeded, otherwise false.</returns>
+        </member>
+        <member name="M:Hangfire.BackgroundJob.Requeue(System.String,System.String)">
+            <summary>
+            Changes state of a job with the specified <paramref name="jobId"/>
+            to the <see cref="T:Hangfire.States.EnqueuedState"/>. If <paramref name="fromState"/> value 
+            is not null, state change will be performed only if the current state name 
+            of a job equal to the given value.
+            </summary>
+            
+            <param name="jobId">Identifier of job, whose state is being changed.</param>
+            <param name="fromState">Current state assertion, or null if unneeded.</param>
+            <returns>True, if state change succeeded, otherwise false.</returns>
+        </member>
+        <member name="M:Hangfire.BackgroundJob.#ctor(System.String,Hangfire.Common.Job,System.DateTime)">
+            <exclude />
+        </member>
+        <member name="P:Hangfire.BackgroundJob.Id">
+            <exclude />
+        </member>
+        <member name="P:Hangfire.BackgroundJob.Job">
+            <exclude />
+        </member>
+        <member name="P:Hangfire.BackgroundJob.CreatedAt">
+            <exclude />
+        </member>
+        <member name="T:Hangfire.QueueAttribute">
+            <summary>
+            Represents attribute, that is used to determine queue name
+            for background jobs. It can be applied to the methods and classes. 
+            If the attribute is not applied neither to the method, nor the class, 
+            then default queue will be used.
+            </summary>
+            
+            <example><![CDATA[
+            
+            [Queue("high")]
+            public class ErrorService
+            {
+                public void ReportError(string message) { }
+            
+                [Queue("critical")]
+                public void ReportFatal(string message) { }
+            }
+            
+            // Background job will be placed on the 'high' queue.
+            BackgroundJob.Enqueue<ErrorService>(x => x.ReportError("Something bad happened"));
+            
+            // Background job will be placed on the 'critical' queue.
+            BackgroundJob.Enqueue<ErrorService>(x => x.ReportFatal("Really bad thing!"));
+            
+            ]]></example>
+        </member>
+        <member name="M:Hangfire.QueueAttribute.#ctor(System.String)">
+            <summary>
+            Initializes a new instance of the <see cref="T:Hangfire.QueueAttribute"/> class
+            using the specified queue name.
+            </summary>
+            <param name="queue">Queue name.</param>
+        </member>
+        <member name="P:Hangfire.QueueAttribute.Queue">
+            <summary>
+            Gets the queue name that will be used for background jobs.
+            </summary>
+        </member>
+        <member name="P:Hangfire.JobActivator.Current">
+            <summary>
+            Gets or sets the current <see cref="T:Hangfire.JobActivator"/> instance 
+            that will be used to activate jobs during performance.
+            </summary>
+        </member>
+        <member name="T:Hangfire.BackgroundJobClient">
+            <summary>
+            Provides methods for creating all the types of background jobs and 
+            changing their states. Represents a default implementation of the 
+            <see cref="T:Hangfire.IBackgroundJobClient"/> interface.
+            </summary>
+            
+            <remarks>
+            <para>This class uses the <see cref="T:Hangfire.Client.IBackgroundJobFactory"/> interface 
+            for creating background jobs and the <see cref="T:Hangfire.States.IBackgroundJobStateChanger"/> 
+            interface for changing their states. Please see documentation for those 
+            types and their implementations to learn the details.</para>
+            
+            <note type="warning">
+            Despite the fact that instance methods of this class are thread-safe,
+            most implementations of the <see cref="T:Hangfire.States.IState"/> interface are <b>neither
+            thread-safe, nor immutable</b>. Please create a new instance of a state 
+            class for each operation to avoid race conditions and unexpected side 
+            effects.
+            </note>
+            </remarks>
+            
+            <threadsafety static="true" instance="true"/>
+        </member>
+        <member name="M:Hangfire.BackgroundJobClient.#ctor">
+            <summary>
+            Initializes a new instance of the <see cref="T:Hangfire.BackgroundJobClient"/>
+            class with the storage from a global configuration.
+            </summary>
+            
+            <remarks>
+            Please see the <see cref="T:Hangfire.GlobalConfiguration"/> class for the
+            details regarding the global configuration.
+            </remarks>
+        </member>
+        <member name="M:Hangfire.BackgroundJobClient.#ctor(Hangfire.JobStorage)">
+            <summary>
+            Initializes a new instance of the <see cref="T:Hangfire.BackgroundJobClient"/>
+            class with the specified storage.
+            </summary>
+            
+            <param name="storage">Job storage to use for background jobs.</param>
+            
+            <exception cref="T:System.ArgumentNullException"><paramref name="storage"/> is null.</exception>
+        </member>
+        <member name="M:Hangfire.BackgroundJobClient.#ctor(Hangfire.JobStorage,Hangfire.Client.IBackgroundJobFactory,Hangfire.States.IBackgroundJobStateChanger)">
+            <summary>
+            Initializes a new instance of the <see cref="T:Hangfire.BackgroundJobClient"/> class
+            with the specified storage, background job factory and state changer.
+            </summary>
+            
+            <param name="storage">Job storage to use for background jobs.</param>
+            <param name="factory">Factory to create background jobs.</param>
+            <param name="stateChanger">State changer to change states of background jobs.</param>
+            
+            <exception cref="T:System.ArgumentNullException"><paramref name="storage"/> is null.</exception>
+            <exception cref="T:System.ArgumentNullException"><paramref name="factory"/> is null.</exception>
+            <exception cref="T:System.ArgumentNullException"><paramref name="stateChanger"/> is null.</exception>
+        </member>
+        <member name="M:Hangfire.BackgroundJobClient.Create(Hangfire.Common.Job,Hangfire.States.IState)">
+            <inheritdoc />
+        </member>
+        <member name="M:Hangfire.BackgroundJobClient.ChangeState(System.String,Hangfire.States.IState,System.String)">
+            <inheritdoc />
+        </member>
+        <member name="T:Hangfire.States.ProcessingState">
+            <summary>
+            Defines the <i>intermediate</i> state of a background job when a 
+            <see cref="T:Hangfire.Server.Worker"/> has started to process it.
+            </summary>
+            
+            <threadsafety static="true" instance="false"/>
+        </member>
+        <member name="F:Hangfire.States.ProcessingState.StateName">
+            <summary>
+            Represents the name of the <i>Processing</i> state. This field is read-only.
+            </summary>
+            <remarks>
+            The value of this field is <c>"Processing"</c>.
+            </remarks>
+        </member>
+        <member name="M:Hangfire.States.ProcessingState.SerializeData">
+            <inheritdoc/>
+            <remarks>
+            <para>Returning dictionary contains the following keys. You can obtain 
+            the state data by using the <see cref="M:Hangfire.Storage.IStorageConnection.GetStateData(System.String)"/>
+            method.</para>
+            <list type="table">
+                <listheader>
+                    <term>Key</term>
+                    <term>Type</term>
+                    <term>Deserialize Method</term>
+                    <description>Notes</description>
+                </listheader>
+                <item>
+                    <term><c>StartedAt</c></term>
+                    <term><see cref="T:System.DateTime"/></term>
+                    <term><see cref="M:Hangfire.Common.JobHelper.DeserializeDateTime(System.String)"/></term>
+                    <description>Please see the <see cref="P:Hangfire.States.ProcessingState.StartedAt"/> property.</description>
+                </item>
+                <item>
+                    <term><c>ServerId</c></term>
+                    <term><see cref="T:System.String"/></term>
+                    <term><i>Not required</i></term>
+                    <description>Please see the <see cref="P:Hangfire.States.ProcessingState.ServerId"/> property.</description>
+                </item>
+                <item>
+                    <term><c>WorkerId</c></term>
+                    <term><see cref="T:System.String"/></term>
+                    <term><i>Not required</i></term>
+                    <description>Please see the <see cref="P:Hangfire.States.ProcessingState.WorkerId"/> property.</description>
+                </item>
+            </list>
+            </remarks>
+        </member>
+        <member name="P:Hangfire.States.ProcessingState.StartedAt">
+            <summary>
+            Gets a date/time when the current state instance was created.
+            </summary>
+        </member>
+        <member name="P:Hangfire.States.ProcessingState.ServerId">
+            <summary>
+            Gets the <i>instance id</i> of an instance of the <see cref="T:Hangfire.Server.BackgroundProcessingServer"/>
+            class, whose <see cref="T:Hangfire.Server.Worker"/> background process started to process an 
+            <i>enqueued</i> background job.
+            </summary>
+            <value>Usually the string representation of a GUID value, may vary in future versions.</value>
+        </member>
+        <member name="P:Hangfire.States.ProcessingState.WorkerId">
+            <summary>
+            Gets the identifier of a <see cref="T:Hangfire.Server.Worker"/> that started to
+            process an <i>enqueued</i> background job.
+            </summary>
+        </member>
+        <member name="P:Hangfire.States.ProcessingState.Name">
+            <inheritdoc/>
+            <remarks>
+            Always equals to <see cref="F:Hangfire.States.ProcessingState.StateName"/> for the <see cref="T:Hangfire.States.ProcessingState"/>.
+            Please see the remarks section of the <see cref="P:Hangfire.States.IState.Name">IState.Name</see>
+            article for the details.
+            </remarks>
+        </member>
+        <member name="P:Hangfire.States.ProcessingState.Reason">
+            <inheritdoc />
+        </member>
+        <member name="P:Hangfire.States.ProcessingState.IsFinal">
+            <inheritdoc/>
+            <remarks>
+            Always returns <see langword="true"/> for the <see cref="T:Hangfire.States.ProcessingState"/>.
+            Please refer to the <see cref="P:Hangfire.States.IState.IsFinal">IState.IsFinal</see> documentation
+            for the details.
+            </remarks>
+        </member>
+        <member name="P:Hangfire.States.ProcessingState.IgnoreJobLoadException">
+            <inheritdoc/>
+            <remarks>
+            Always returns <see langword="false"/> for the <see cref="T:Hangfire.States.ProcessingState"/>.
+            Please see the description of this property in the
+            <see cref="P:Hangfire.States.IState.IgnoreJobLoadException">IState.IgnoreJobLoadException</see>
+            article.
+            </remarks>
+        </member>
+        <member name="T:Hangfire.States.ScheduledState">
+            <summary>
+            Defines the <i>intermediate</i> state of a background job when it is placed 
+            on a schedule to be moved to the <see cref="T:Hangfire.States.EnqueuedState"/> in the future 
+            by <see cref="T:Hangfire.Server.DelayedJobScheduler"/> background process.
+            </summary>
+            
+            <remarks>
+            <para>Background job in <see cref="T:Hangfire.States.ScheduledState"/> is referred as
+            <b>delayed job</b>.</para>
+            </remarks>
+            
+            <example>
+            The following example demonstrates the creation of a background job that will
+            be processed after two hours. Please see <see cref="O:Hangfire.BackgroundJob.Schedule">BackgroundJob.Schedule</see>
+            and <see cref="O:Hangfire.BackgroundJobClientExtensions.Schedule">BackgroundJobClientExtensions.Schedule</see>
+            method overloads for simpler API.
+            
+            <code lang="cs" source="..\Samples\States.cs" region="ScheduledState"/>
+            </example>
+            
+            <seealso cref="O:Hangfire.BackgroundJob.Schedule">BackgroundJob.Schedule Overload</seealso>
+            <seealso cref="O:Hangfire.BackgroundJobClientExtensions.Schedule">BackgroundJobClientExtensions.Schedule Overload</seealso>
+            <seealso cref="T:Hangfire.Server.DelayedJobScheduler"/>
+            <seealso cref="T:Hangfire.States.EnqueuedState"/>
+            
+            <threadsafety static="true" instance="false"/>
+        </member>
+        <member name="F:Hangfire.States.ScheduledState.StateName">
+            <summary>
+            Represents the name of the <i>Scheduled</i> state. This field is read-only.
+            </summary>
+            <remarks>
+            The value of this field is <c>"Scheduled"</c>.
+            </remarks>
+        </member>
+        <member name="M:Hangfire.States.ScheduledState.#ctor(System.TimeSpan)">
+            <summary>
+            Initializes a new instance of the <see cref="T:Hangfire.States.ScheduledState"/> class
+            with the specified <i>time interval</i> after which a job should be moved to
+            the <see cref="T:Hangfire.States.EnqueuedState"/>.
+            </summary>
+            <param name="enqueueIn">The time interval after which a job will be
+            moved to the <see cref="T:Hangfire.States.EnqueuedState"/>.</param>
+        </member>
+        <member name="M:Hangfire.States.ScheduledState.#ctor(System.DateTime)">
+            <summary>
+            Initializes a new instance of the <see cref="T:Hangfire.States.ScheduledState"/>
+            class with the specified <i>date/time in UTC format</i> when a job should 
+            be moved to the <see cref="T:Hangfire.States.EnqueuedState"/>.
+            </summary>
+            <param name="enqueueAt">The date/time when a job will be moved to the 
+            <see cref="T:Hangfire.States.EnqueuedState"/>.</param>
+        </member>
+        <member name="M:Hangfire.States.ScheduledState.SerializeData">
+            <inheritdoc/>
+            <remarks>
+            <para>Returning dictionary contains the following keys. You can obtain 
+            the state data by using the <see cref="M:Hangfire.Storage.IStorageConnection.GetStateData(System.String)"/>
+            method.</para>
+            <list type="table">
+                <listheader>
+                    <term>Key</term>
+                    <term>Type</term>
+                    <term>Deserialize Method</term>
+                    <description>Notes</description>
+                </listheader>
+                <item>
+                    <term><c>EnqueueAt</c></term>
+                    <term><see cref="T:System.DateTime"/></term>
+                    <term><see cref="M:Hangfire.Common.JobHelper.DeserializeDateTime(System.String)"/></term>
+                    <description>Please see the <see cref="P:Hangfire.States.ScheduledState.EnqueueAt"/> property.</description>
+                </item>
+                <item>
+                    <term><c>ScheduledAt</c></term>
+                    <term><see cref="T:System.DateTime"/></term>
+                    <term><see cref="M:Hangfire.Common.JobHelper.DeserializeDateTime(System.String)"/></term>
+                    <description>Please see the <see cref="P:Hangfire.States.ScheduledState.ScheduledAt"/> property.</description>
+                </item>
+            </list>
+            </remarks>
+        </member>
+        <member name="P:Hangfire.States.ScheduledState.EnqueueAt">
+            <summary>
+            Gets a date/time when a background job should be <i>enqueued</i>.
+            </summary>
+            <value>Any date/time in <see cref="F:System.DateTimeKind.Utc"/> format.</value>
+        </member>
+        <member name="P:Hangfire.States.ScheduledState.ScheduledAt">
+            <summary>
+            Gets a date/time when the current state instance was created.
+            </summary>
+        </member>
+        <member name="P:Hangfire.States.ScheduledState.Name">
+            <inheritdoc/>
+            <remarks>
+            Always equals to <see cref="F:Hangfire.States.ScheduledState.StateName"/> for the <see cref="T:Hangfire.States.ScheduledState"/>.
+            Please see the remarks section of the <see cref="P:Hangfire.States.IState.Name">IState.Name</see>
+            article for the details.
+            </remarks>
+        </member>
+        <member name="P:Hangfire.States.ScheduledState.Reason">
+            <inheritdoc />
+        </member>
+        <member name="P:Hangfire.States.ScheduledState.IsFinal">
+            <inheritdoc/>
+            <remarks>
+            Always returns <see langword="false"/> for the <see cref="T:Hangfire.States.ScheduledState"/>.
+            Please refer to the <see cref="P:Hangfire.States.IState.IsFinal">IState.IsFinal</see> documentation
+            for the details.
+            </remarks>
+        </member>
+        <member name="P:Hangfire.States.ScheduledState.IgnoreJobLoadException">
+            <inheritdoc/>
+            <remarks>
+            Always returns <see langword="false"/> for the <see cref="T:Hangfire.States.ScheduledState"/>.
+            Please see the description of this property in the
+            <see cref="P:Hangfire.States.IState.IgnoreJobLoadException">IState.IgnoreJobLoadException</see>
+            article.
+            </remarks>
+        </member>
+        <member name="T:Hangfire.States.StateContext">
+            <exclude />
+        </member>
+        <member name="T:Hangfire.States.SucceededState">
+            <summary>
+            Defines the <i>final</i> state of a background job when a <see cref="T:Hangfire.Server.Worker"/>
+            performed an <i>enqueued</i> job without any exception thrown during the performance.
+            </summary>
+            <remarks>
+            <para>All the transitions to the <i>Succeeded</i> state are internal for the <see cref="T:Hangfire.Server.Worker"/>
+            background process. You can't create background jobs using this state, and can't change state
+            to <i>Succeeded</i>.</para>
+            <para>This state is used in a user code primarily in state change filters (TODO: add a link)
+            to add custom logic during state transitions.</para> 
+            </remarks> 
+            
+            <seealso cref="T:Hangfire.States.EnqueuedState"/>
+            <seealso cref="T:Hangfire.Server.Worker"/>
+            <seealso cref="T:Hangfire.States.IState"/>
+            
+            <threadsafety static="true" instance="false"/>
+        </member>
+        <member name="F:Hangfire.States.SucceededState.StateName">
+            <summary>
+            Represents the name of the <i>Succeeded</i> state. This field is read-only.
+            </summary>
+            <remarks>
+            The value of this field is <c>"Succeeded"</c>.
+            </remarks>
+        </member>
+        <member name="M:Hangfire.States.SucceededState.SerializeData">
+            <inheritdoc/>
+            <remarks>
+            <para>Returning dictionary contains the following keys. You can obtain 
+            the state data by using the <see cref="M:Hangfire.Storage.IStorageConnection.GetStateData(System.String)"/>
+            method.</para>
+            <list type="table">
+                <listheader>
+                    <term>Key</term>
+                    <term>Type</term>
+                    <term>Deserialize Method</term>
+                    <description>Notes</description>
+                </listheader>
+                <item>
+                    <term><c>SucceededAt</c></term>
+                    <term><see cref="T:System.DateTime"/></term>
+                    <term><see cref="M:Hangfire.Common.JobHelper.DeserializeDateTime(System.String)"/></term>
+                    <description>Please see the <see cref="P:Hangfire.States.SucceededState.SucceededAt"/> property.</description>
+                </item>
+                <item>
+                    <term><c>PerformanceDuration</c></term>
+                    <term><see cref="T:System.Int64"/></term>
+                    <term>
+                        <see cref="M:System.Int64.Parse(System.String,System.IFormatProvider)"/> with 
+                        <see cref="P:System.Globalization.CultureInfo.InvariantCulture"/>
+                    </term>
+                    <description>Please see the <see cref="P:Hangfire.States.SucceededState.PerformanceDuration"/> property.</description>
+                </item>
+                <item>
+                    <term><c>Latency</c></term>
+                    <term><see cref="T:System.Int64"/></term>
+                    <term>
+                        <see cref="M:System.Int64.Parse(System.String,System.IFormatProvider)"/> with 
+                        <see cref="P:System.Globalization.CultureInfo.InvariantCulture"/>
+                    </term>
+                    <description>Please see the <see cref="P:Hangfire.States.SucceededState.Latency"/> property.</description>
+                </item>
+                <item>
+                    <term><c>Result</c></term>
+                    <term><see cref="T:System.Object"/></term>
+                    <term><see cref="M:Hangfire.Common.JobHelper.FromJson(System.String,System.Type)"/></term>
+                    <description>
+                        <para>Please see the <see cref="P:Hangfire.States.SucceededState.Result"/> property.</para>
+                        <para>This key may be missing from the dictionary, when the return 
+                        value was <see langword="null"/>. Always check for its existence 
+                        before using it.</para>
+                    </description>
+                </item>
+            </list>
+            </remarks>
+        </member>
+        <member name="P:Hangfire.States.SucceededState.SucceededAt">
+            <summary>
+            Gets a date/time when the current state instance was created.
+            </summary>
+        </member>
+        <member name="P:Hangfire.States.SucceededState.Result">
+            <summary>
+            Gets the value returned by a job method.
+            </summary>
+        </member>
+        <member name="P:Hangfire.States.SucceededState.Latency">
+            <summary>
+            Gets the total number of milliseconds passed from a job
+            creation time till the start of the performance.
+            </summary>
+        </member>
+        <member name="P:Hangfire.States.SucceededState.PerformanceDuration">
+            <summary>
+            Gets the total milliseconds elapsed from a processing start.
+            </summary>
+        </member>
+        <member name="P:Hangfire.States.SucceededState.Name">
+            <inheritdoc/>
+            <remarks>
+            Always equals to <see cref="F:Hangfire.States.SucceededState.StateName"/> for the <see cref="T:Hangfire.States.SucceededState"/>.
+            Please see the remarks section of the <see cref="P:Hangfire.States.IState.Name">IState.Name</see>
+            article for the details.
+            </remarks>
+        </member>
+        <member name="P:Hangfire.States.SucceededState.Reason">
+            <inheritdoc />
+        </member>
+        <member name="P:Hangfire.States.SucceededState.IsFinal">
+            <inheritdoc/>
+            <remarks>
+            Always returns <see langword="true"/> for the <see cref="T:Hangfire.States.SucceededState"/>.
+            Please refer to the <see cref="P:Hangfire.States.IState.IsFinal">IState.IsFinal</see> documentation
+            for the details.
+            </remarks>
+        </member>
+        <member name="P:Hangfire.States.SucceededState.IgnoreJobLoadException">
+            <inheritdoc/>
+            <remarks>
+            Always returns <see langword="false"/> for the <see cref="T:Hangfire.States.SucceededState"/>.
+            Please see the description of this property in the
+            <see cref="P:Hangfire.States.IState.IgnoreJobLoadException">IState.IgnoreJobLoadException</see>
+            article.
+            </remarks>
+        </member>
+    </members>
+</doc>
