Copying //UE4/Dev-Core to //UE4/Main

diff --git a/Engine/Source/Programs/CrashReporter/CrashReportWebSite/packages/HangFire.SqlServer.1.5.3/lib/net45/Hangfire.SqlServer.xml b/Engine/Source/Programs/CrashReporter/CrashReportWebSite/packages/HangFire.SqlServer.1.5.3/lib/net45/Hangfire.SqlServer.xml
new file mode 100644
--- /dev/null
+++ b/Engine/Source/Programs/CrashReporter/CrashReportWebSite/packages/HangFire.SqlServer.1.5.3/lib/net45/Hangfire.SqlServer.xml
@@ -0,0 +1,49 @@
+<?xml version="1.0"?>
+<doc>
+    <assembly>
+        <name>Hangfire.SqlServer</name>
+    </assembly>
+    <members>
+        <member name="M:Hangfire.SqlServer.SqlServerStorage.#ctor(System.String,Hangfire.SqlServer.SqlServerStorageOptions)">
+            <summary>
+            Initializes SqlServerStorage from the provided SqlServerStorageOptions and either the provided connection
+            string or the connection string with provided name pulled from the application config file.       
+            </summary>
+            <param name="nameOrConnectionString">Either a SQL Server connection string or the name of 
+            a SQL Server connection string located in the connectionStrings node in the application config</param>
+            <param name="options"></param>
+            <exception cref="T:System.ArgumentNullException"><paramref name="nameOrConnectionString"/> argument is null.</exception>
+            <exception cref="T:System.ArgumentNullException"><paramref name="options"/> argument is null.</exception>
+            <exception cref="T:System.ArgumentException"><paramref name="nameOrConnectionString"/> argument is neither 
+            a valid SQL Server connection string nor the name of a connection string in the application
+            config file.</exception>
+        </member>
+        <member name="M:Hangfire.SqlServer.SqlServerStorage.#ctor(System.Data.SqlClient.SqlConnection)">
+            <summary>
+            Initializes a new instance of the <see cref="T:Hangfire.SqlServer.SqlServerStorage"/> class with
+            explicit instance of the <see cref="T:System.Data.SqlClient.SqlConnection"/> class that will be used
+            to query the data.
+            </summary>
+            <param name="existingConnection">Existing connection</param>
+        </member>
+        <member name="M:Hangfire.SqlServer.SqlServerBootstrapperConfigurationExtensions.UseSqlServerStorage(Hangfire.IBootstrapperConfiguration,System.String)">
+            <summary>
+            Tells the bootstrapper to use SQL Server as a job storage,
+            that can be accessed using the given connection string or 
+            its name.
+            </summary>
+            <param name="configuration">Configuration</param>
+            <param name="nameOrConnectionString">Connection string or its name</param>
+        </member>
+        <member name="M:Hangfire.SqlServer.SqlServerBootstrapperConfigurationExtensions.UseSqlServerStorage(Hangfire.IBootstrapperConfiguration,System.String,Hangfire.SqlServer.SqlServerStorageOptions)">
+            <summary>
+            Tells the bootstrapper to use SQL Server as a job storage
+            with the given options, that can be accessed using the specified
+            connection string or its name.
+            </summary>
+            <param name="configuration">Configuration</param>
+            <param name="nameOrConnectionString">Connection string or its name</param>
+            <param name="options">Advanced options</param>
+        </member>
+    </members>
+</doc>
