Copying //UE4/Dev-Core to //UE4/Main

diff --git a/Engine/Source/Programs/CrashReporter/CrashReportWebSite/packages/Microsoft.AspNet.Razor.3.2.3/lib/net45/System.Web.Razor.xml b/Engine/Source/Programs/CrashReporter/CrashReportWebSite/packages/Microsoft.AspNet.Razor.3.2.3/lib/net45/System.Web.Razor.xml
new file mode 100644
--- /dev/null
+++ b/Engine/Source/Programs/CrashReporter/CrashReportWebSite/packages/Microsoft.AspNet.Razor.3.2.3/lib/net45/System.Web.Razor.xml
@@ -0,0 +1,5742 @@
+﻿<?xml version="1.0" encoding="utf-8"?>
+<doc>
+  <assembly>
+    <name>System.Web.Razor</name>
+  </assembly>
+  <members>
+    <member name="T:System.Web.Razor.CSharpRazorCodeLanguage">
+      <summary>This type/member supports the .NET Framework infrastructure and is not intended to be used directly from your code. Represents a Razor code language that is based on C# syntax.</summary>
+    </member>
+    <member name="M:System.Web.Razor.CSharpRazorCodeLanguage.#ctor">
+      <summary>This type/member supports the .NET Framework infrastructure and is not intended to be used directly from your code. Initializes a new instance of the <see cref="T:System.Web.Razor.CSharpRazorCodeLanguage" /> class.</summary>
+    </member>
+    <member name="P:System.Web.Razor.CSharpRazorCodeLanguage.CodeDomProviderType">
+      <summary>This type/member supports the .NET Framework infrastructure and is not intended to be used directly from your code. Gets the type of the code provider.</summary>
+      <returns>The type of the code provider.</returns>
+    </member>
+    <member name="M:System.Web.Razor.CSharpRazorCodeLanguage.CreateCodeGenerator(System.String,System.String,System.String,System.Web.Razor.RazorEngineHost)">
+      <summary>This type/member supports the .NET Framework infrastructure and is not intended to be used directly from your code. Creates a new Razor code generator based on C# code language.</summary>
+      <returns>The newly created Razor code generator based on C# code language.</returns>
+      <param name="className">The class name for the generated code.</param>
+      <param name="rootNamespaceName">The name of the root namespace for the generated code.</param>
+      <param name="sourceFileName">The name of the source code file.</param>
+      <param name="host">The Razor engine host.</param>
+    </member>
+    <member name="M:System.Web.Razor.CSharpRazorCodeLanguage.CreateCodeParser">
+      <summary>This type/member supports the .NET Framework infrastructure and is not intended to be used directly from your code. Creates a new code parser for C# code language.</summary>
+      <returns>The newly created code parser for C# code language.</returns>
+    </member>
+    <member name="P:System.Web.Razor.CSharpRazorCodeLanguage.LanguageName">
+      <summary>This type/member supports the .NET Framework infrastructure and is not intended to be used directly from your code. Gets the name of the C# code language.</summary>
+      <returns>The name of the C# code language. Value is ‘csharp’.</returns>
+    </member>
+    <member name="T:System.Web.Razor.DocumentParseCompleteEventArgs">
+      <summary>This type/member supports the .NET Framework infrastructure and is not intended to be used directly from your code.</summary>
+    </member>
+    <member name="M:System.Web.Razor.DocumentParseCompleteEventArgs.#ctor">
+      <summary>This type/member supports the .NET Framework infrastructure and is not intended to be used directly from your code.</summary>
+    </member>
+    <member name="P:System.Web.Razor.DocumentParseCompleteEventArgs.GeneratorResults">
+      <summary>This type/member supports the .NET Framework infrastructure and is not intended to be used directly from your code.</summary>
+    </member>
+    <member name="P:System.Web.Razor.DocumentParseCompleteEventArgs.SourceChange">
+      <summary>This type/member supports the .NET Framework infrastructure and is not intended to be used directly from your code.</summary>
+    </member>
+    <member name="P:System.Web.Razor.DocumentParseCompleteEventArgs.TreeStructureChanged">
+      <summary>This type/member supports the .NET Framework infrastructure and is not intended to be used directly from your code.</summary>
+    </member>
+    <member name="T:System.Web.Razor.GeneratorResults">
+      <summary>This type/member supports the .NET Framework infrastructure and is not intended to be used directly from your code.Represents results from code generation.</summary>
+    </member>
+    <member name="M:System.Web.Razor.GeneratorResults.#ctor(System.Boolean,System.Web.Razor.Parser.SyntaxTree.Block,System.Collections.Generic.IList{System.Web.Razor.Parser.SyntaxTree.RazorError},System.CodeDom.CodeCompileUnit,System.Collections.Generic.IDictionary{System.Int32,System.Web.Razor.Generator.GeneratedCodeMapping})">
+      <summary>This type/member supports the .NET Framework infrastructure and is not intended to be used directly from your code.Initializes a new instance of the <see cref="T:System.Web.Razor.GeneratorResults" /> class.</summary>
+      <param name="success">true if the code generation is a success; otherwise, false.</param>
+      <param name="document">The document.</param>
+      <param name="parserErrors">The parser errors.</param>
+      <param name="generatedCode">The generated code.</param>
+      <param name="designTimeLineMappings">The dictionary of design-time generated code mappings.</param>
+    </member>
+    <member name="M:System.Web.Razor.GeneratorResults.#ctor(System.Web.Razor.Parser.SyntaxTree.Block,System.Collections.Generic.IList{System.Web.Razor.Parser.SyntaxTree.RazorError},System.CodeDom.CodeCompileUnit,System.Collections.Generic.IDictionary{System.Int32,System.Web.Razor.Generator.GeneratedCodeMapping})">
+      <summary>This type/member supports the .NET Framework infrastructure and is not intended to be used directly from your code.Initializes a new instance of the <see cref="T:System.Web.Razor.GeneratorResults" /> class.</summary>
+      <param name="document">The document.</param>
+      <param name="parserErrors">The parser errors.</param>
+      <param name="generatedCode">The generated code.</param>
+      <param name="designTimeLineMappings">The dictionary of design-time generated code mappings.</param>
+    </member>
+    <member name="M:System.Web.Razor.GeneratorResults.#ctor(System.Web.Razor.ParserResults,System.CodeDom.CodeCompileUnit,System.Collections.Generic.IDictionary{System.Int32,System.Web.Razor.Generator.GeneratedCodeMapping})">
+      <summary>This type/member supports the .NET Framework infrastructure and is not intended to be used directly from your code.Initializes a new instance of the <see cref="T:System.Web.Razor.GeneratorResults" /> class.</summary>
+      <param name="parserResults">The parser results.</param>
+      <param name="generatedCode">The generated code.</param>
+      <param name="designTimeLineMappings">The dictionary of design-time generated code mappings.</param>
+    </member>
+    <member name="P:System.Web.Razor.GeneratorResults.DesignTimeLineMappings">
+      <summary>This type/member supports the .NET Framework infrastructure and is not intended to be used directly from your code.Gets or sets the dictionary of design-time generated code mappings.</summary>
+      <returns>The dictionary of design-time generated code mappings.</returns>
+    </member>
+    <member name="P:System.Web.Razor.GeneratorResults.GeneratedCode">
+      <summary>This type/member supports the .NET Framework infrastructure and is not intended to be used directly from your code.Gets or sets the generated code.</summary>
+      <returns>The generated code.</returns>
+    </member>
+    <member name="T:System.Web.Razor.ParserResults">
+      <summary>This type/member supports the .NET Framework infrastructure and is not intended to be used directly from your code.Represents the results of parsing a Razor document.</summary>
+    </member>
+    <member name="M:System.Web.Razor.ParserResults.#ctor(System.Boolean,System.Web.Razor.Parser.SyntaxTree.Block,System.Collections.Generic.IList{System.Web.Razor.Parser.SyntaxTree.RazorError})">
+      <summary>This type/member supports the .NET Framework infrastructure and is not intended to be used directly from your code.Initializes a new instance of the <see cref="T:System.Web.Razor.ParserResults" /> class.</summary>
+      <param name="success">true if parsing was successful; otherwise, false.</param>
+      <param name="document">The root node in the document’s syntax tree.</param>
+      <param name="errors">The list of errors which occurred during parsing.</param>
+    </member>
+    <member name="M:System.Web.Razor.ParserResults.#ctor(System.Web.Razor.Parser.SyntaxTree.Block,System.Collections.Generic.IList{System.Web.Razor.Parser.SyntaxTree.RazorError})">
+      <summary>This type/member supports the .NET Framework infrastructure and is not intended to be used directly from your code.Initializes a new instance of the <see cref="T:System.Web.Razor.ParserResults" /> class.</summary>
+      <param name="document">The root node in the document’s syntax tree.</param>
+      <param name="parserErrors">The list of errors which occurred during parsing.</param>
+    </member>
+    <member name="P:System.Web.Razor.ParserResults.Document">
+      <summary>This type/member supports the .NET Framework infrastructure and is not intended to be used directly from your code.Gets or sets the root node in the document’s syntax tree.</summary>
+      <returns>The root node in the document’s syntax tree.</returns>
+    </member>
+    <member name="P:System.Web.Razor.ParserResults.ParserErrors">
+      <summary>This type/member supports the .NET Framework infrastructure and is not intended to be used directly from your code.Gets or sets the list of errors which occurred during parsing.</summary>
+      <returns>The list of errors which occurred during parsing.</returns>
+    </member>
+    <member name="P:System.Web.Razor.ParserResults.Success">
+      <summary>This type/member supports the .NET Framework infrastructure and is not intended to be used directly from your code.Gets or sets a value indicating whether parsing was successful.</summary>
+      <returns>true if parsing was successful; otherwise, false.</returns>
+    </member>
+    <member name="T:System.Web.Razor.PartialParseResult">
+      <summary>This type/member supports the .NET Framework infrastructure and is not intended to be used directly from your code.</summary>
+    </member>
+    <member name="F:System.Web.Razor.PartialParseResult.Accepted">
+      <summary>This type/member supports the .NET Framework infrastructure and is not intended to be used directly from your code.</summary>
+    </member>
+    <member name="F:System.Web.Razor.PartialParseResult.AutoCompleteBlock">
+      <summary>This type/member supports the .NET Framework infrastructure and is not intended to be used directly from your code.</summary>
+    </member>
+    <member name="F:System.Web.Razor.PartialParseResult.Provisional">
+      <summary>This type/member supports the .NET Framework infrastructure and is not intended to be used directly from your code.</summary>
+    </member>
+    <member name="F:System.Web.Razor.PartialParseResult.Rejected">
+      <summary>This type/member supports the .NET Framework infrastructure and is not intended to be used directly from your code.</summary>
+    </member>
+    <member name="F:System.Web.Razor.PartialParseResult.SpanContextChanged">
+      <summary>This type/member supports the .NET Framework infrastructure and is not intended to be used directly from your code.</summary>
+    </member>
+    <member name="T:System.Web.Razor.RazorCodeLanguage">
+      <summary>Represents the base for all Razor code language.This type/member supports the .NET Framework infrastructure and is not intended to be used directly from your code.</summary>
+    </member>
+    <member name="M:System.Web.Razor.RazorCodeLanguage.#ctor">
+      <summary>Initializes a new instance of the <see cref="T:System.Web.Razor.RazorCodeLanguage" /> class.This type/member supports the .NET Framework infrastructure and is not intended to be used directly from your code.</summary>
+    </member>
+    <member name="P:System.Web.Razor.RazorCodeLanguage.CodeDomProviderType">
+      <summary>Gets the type of the CodeDOM provider.This type/member supports the .NET Framework infrastructure and is not intended to be used directly from your code.</summary>
+      <returns>The type of the CodeDOM provider.</returns>
+    </member>
+    <member name="M:System.Web.Razor.RazorCodeLanguage.CreateCodeGenerator(System.String,System.String,System.String,System.Web.Razor.RazorEngineHost)">
+      <summary>Creates the code generator for the Razor code language.This type/member supports the .NET Framework infrastructure and is not intended to be used directly from your code.</summary>
+      <returns>The code generator for the Razor code language.</returns>
+      <param name="className">The class name.</param>
+      <param name="rootNamespaceName">The name of the root namespace.</param>
+      <param name="sourceFileName">The source file name.</param>
+      <param name="host">The Razor engine host.</param>
+    </member>
+    <member name="M:System.Web.Razor.RazorCodeLanguage.CreateCodeParser">
+      <summary>Creates the code parser for the Razor code language.This type/member supports the .NET Framework infrastructure and is not intended to be used directly from your code.</summary>
+      <returns>The code parser for the Razor code language.</returns>
+    </member>
+    <member name="M:System.Web.Razor.RazorCodeLanguage.GetLanguageByExtension(System.String)">
+      <summary>Gets the language of the Razor code using the specified file extension.This type/member supports the .NET Framework infrastructure and is not intended to be used directly from your code.</summary>
+      <returns>The language of the Razor code.</returns>
+      <param name="fileExtension">The file extension.</param>
+    </member>
+    <member name="P:System.Web.Razor.RazorCodeLanguage.LanguageName">
+      <summary>Gets the language name of the current Razor code, that is “csharp” or “vb”.This type/member supports the .NET Framework infrastructure and is not intended to be used directly from your code.</summary>
+      <returns>The language name of the current Razor code.</returns>
+    </member>
+    <member name="P:System.Web.Razor.RazorCodeLanguage.Languages">
+      <summary>Gets the list of language supported by the Razor code.This type/member supports the .NET Framework infrastructure and is not intended to be used directly from your code.</summary>
+      <returns>The list of language supported by the Razor code.</returns>
+    </member>
+    <member name="T:System.Web.Razor.RazorDirectiveAttribute">
+      <summary>This type/member supports the .NET Framework infrastructure and is not intended to be used directly from your code. Represents an attribute for the Razor directive.</summary>
+    </member>
+    <member name="M:System.Web.Razor.RazorDirectiveAttribute.#ctor(System.String,System.String)">
+      <summary>This type/member supports the .NET Framework infrastructure and is not intended to be used directly from your code. Initializes a new instance of the <see cref="T:System.Web.Razor.RazorDirectiveAttribute" /> class.</summary>
+      <param name="name">The name of the attribute.</param>
+      <param name="value">The value of the attribute.</param>
+    </member>
+    <member name="M:System.Web.Razor.RazorDirectiveAttribute.Equals(System.Object)">
+      <summary>This type/member supports the .NET Framework infrastructure and is not intended to be used directly from your code. Determines whether this instance is equal to a specified object.</summary>
+      <returns>true if the object is equal to the this instance; otherwise, false.</returns>
+      <param name="obj">The object to compare with this instance.</param>
+    </member>
+    <member name="M:System.Web.Razor.RazorDirectiveAttribute.GetHashCode">
+      <summary>This type/member supports the .NET Framework infrastructure and is not intended to be used directly from your code. Returns the hash code for this instance.</summary>
+      <returns>The hash code for this instance.</returns>
+    </member>
+    <member name="P:System.Web.Razor.RazorDirectiveAttribute.Name">
+      <summary>This type/member supports the .NET Framework infrastructure and is not intended to be used directly from your code. Gets or sets the name of the attribute.</summary>
+      <returns>The name of the attribute.</returns>
+    </member>
+    <member name="P:System.Web.Razor.RazorDirectiveAttribute.TypeId">
+      <summary>This type/member supports the .NET Framework infrastructure and is not intended to be used directly from your code. Gets the unique type ID of the attribute.</summary>
+      <returns>The unique type ID of the attribute.</returns>
+    </member>
+    <member name="P:System.Web.Razor.RazorDirectiveAttribute.Value">
+      <summary>This type/member supports the .NET Framework infrastructure and is not intended to be used directly from your code. Gets or sets the value of the attribute.</summary>
+      <returns>The value of the attribute.</returns>
+    </member>
+    <member name="T:System.Web.Razor.RazorEditorParser">
+      <summary>This type/member supports the .NET Framework infrastructure and is not intended to be used directly from your code.Parser used by editors to avoid reparsing the entire document on each text change.</summary>
+    </member>
+    <member name="M:System.Web.Razor.RazorEditorParser.#ctor(System.Web.Razor.RazorEngineHost,System.String)">
+      <summary>This type/member supports the .NET Framework infrastructure and is not intended to be used directly from your code.Constructs the editor parser.</summary>
+      <param name="host">The <see cref="T:System.Web.Razor.RazorEngineHost" /> which defines the environment in which the generated code will live.</param>
+      <param name="sourceFileName">The physical path to use in line pragmas.</param>
+    </member>
+    <member name="M:System.Web.Razor.RazorEditorParser.CheckForStructureChanges(System.Web.Razor.Text.TextChange)">
+      <summary>This type/member supports the .NET Framework infrastructure and is not intended to be used directly from your code.Determines if a change will cause a structural change to the document and if not, applies it to the existing tree. If a structural change would occur, automatically starts a reparse.</summary>
+      <returns>A <see cref="T:System.Web.Razor.PartialParseResult" /> value indicating the result of the incremental parse.</returns>
+      <param name="change">The change to apply to the parse tree.</param>
+    </member>
+    <member name="P:System.Web.Razor.RazorEditorParser.CurrentParseTree">
+      <summary>This type/member supports the .NET Framework infrastructure and is not intended to be used directly from your code.Gets the current parse tree.</summary>
+      <returns>The current parse tree.</returns>
+    </member>
+    <member name="M:System.Web.Razor.RazorEditorParser.Dispose">
+      <summary>This type/member supports the .NET Framework infrastructure and is not intended to be used directly from your code.Releases all resources used by the current instance of the <see cref="T:System.Web.Razor.RazorEditorParser" />.</summary>
+    </member>
+    <member name="M:System.Web.Razor.RazorEditorParser.Dispose(System.Boolean)">
+      <summary>This type/member supports the .NET Framework infrastructure and is not intended to be used directly from your code.Releases the unmanaged resources used by the <see cref="T:System.Web.Razor.RazorEditorParser" /> class and optionally releases the managed resources.</summary>
+      <param name="disposing">true to release both managed and unmanaged resources; false to release only unmanaged resources.</param>
+    </member>
+    <member name="E:System.Web.Razor.RazorEditorParser.DocumentParseComplete">
+      <summary>This type/member supports the .NET Framework infrastructure and is not intended to be used directly from your code.Event fired when a full reparse of the document completes.</summary>
+    </member>
+    <member name="P:System.Web.Razor.RazorEditorParser.FileName">
+      <summary>This type/member supports the .NET Framework infrastructure and is not intended to be used directly from your code.Gets or sets the filename of the document to parse.</summary>
+      <returns>The filename of the document to parse.</returns>
+    </member>
+    <member name="M:System.Web.Razor.RazorEditorParser.GetAutoCompleteString">
+      <summary>This type/member supports the .NET Framework infrastructure and is not intended to be used directly from your code.Retrieves the auto complete string.</summary>
+      <returns>The auto complete string.</returns>
+    </member>
+    <member name="P:System.Web.Razor.RazorEditorParser.Host">
+      <summary>This type/member supports the .NET Framework infrastructure and is not intended to be used directly from your code.Gets or sets the host for the parse.</summary>
+      <returns>The host for the parse.</returns>
+    </member>
+    <member name="P:System.Web.Razor.RazorEditorParser.LastResultProvisional">
+      <summary>This type/member supports the .NET Framework infrastructure and is not intended to be used directly from your code.Gets or sets a value indicating whether the last result of the parse was provisionally accepted for next partial parse.</summary>
+      <returns>true if the last result of the parse was provisionally accepted for next partial parse; otherwise, false.</returns>
+    </member>
+    <member name="T:System.Web.Razor.RazorEngineHost">
+      <summary>This type/member supports the .NET Framework infrastructure and is not intended to be used directly from your code.Represents the generated code for the razor engine host.</summary>
+    </member>
+    <member name="M:System.Web.Razor.RazorEngineHost.#ctor">
+      <summary>This type/member supports the .NET Framework infrastructure and is not intended to be used directly from your code.Initializes a new instance of the <see cref="T:System.Web.Razor.RazorEngineHost" /> class.</summary>
+    </member>
+    <member name="M:System.Web.Razor.RazorEngineHost.#ctor(System.Web.Razor.RazorCodeLanguage)">
+      <summary>This type/member supports the .NET Framework infrastructure and is not intended to be used directly from your code.Initializes a new instance of the <see cref="T:System.Web.Razor.RazorEngineHost" /> class.</summary>
+      <param name="codeLanguage">The specified code language.</param>
+    </member>
+    <member name="M:System.Web.Razor.RazorEngineHost.#ctor(System.Web.Razor.RazorCodeLanguage,System.Func{System.Web.Razor.Parser.ParserBase})">
+      <summary>Initializes a new instance of the <see cref="T:System.Web.Razor.RazorEngineHost" /> class.</summary>
+      <param name="codeLanguage">The specified code language.</param>
+      <param name="markupParserFactory">The markup parser factory.</param>
+    </member>
+    <member name="P:System.Web.Razor.RazorEngineHost.CodeLanguage">
+      <summary>This type/member supports the .NET Framework infrastructure and is not intended to be used directly from your code.Gets the languages supported by the code generator.</summary>
+      <returns>The languages supported that by the code generator.</returns>
+    </member>
+    <member name="M:System.Web.Razor.RazorEngineHost.CreateMarkupParser">
+      <summary>This type/member supports the .NET Framework infrastructure and is not intended to be used directly from your code.Creates a markup parser using the specified language parser for the <see cref="T:System.Web.Razor.RazorEngineHost" />.</summary>
+      <returns>A markup parser to create using the specified language parser for the <see cref="T:System.Web.Razor.RazorEngineHost" />.</returns>
+    </member>
+    <member name="M:System.Web.Razor.RazorEngineHost.DecorateCodeGenerator(System.Web.Razor.Generator.RazorCodeGenerator)">
+      <summary>This type/member supports the .NET Framework infrastructure and is not intended to be used directly from your code.Returns the methods as language-specific Razor code generator.</summary>
+      <returns>The methods as language-specific Razor code generator.</returns>
+      <param name="incomingCodeGenerator">The C# or Visual Basic code generator.</param>
+    </member>
+    <member name="M:System.Web.Razor.RazorEngineHost.DecorateCodeParser(System.Web.Razor.Parser.ParserBase)">
+      <summary>This type/member supports the .NET Framework infrastructure and is not intended to be used directly from your code.Returns the methods as language-specific Razor code parser using the specified language parser.</summary>
+      <returns>The methods as language-specific Razor code parser using the specified language parser.</returns>
+      <param name="incomingCodeParser">The C# or Visual Basic code parser.</param>
+    </member>
+    <member name="M:System.Web.Razor.RazorEngineHost.DecorateMarkupParser(System.Web.Razor.Parser.ParserBase)">
+      <summary>This type/member supports the .NET Framework infrastructure and is not intended to be used directly from your code.Returns the method to decorate markup parser using the specified language parser.</summary>
+      <returns>The method to decorate markup parser using the specified language parser.</returns>
+      <param name="incomingMarkupParser">The C# or Visual Basic code parser.</param>
+    </member>
+    <member name="P:System.Web.Razor.RazorEngineHost.DefaultBaseClass">
+      <summary>This type/member supports the .NET Framework infrastructure and is not intended to be used directly from your code.Gets or sets the default base class for the host.</summary>
+      <returns>The default base class for the host.</returns>
+    </member>
+    <member name="P:System.Web.Razor.RazorEngineHost.DefaultClassName">
+      <summary>This type/member supports the .NET Framework infrastructure and is not intended to be used directly from your code.Gets or sets the default class name for the host.</summary>
+      <returns>The default class name for the host.</returns>
+    </member>
+    <member name="P:System.Web.Razor.RazorEngineHost.DefaultNamespace">
+      <summary>This type/member supports the .NET Framework infrastructure and is not intended to be used directly from your code.Gets or sets the default namespace for the host.</summary>
+      <returns>The default namespace for the host.</returns>
+    </member>
+    <member name="P:System.Web.Razor.RazorEngineHost.DesignTimeMode">
+      <summary>This type/member supports the .NET Framework infrastructure and is not intended to be used directly from your code.Gets or sets a value that indicates whether the mode designs a time for the host.</summary>
+      <returns>true if the mode designs a time for the host; otherwise, false.</returns>
+    </member>
+    <member name="P:System.Web.Razor.RazorEngineHost.EnableInstrumentation">
+      <summary>This type/member supports the .NET Framework infrastructure and is not intended to be used directly from your code.Gets or sets the host that enables the instrumentation.</summary>
+      <returns>The host that enables the instrumentation.</returns>
+    </member>
+    <member name="P:System.Web.Razor.RazorEngineHost.GeneratedClassContext">
+      <summary>This type/member supports the .NET Framework infrastructure and is not intended to be used directly from your code.Gets or sets the generated class context for the host.</summary>
+      <returns>The generated class context for the host.</returns>
+    </member>
+    <member name="P:System.Web.Razor.RazorEngineHost.InstrumentedSourceFilePath">
+      <summary>This type/member supports the .NET Framework infrastructure and is not intended to be used directly from your code.Gets or sets the instrumented source file path for the host.</summary>
+      <returns>The instrumented source file path for the host.</returns>
+    </member>
+    <member name="P:System.Web.Razor.RazorEngineHost.IsIndentingWithTabs">
+      <summary>Gets or sets whether the design time editor is using tabs or spaces for indentation.</summary>
+    </member>
+    <member name="P:System.Web.Razor.RazorEngineHost.NamespaceImports">
+      <summary>This type/member supports the .NET Framework infrastructure and is not intended to be used directly from your code.Gets the namespace imports for the host.</summary>
+      <returns>The namespace imports for the host.</returns>
+    </member>
+    <member name="M:System.Web.Razor.RazorEngineHost.PostProcessGeneratedCode(System.CodeDom.CodeCompileUnit,System.CodeDom.CodeNamespace,System.CodeDom.CodeTypeDeclaration,System.CodeDom.CodeMemberMethod)">
+      <summary>This type/member supports the .NET Framework infrastructure and is not intended to be used directly from your code.Returns this method to post all the processed generated code for the host.</summary>
+      <param name="codeCompileUnit">The code compile unit.</param>
+      <param name="generatedNamespace">The generated namespace.</param>
+      <param name="generatedClass">The generated class.</param>
+      <param name="executeMethod">The execute method.</param>
+    </member>
+    <member name="M:System.Web.Razor.RazorEngineHost.PostProcessGeneratedCode(System.Web.Razor.Generator.CodeGeneratorContext)">
+      <summary>This type/member supports the .NET Framework infrastructure and is not intended to be used directly from your code.Returns this method to post all the processed generated code for the host.</summary>
+      <param name="context">The code generator context.</param>
+    </member>
+    <member name="P:System.Web.Razor.RazorEngineHost.StaticHelpers">
+      <summary>This type/member supports the .NET Framework infrastructure and is not intended to be used directly from your code.Gets or sets the static helpers for the host.</summary>
+      <returns>The static helpers for the host.</returns>
+    </member>
+    <member name="P:System.Web.Razor.RazorEngineHost.TabSize">
+      <summary>Tab size used by the hosting editor, when indenting with tabs.</summary>
+    </member>
+    <member name="T:System.Web.Razor.RazorTemplateEngine">
+      <summary>This type/member supports the .NET Framework infrastructure and is not intended to be used directly from your code.Represents an entry-point to the Razor Template Engine.</summary>
+    </member>
+    <member name="M:System.Web.Razor.RazorTemplateEngine.#ctor(System.Web.Razor.RazorEngineHost)">
+      <summary>This type/member supports the .NET Framework infrastructure and is not intended to be used directly from your code.Initializes a new instance of the <see cref="T:System.Web.Razor.RazorTemplateEngine" /> class.</summary>
+      <param name="host">The host.</param>
+    </member>
+    <member name="M:System.Web.Razor.RazorTemplateEngine.CreateCodeGenerator(System.String,System.String,System.String)">
+      <summary>This type/member supports the .NET Framework infrastructure and is not intended to be used directly from your code.Creates a code generator.</summary>
+      <returns>The created <see cref="T:System.Web.Razor.Generator.RazorCodeGenerator" />.</returns>
+      <param name="className">The name of the generated class.</param>
+      <param name="rootNamespace">The namespace in which the generated class will reside.</param>
+      <param name="sourceFileName">The file name to use in line pragmas.</param>
+    </member>
+    <member name="M:System.Web.Razor.RazorTemplateEngine.CreateParser">
+      <summary>This type/member supports the .NET Framework infrastructure and is not intended to be used directly from your code.Creates a <see cref="T:System.Web.Razor.Parser.RazorParser" />.</summary>
+      <returns>The created <see cref="T:System.Web.Razor.Parser.RazorParser" />.</returns>
+    </member>
+    <member name="F:System.Web.Razor.RazorTemplateEngine.DefaultClassName">
+      <summary>This type/member supports the .NET Framework infrastructure and is not intended to be used directly from your code.Indicates the default class name of the template.</summary>
+    </member>
+    <member name="F:System.Web.Razor.RazorTemplateEngine.DefaultNamespace">
+      <summary>This type/member supports the .NET Framework infrastructure and is not intended to be used directly from your code.Indicates the default namespace for the template.</summary>
+    </member>
+    <member name="M:System.Web.Razor.RazorTemplateEngine.GenerateCode(System.IO.TextReader)">
+      <summary>This type/member supports the .NET Framework infrastructure and is not intended to be used directly from your code.Parses the template specified by the TextBuffer, generates code for it, and returns the constructed CodeDOM tree.</summary>
+      <returns>The resulting parse tree AND generated Code DOM tree.</returns>
+      <param name="input">The input text to parse.</param>
+    </member>
+    <member name="M:System.Web.Razor.RazorTemplateEngine.GenerateCode(System.IO.TextReader,System.Nullable{System.Threading.CancellationToken})">
+      <summary>This type/member supports the .NET Framework infrastructure and is not intended to be used directly from your code.Parses the template specified by the TextBuffer, generates code for it, and returns the constructed CodeDOM tree.</summary>
+      <returns>The resulting parse tree AND generated Code DOM tree.</returns>
+      <param name="input">The input text to parse.</param>
+      <param name="cancelToken">A token used to cancel the parser.</param>
+    </member>
+    <member name="M:System.Web.Razor.RazorTemplateEngine.GenerateCode(System.IO.TextReader,System.String,System.String,System.String)">
+      <summary>This type/member supports the .NET Framework infrastructure and is not intended to be used directly from your code.Parses the template specified by the TextBuffer, generates code for it, and returns the constructed CodeDOM tree.</summary>
+      <returns>The resulting parse tree AND generated Code DOM tree.</returns>
+      <param name="input">The input text to parse.</param>
+      <param name="className">The name of the generated class, overriding whatever is specified in the host.</param>
+      <param name="rootNamespace">The namespace in which the generated class will reside.</param>
+      <param name="sourceFileName">The file name to use in line pragmas.</param>
+    </member>
+    <member name="M:System.Web.Razor.RazorTemplateEngine.GenerateCode(System.IO.TextReader,System.String,System.String,System.String,System.Nullable{System.Threading.CancellationToken})">
+      <summary>This type/member supports the .NET Framework infrastructure and is not intended to be used directly from your code.Parses the template specified by the TextBuffer, generates code for it, and returns the constructed CodeDOM tree.</summary>
+      <returns>The resulting parse tree AND generated Code DOM tree.</returns>
+      <param name="input">The input text to parse.</param>
+      <param name="className">The name of the generated class, overriding whatever is specified in the host.</param>
+      <param name="rootNamespace">The namespace in which the generated class will reside.</param>
+      <param name="sourceFileName">The file name to use in line pragmas.</param>
+      <param name="cancelToken">A token used to cancel the parser.</param>
+    </member>
+    <member name="M:System.Web.Razor.RazorTemplateEngine.GenerateCode(System.Web.Razor.Text.ITextBuffer)">
+      <summary>This type/member supports the .NET Framework infrastructure and is not intended to be used directly from your code.Parses the template specified by the TextBuffer, generates code for it, and returns the constructed CodeDOM tree.</summary>
+      <returns>The resulting parse tree AND generated Code DOM tree.</returns>
+      <param name="input">The input text to parse.</param>
+    </member>
+    <member name="M:System.Web.Razor.RazorTemplateEngine.GenerateCode(System.Web.Razor.Text.ITextBuffer,System.Nullable{System.Threading.CancellationToken})">
+      <summary>This type/member supports the .NET Framework infrastructure and is not intended to be used directly from your code.Parses the template specified by the TextBuffer, generates code for it, and returns the constructed CodeDOM tree.</summary>
+      <returns>The resulting parse tree AND generated Code DOM tree.</returns>
+      <param name="input">The input text to parse.</param>
+      <param name="cancelToken">A token used to cancel the parser.</param>
+    </member>
+    <member name="M:System.Web.Razor.RazorTemplateEngine.GenerateCode(System.Web.Razor.Text.ITextBuffer,System.String,System.String,System.String)">
+      <summary>This type/member supports the .NET Framework infrastructure and is not intended to be used directly from your code.Parses the template specified by the TextBuffer, generates code for it, and returns the constructed CodeDOM tree.</summary>
+      <returns>The resulting parse tree AND generated Code DOM tree.</returns>
+      <param name="input">The input text to parse.</param>
+      <param name="className">The name of the generated class, overriding whatever is specified in the host.</param>
+      <param name="rootNamespace">The namespace in which the generated class will reside.</param>
+      <param name="sourceFileName">The file name to use in line pragmas.</param>
+    </member>
+    <member name="M:System.Web.Razor.RazorTemplateEngine.GenerateCode(System.Web.Razor.Text.ITextBuffer,System.String,System.String,System.String,System.Nullable{System.Threading.CancellationToken})">
+      <summary>This type/member supports the .NET Framework infrastructure and is not intended to be used directly from your code.Parses the template specified by the TextBuffer, generates code for it, and returns the constructed CodeDOM tree.</summary>
+      <returns>The resulting parse tree AND generated Code DOM tree.</returns>
+      <param name="input">The input text to parse.</param>
+      <param name="className">The name of the generated class, overriding whatever is specified in the host.</param>
+      <param name="rootNamespace">The namespace in which the generated class will reside.</param>
+      <param name="sourceFileName">The file name to use in line pragmas.</param>
+      <param name="cancelToken">A token used to cancel the parser.</param>
+    </member>
+    <member name="M:System.Web.Razor.RazorTemplateEngine.GenerateCodeCore(System.Web.Razor.Text.ITextDocument,System.String,System.String,System.String,System.Nullable{System.Threading.CancellationToken})">
+      <summary>This type/member supports the .NET Framework infrastructure and is not intended to be used directly from your code.Generates a code core.</summary>
+      <returns>The results of the generated core.</returns>
+      <param name="input">The input text to parse.</param>
+      <param name="className">The name of the generated class, overriding whatever is specified in the host.</param>
+      <param name="rootNamespace">The namespace in which the generated class will reside.</param>
+      <param name="sourceFileName">The file name to use in line pragmas.</param>
+      <param name="cancelToken">A token used to cancel the parser.</param>
+    </member>
+    <member name="P:System.Web.Razor.RazorTemplateEngine.Host">
+      <summary>This type/member supports the .NET Framework infrastructure and is not intended to be used directly from your code.Gets or sets the <see cref="T:System.Web.Razor.RazorEngineHost" /> which defines the environment in which the generated template code will live.</summary>
+      <returns>The <see cref="T:System.Web.Razor.RazorEngineHost" /> which defines the environment in which the generated template code will live.</returns>
+    </member>
+    <member name="M:System.Web.Razor.RazorTemplateEngine.ParseTemplate(System.IO.TextReader)">
+      <summary>This type/member supports the .NET Framework infrastructure and is not intended to be used directly from your code.Parses the template specified by the TextBuffer and returns its result.</summary>
+      <returns>The resulting parse tree.</returns>
+      <param name="input">The input text to parse.</param>
+    </member>
+    <member name="M:System.Web.Razor.RazorTemplateEngine.ParseTemplate(System.IO.TextReader,System.Nullable{System.Threading.CancellationToken})">
+      <summary>This type/member supports the .NET Framework infrastructure and is not intended to be used directly from your code.Parses the template specified by the TextBuffer and returns its result.</summary>
+      <returns>The resulting parse tree.</returns>
+      <param name="input">The input text to parse.</param>
+      <param name="cancelToken">A token used to cancel the parser.</param>
+    </member>
+    <member name="M:System.Web.Razor.RazorTemplateEngine.ParseTemplate(System.Web.Razor.Text.ITextBuffer)">
+      <summary>This type/member supports the .NET Framework infrastructure and is not intended to be used directly from your code.Parses the template specified by the TextBuffer and returns its result.</summary>
+      <returns>The resulting parse tree.</returns>
+      <param name="input">The input text to parse.</param>
+    </member>
+    <member name="M:System.Web.Razor.RazorTemplateEngine.ParseTemplate(System.Web.Razor.Text.ITextBuffer,System.Nullable{System.Threading.CancellationToken})">
+      <summary>This type/member supports the .NET Framework infrastructure and is not intended to be used directly from your code.Parses the template specified by the TextBuffer and returns its result.</summary>
+      <returns>The resulting parse tree.</returns>
+      <param name="input">The input text to parse.</param>
+      <param name="cancelToken">A token used to cancel the parser.</param>
+    </member>
+    <member name="M:System.Web.Razor.RazorTemplateEngine.ParseTemplateCore(System.Web.Razor.Text.ITextDocument,System.Nullable{System.Threading.CancellationToken})">
+      <summary>This type/member supports the .NET Framework infrastructure and is not intended to be used directly from your code.Parses the template core.</summary>
+      <returns>The resulting parse tree.</returns>
+      <param name="input">The input text to parse.</param>
+      <param name="cancelToken">A token used to cancel the parser.</param>
+    </member>
+    <member name="T:System.Web.Razor.StateMachine`1">
+      <summary>This type/member supports the .NET Framework infrastructure and is not intended to be used directly from your code.Represents the state of the machine.</summary>
+      <typeparam name="TReturn">The generic type Return.</typeparam>
+    </member>
+    <member name="M:System.Web.Razor.StateMachine`1.#ctor">
+      <summary>This type/member supports the .NET Framework infrastructure and is not intended to be used directly from your code.Initializes a new instance of the <see cref="T:System.Web.Razor.StateMachine`1" /> class.</summary>
+    </member>
+    <member name="P:System.Web.Razor.StateMachine`1.CurrentState">
+      <summary>This type/member supports the .NET Framework infrastructure and is not intended to be used directly from your code.Gets or sets the current state of the machine.</summary>
+      <returns>The current state of the machine.</returns>
+    </member>
+    <member name="P:System.Web.Razor.StateMachine`1.StartState">
+      <summary>This type/member supports the .NET Framework infrastructure and is not intended to be used directly from your code.Gets the starting state of the machine.</summary>
+      <returns>The starting state of the machine.</returns>
+    </member>
+    <member name="M:System.Web.Razor.StateMachine`1.Stay">
+      <summary>This type/member supports the .NET Framework infrastructure and is not intended to be used directly from your code.Stays into the machine during the transition.</summary>
+      <returns>Transition of the state machine.</returns>
+    </member>
+    <member name="M:System.Web.Razor.StateMachine`1.Stay(`0)">
+      <summary>This type/member supports the .NET Framework infrastructure and is not intended to be used directly from your code.Stays into the machine during the transition with the specified output.</summary>
+      <returns>The output of the transition.</returns>
+      <param name="output">The output.</param>
+    </member>
+    <member name="M:System.Web.Razor.StateMachine`1.Stop">
+      <summary>This type/member supports the .NET Framework infrastructure and is not intended to be used directly from your code.Disables the machine upon transition.</summary>
+      <returns>The machine to stop.</returns>
+    </member>
+    <member name="M:System.Web.Razor.StateMachine`1.Transition(System.Web.Razor.StateMachine{`0}.State)">
+      <summary>This type/member supports the .NET Framework infrastructure and is not intended to be used directly from your code.Indicates the new transition of the state.</summary>
+      <returns>The new transition of the state.</returns>
+      <param name="newState">The new state.</param>
+    </member>
+    <member name="M:System.Web.Razor.StateMachine`1.Transition(`0,System.Web.Razor.StateMachine{`0}.State)">
+      <summary>This type/member supports the .NET Framework infrastructure and is not intended to be used directly from your code.Indicates the new transition of the state with the specified output.</summary>
+      <returns>The new transition of the state with the specified output.</returns>
+      <param name="output">The output.</param>
+      <param name="newState">The new state.</param>
+    </member>
+    <member name="M:System.Web.Razor.StateMachine`1.Turn">
+      <summary>This type/member supports the .NET Framework infrastructure and is not intended to be used directly from your code.Describes the turning process of the state.</summary>
+      <returns>The turning process of the state.</returns>
+    </member>
+    <member name="T:System.Web.Razor.StateMachine`1.State">
+      <summary>This type/member supports the .NET Framework infrastructure and is not intended to be used directly from your code.</summary>
+    </member>
+    <member name="T:System.Web.Razor.StateMachine`1.StateResult">
+      <summary>This type/member supports the .NET Framework infrastructure and is not intended to be used directly from your code.Represents the state result.</summary>
+    </member>
+    <member name="M:System.Web.Razor.StateMachine`1.StateResult.#ctor(System.Web.Razor.StateMachine{`0}.State)">
+      <summary>This type/member supports the .NET Framework infrastructure and is not intended to be used directly from your code.Initializes a new instance of the <see cref="T:System.Web.Razor.StateMachine`1.StateResult" /> class.</summary>
+      <param name="next">The next output.</param>
+    </member>
+    <member name="M:System.Web.Razor.StateMachine`1.StateResult.#ctor(`0,System.Web.Razor.StateMachine{`0}.State)">
+      <summary>This type/member supports the .NET Framework infrastructure and is not intended to be used directly from your code.Initializes a new instance of the <see cref="T:System.Web.Razor.StateMachine`1.StateResult" /> class.</summary>
+      <param name="output">The output.</param>
+      <param name="next">The next state.</param>
+    </member>
+    <member name="P:System.Web.Razor.StateMachine`1.StateResult.HasOutput">
+      <summary>This type/member supports the .NET Framework infrastructure and is not intended to be used directly from your code.Gets or sets a value indicating whether the state has output.</summary>
+      <returns>true if the state has output; otherwise, false.</returns>
+    </member>
+    <member name="P:System.Web.Razor.StateMachine`1.StateResult.Next">
+      <summary>This type/member supports the .NET Framework infrastructure and is not intended to be used directly from your code.Gets or sets the next state in the machine.</summary>
+      <returns>The next state in the machine.</returns>
+    </member>
+    <member name="P:System.Web.Razor.StateMachine`1.StateResult.Output">
+      <summary>This type/member supports the .NET Framework infrastructure and is not intended to be used directly from your code.Gets or sets the output.</summary>
+      <returns>The <see cref="T:System.Web.Razor.StateMachine`1.State" /> representing the output.</returns>
+    </member>
+    <member name="T:System.Web.Razor.VBRazorCodeLanguage">
+      <summary>This type/member supports the .NET Framework infrastructure and is not intended to be used directly from your code.Represents a language generator and provider of the VB razor code.</summary>
+    </member>
+    <member name="M:System.Web.Razor.VBRazorCodeLanguage.#ctor">
+      <summary>This type/member supports the .NET Framework infrastructure and is not intended to be used directly from your code.Initializes a new instance of the <see cref="T:System.Web.Razor.VBRazorCodeLanguage" /> class.</summary>
+    </member>
+    <member name="P:System.Web.Razor.VBRazorCodeLanguage.CodeDomProviderType">
+      <summary>This type/member supports the .NET Framework infrastructure and is not intended to be used directly from your code.Gets the type of CodeDomProvider.</summary>
+      <returns>The type of CodeDomProvider.</returns>
+    </member>
+    <member name="M:System.Web.Razor.VBRazorCodeLanguage.CreateCodeGenerator(System.String,System.String,System.String,System.Web.Razor.RazorEngineHost)">
+      <summary>This type/member supports the .NET Framework infrastructure and is not intended to be used directly from your code.Creates the code language generator.</summary>
+      <returns>The code language generator.</returns>
+      <param name="className">The name of the class.</param>
+      <param name="rootNamespaceName">The root namespace name.</param>
+      <param name="sourceFileName">The source File name.</param>
+      <param name="host">The <see cref="T:System.Web.Razor.RazorEngineHost" />.</param>
+    </member>
+    <member name="M:System.Web.Razor.VBRazorCodeLanguage.CreateCodeParser">
+      <summary>This type/member supports the .NET Framework infrastructure and is not intended to be used directly from your code.Creates a code parser in a <see cref="T:System.Web.Razor.Parser.ParserBase" />.</summary>
+      <returns>A code parser in a <see cref="T:System.Web.Razor.Parser.ParserBase" />.</returns>
+    </member>
+    <member name="P:System.Web.Razor.VBRazorCodeLanguage.LanguageName">
+      <summary>This type/member supports the .NET Framework infrastructure and is not intended to be used directly from your code.Gets the language name.</summary>
+      <returns>The language name.</returns>
+    </member>
+    <member name="T:System.Web.Razor.Editor.EditorHints">
+      <summary>This type/member supports the .NET Framework infrastructure and is not intended to be used directly from your code.</summary>
+    </member>
+    <member name="F:System.Web.Razor.Editor.EditorHints.LayoutPage">
+      <summary>This type/member supports the .NET Framework infrastructure and is not intended to be used directly from your code.</summary>
+    </member>
+    <member name="F:System.Web.Razor.Editor.EditorHints.None">
+      <summary>This type/member supports the .NET Framework infrastructure and is not intended to be used directly from your code.</summary>
+    </member>
+    <member name="F:System.Web.Razor.Editor.EditorHints.VirtualPath">
+      <summary>This type/member supports the .NET Framework infrastructure and is not intended to be used directly from your code.</summary>
+    </member>
+    <member name="T:System.Web.Razor.Editor.EditResult">
+      <summary>This type/member supports the .NET Framework infrastructure and is not intended to be used directly from your code.Represents the editing result of the Editor.</summary>
+    </member>
+    <member name="M:System.Web.Razor.Editor.EditResult.#ctor(System.Web.Razor.PartialParseResult,System.Web.Razor.Parser.SyntaxTree.SpanBuilder)">
+      <summary>This type/member supports the .NET Framework infrastructure and is not intended to be used directly from your code.Initializes a new instance of the <see cref="T:System.Web.Razor.Editor.EditResult" /> class.</summary>
+      <param name="result">The partial parse result.</param>
+      <param name="editedSpan">The edited span builder.</param>
+    </member>
+    <member name="P:System.Web.Razor.Editor.EditResult.EditedSpan">
+      <summary>This type/member supports the .NET Framework infrastructure and is not intended to be used directly from your code.Gets or sets the edited span of the <see cref="T:System.Web.Razor.Parser.SyntaxTree.SpanBuilder" />.</summary>
+      <returns>The edited span of the <see cref="T:System.Web.Razor.Parser.SyntaxTree.SpanBuilder" />.</returns>
+    </member>
+    <member name="P:System.Web.Razor.Editor.EditResult.Result">
+      <summary>This type/member supports the .NET Framework infrastructure and is not intended to be used directly from your code.Gets or sets the partial parse result.</summary>
+      <returns>The partial parse result.</returns>
+    </member>
+    <member name="T:System.Web.Razor.Editor.ImplicitExpressionEditHandler">
+      <summary>This type/member supports the .NET Framework infrastructure and is not intended to be used directly from your code.Provides edit handler for implicit expression.</summary>
+    </member>
+    <member name="M:System.Web.Razor.Editor.ImplicitExpressionEditHandler.#ctor(System.Func{System.String,System.Collections.Generic.IEnumerable{System.Web.Razor.Tokenizer.Symbols.ISymbol}},System.Collections.Generic.ISet{System.String},System.Boolean)">
+      <summary>Initializes a new instance of the <see cref="T:System.Web.Razor.Editor.ImplicitExpressionEditHandler" /> class.</summary>
+      <param name="tokenizer">The tokenizer.</param>
+      <param name="keywords">The keywords.</param>
+      <param name="acceptTrailingDot">true to accept trailing dot; otherwise, false.</param>
+    </member>
+    <member name="P:System.Web.Razor.Editor.ImplicitExpressionEditHandler.AcceptTrailingDot">
+      <summary>This type/member supports the .NET Framework infrastructure and is not intended to be used directly from your code.Gets or sets a value indicating whether the expression accepts trailing dot.</summary>
+      <returns>true if the expression accepts trailing dot; otherwise, false.</returns>
+    </member>
+    <member name="M:System.Web.Razor.Editor.ImplicitExpressionEditHandler.CanAcceptChange(System.Web.Razor.Parser.SyntaxTree.Span,System.Web.Razor.Text.TextChange)">
+      <summary>This type/member supports the .NET Framework infrastructure and is not intended to be used directly from your code.Indicates the parse that can accept change.</summary>
+      <returns>The partial parse result.</returns>
+      <param name="target">The target.</param>
+      <param name="normalizedChange">The normalized change.</param>
+    </member>
+    <member name="M:System.Web.Razor.Editor.ImplicitExpressionEditHandler.Equals(System.Object)">
+      <summary>This type/member supports the .NET Framework infrastructure and is not intended to be used directly from your code.Indicates whether the specified object is equal to the current object.</summary>
+      <returns>true if the specified object is equal to the current objet; otherwise, false.</returns>
+      <param name="obj">The object to compare to.</param>
+    </member>
+    <member name="M:System.Web.Razor.Editor.ImplicitExpressionEditHandler.GetHashCode">
+      <summary>This type/member supports the .NET Framework infrastructure and is not intended to be used directly from your code.Retrieves the hash code for this current instance.</summary>
+      <returns>The hash code for this current instance.</returns>
+    </member>
+    <member name="P:System.Web.Razor.Editor.ImplicitExpressionEditHandler.Keywords">
+      <summary>This type/member supports the .NET Framework infrastructure and is not intended to be used directly from your code.Gets or sets the keywords associated with the expression.</summary>
+      <returns>The keywords associated with the expression.</returns>
+    </member>
+    <member name="M:System.Web.Razor.Editor.ImplicitExpressionEditHandler.ToString">
+      <summary>This type/member supports the .NET Framework infrastructure and is not intended to be used directly from your code.Returns a string representation of this current instance.</summary>
+      <returns>A string representation of this current instance.</returns>
+    </member>
+    <member name="T:System.Web.Razor.Editor.SingleLineMarkupEditHandler">
+      <summary>This type/member supports the .NET Framework infrastructure and is not intended to be used directly from your code.Represents the handler editor for this webpages.</summary>
+    </member>
+    <member name="M:System.Web.Razor.Editor.SingleLineMarkupEditHandler.#ctor(System.Func{System.String,System.Collections.Generic.IEnumerable{System.Web.Razor.Tokenizer.Symbols.ISymbol}})">
+      <summary>Initializes a new instance of the <see cref="T:System.Web.Razor.Editor.SingleLineMarkupEditHandler" /> class.</summary>
+      <param name="tokenizer">The tokenizer symbols.</param>
+    </member>
+    <member name="M:System.Web.Razor.Editor.SingleLineMarkupEditHandler.#ctor(System.Func{System.String,System.Collections.Generic.IEnumerable{System.Web.Razor.Tokenizer.Symbols.ISymbol}},System.Web.Razor.Parser.SyntaxTree.AcceptedCharacters)">
+      <summary>Initializes a new instance of the <see cref="T:System.Web.Razor.Editor.SingleLineMarkupEditHandler" /> class.</summary>
+      <param name="tokenizer">The tokenizer symbols.</param>
+      <param name="accepted">The accepted characters.</param>
+    </member>
+    <member name="T:System.Web.Razor.Editor.SpanEditHandler">
+      <summary>This type/member supports the .NET Framework infrastructure and is not intended to be used directly from your code. Provides methods for handling the span edits.</summary>
+    </member>
+    <member name="M:System.Web.Razor.Editor.SpanEditHandler.#ctor(System.Func{System.String,System.Collections.Generic.IEnumerable{System.Web.Razor.Tokenizer.Symbols.ISymbol}})">
+      <summary>This type/member supports the .NET Framework infrastructure and is not intended to be used directly from your code. Initializes a new instance of the <see cref="T:System.Web.Razor.Editor.SpanEditHandler" /> class.</summary>
+      <param name="tokenizer">The method used to parse string into tokens.</param>
+    </member>
+    <member name="M:System.Web.Razor.Editor.SpanEditHandler.#ctor(System.Func{System.String,System.Collections.Generic.IEnumerable{System.Web.Razor.Tokenizer.Symbols.ISymbol}},System.Web.Razor.Parser.SyntaxTree.AcceptedCharacters)">
+      <summary>This type/member supports the .NET Framework infrastructure and is not intended to be used directly from your code. Initializes a new instance of the <see cref="T:System.Web.Razor.Editor.SpanEditHandler" /> class.</summary>
+      <param name="tokenizer">The method used to parse string into tokens.</param>
+      <param name="accepted">One of the values of the <see cref="T:System.Web.Razor.Parser.SyntaxTree.AcceptedCharacters" /> enumeration.</param>
+    </member>
+    <member name="P:System.Web.Razor.Editor.SpanEditHandler.AcceptedCharacters">
+      <summary>This type/member supports the .NET Framework infrastructure and is not intended to be used directly from your code. Gets or sets a value that specifies the accepted characters.</summary>
+      <returns>One of the values of the <see cref="T:System.Web.Razor.Parser.SyntaxTree.AcceptedCharacters" /> enumeration.</returns>
+    </member>
+    <member name="M:System.Web.Razor.Editor.SpanEditHandler.ApplyChange(System.Web.Razor.Parser.SyntaxTree.Span,System.Web.Razor.Text.TextChange)">
+      <summary>This type/member supports the .NET Framework infrastructure and is not intended to be used directly from your code. Applies the text change to the span.</summary>
+      <returns>The result of the apply operation.</returns>
+      <param name="target">The span to apply changes to.</param>
+      <param name="change">The change to apply.</param>
+    </member>
+    <member name="M:System.Web.Razor.Editor.SpanEditHandler.ApplyChange(System.Web.Razor.Parser.SyntaxTree.Span,System.Web.Razor.Text.TextChange,System.Boolean)">
+      <summary>This type/member supports the .NET Framework infrastructure and is not intended to be used directly from your code. Applies the text change to the span.</summary>
+      <returns>The result of the apply operation.</returns>
+      <param name="target">The span to apply changes to.</param>
+      <param name="change">The change to apply.</param>
+      <param name="force">true to accept partial result; otherwise, false.</param>
+    </member>
+    <member name="M:System.Web.Razor.Editor.SpanEditHandler.CanAcceptChange(System.Web.Razor.Parser.SyntaxTree.Span,System.Web.Razor.Text.TextChange)">
+      <summary>This type/member supports the .NET Framework infrastructure and is not intended to be used directly from your code. Determines whether the span can accept the specified change.</summary>
+      <returns>true if the span can accept the specified change; otherwise, false.</returns>
+      <param name="target">The span to check.</param>
+      <param name="normalizedChange">The change to apply.</param>
+    </member>
+    <member name="M:System.Web.Razor.Editor.SpanEditHandler.CreateDefault">
+      <summary>This type/member supports the .NET Framework infrastructure and is not intended to be used directly from your code. Creates a new default span edit handler.</summary>
+      <returns>A newly created default span edit handler.</returns>
+    </member>
+    <member name="M:System.Web.Razor.Editor.SpanEditHandler.CreateDefault(System.Func{System.String,System.Collections.Generic.IEnumerable{System.Web.Razor.Tokenizer.Symbols.ISymbol}})">
+      <summary>This type/member supports the .NET Framework infrastructure and is not intended to be used directly from your code. Creates a new default span edit handler.</summary>
+      <returns>A newly created default span edit handler.</returns>
+      <param name="tokenizer">The method used to parse string into tokens.</param>
+    </member>
+    <member name="P:System.Web.Razor.Editor.SpanEditHandler.EditorHints">
+      <summary>This type/member supports the .NET Framework infrastructure and is not intended to be used directly from your code. Gets or sets the editor hints.</summary>
+      <returns>The editor hints.</returns>
+    </member>
+    <member name="M:System.Web.Razor.Editor.SpanEditHandler.Equals(System.Object)">
+      <summary>This type/member supports the .NET Framework infrastructure and is not intended to be used directly from your code. Determines whether this instance is equal to a specified object.</summary>
+      <returns>true if the object is equal to the this instance; otherwise, false.</returns>
+      <param name="obj">The object to compare with this instance.</param>
+    </member>
+    <member name="M:System.Web.Razor.Editor.SpanEditHandler.GetHashCode">
+      <summary>This type/member supports the .NET Framework infrastructure and is not intended to be used directly from your code. Returns the hash code for this instance.</summary>
+      <returns>The hash code for this instance.</returns>
+    </member>
+    <member name="M:System.Web.Razor.Editor.SpanEditHandler.GetOldText(System.Web.Razor.Parser.SyntaxTree.Span,System.Web.Razor.Text.TextChange)">
+      <summary>This type/member supports the .NET Framework infrastructure and is not intended to be used directly from your code. Gets the old text from the span content.</summary>
+      <returns>The old text from the span content.</returns>
+      <param name="target">The span to get old text from.</param>
+      <param name="change">The text change which contains the location of the old text.</param>
+    </member>
+    <member name="M:System.Web.Razor.Editor.SpanEditHandler.IsAtEndOfFirstLine(System.Web.Razor.Parser.SyntaxTree.Span,System.Web.Razor.Text.TextChange)">
+      <summary>This type/member supports the .NET Framework infrastructure and is not intended to be used directly from your code. Determines whether the specified change is at the end of first line of the span content.</summary>
+      <returns>true if the specified change is at the end of first line of the span content; otherwise, false.</returns>
+      <param name="target">The span to check.</param>
+      <param name="change">The change to check.</param>
+    </member>
+    <member name="M:System.Web.Razor.Editor.SpanEditHandler.IsAtEndOfSpan(System.Web.Razor.Parser.SyntaxTree.Span,System.Web.Razor.Text.TextChange)">
+      <summary>This type/member supports the .NET Framework infrastructure and is not intended to be used directly from your code. Determines whether the specified change is at the end of the span.</summary>
+      <returns>true if the specified change is at the end of the span; otherwise, false.</returns>
+      <param name="target">The span to check.</param>
+      <param name="change">The change to chek.</param>
+    </member>
+    <member name="M:System.Web.Razor.Editor.SpanEditHandler.IsEndDeletion(System.Web.Razor.Parser.SyntaxTree.Span,System.Web.Razor.Text.TextChange)">
+      <summary>This type/member supports the .NET Framework infrastructure and is not intended to be used directly from your code. Determines whether the specified change is at the end the span content and for deletion.</summary>
+      <returns>true if the specified change is at the end the span content and for deletion; otherwise, false.</returns>
+      <param name="target">The span to check.</param>
+      <param name="change">The change to check.</param>
+    </member>
+    <member name="M:System.Web.Razor.Editor.SpanEditHandler.IsEndInsertion(System.Web.Razor.Parser.SyntaxTree.Span,System.Web.Razor.Text.TextChange)">
+      <summary>This type/member supports the .NET Framework infrastructure and is not intended to be used directly from your code. Determines whether the specified change is at the end the span content and for insertion.</summary>
+      <returns>true if the specified change is at the end the span content and for insertion; otherwise, false.</returns>
+      <param name="target">The span to check.</param>
+      <param name="change">The change to check.</param>
+    </member>
+    <member name="M:System.Web.Razor.Editor.SpanEditHandler.IsEndReplace(System.Web.Razor.Parser.SyntaxTree.Span,System.Web.Razor.Text.TextChange)">
+      <summary>This type/member supports the .NET Framework infrastructure and is not intended to be used directly from your code. Determines whether the specified change is at the end the span content and for replacement.</summary>
+      <returns>true if the specified change is at the end the span content and for replacement; otherwise, false.</returns>
+      <param name="target">The span to check.</param>
+      <param name="change">The change to check.</param>
+    </member>
+    <member name="M:System.Web.Razor.Editor.SpanEditHandler.OwnsChange(System.Web.Razor.Parser.SyntaxTree.Span,System.Web.Razor.Text.TextChange)">
+      <summary>This type/member supports the .NET Framework infrastructure and is not intended to be used directly from your code. Determines whether the span owns the specified change.</summary>
+      <returns>true if the span owns the specified change; otherwise, false.</returns>
+      <param name="target">The span to check.</param>
+      <param name="change">The change to check.</param>
+    </member>
+    <member name="P:System.Web.Razor.Editor.SpanEditHandler.Tokenizer">
+      <summary>This type/member supports the .NET Framework infrastructure and is not intended to be used directly from your code. Gets or sets the method used to parse string into tokens. </summary>
+      <returns>The method used to parse string into tokens.</returns>
+    </member>
+    <member name="M:System.Web.Razor.Editor.SpanEditHandler.ToString">
+      <summary>This type/member supports the .NET Framework infrastructure and is not intended to be used directly from your code. Returns the string representation of the span edit handler.</summary>
+      <returns>The string representation of the span edit handler.</returns>
+    </member>
+    <member name="M:System.Web.Razor.Editor.SpanEditHandler.UpdateSpan(System.Web.Razor.Parser.SyntaxTree.Span,System.Web.Razor.Text.TextChange)">
+      <summary>This type/member supports the .NET Framework infrastructure and is not intended to be used directly from your code. Updates the span using the normalized change.</summary>
+      <returns>The new span builder for the specified target.</returns>
+      <param name="target">The span to update.</param>
+      <param name="normalizedChange">The normalized change.</param>
+    </member>
+    <member name="T:System.Web.Razor.Generator.AddImportCodeGenerator">
+      <summary>This type/member supports the .NET Framework infrastructure and is not intended to be used directly from your code.Represents the added import code generator for the razor.</summary>
+    </member>
+    <member name="M:System.Web.Razor.Generator.AddImportCodeGenerator.#ctor(System.String,System.Int32)">
+      <summary>This type/member supports the .NET Framework infrastructure and is not intended to be used directly from your code.Initializes a new instance of the <see cref="T:System.Web.Razor.Generator.AddImportCodeGenerator" /> class.</summary>
+      <param name="ns">The string namespace.</param>
+      <param name="namespaceKeywordLength">The length of the keyword namespace.</param>
+    </member>
+    <member name="M:System.Web.Razor.Generator.AddImportCodeGenerator.Equals(System.Object)">
+      <summary>This type/member supports the .NET Framework infrastructure and is not intended to be used directly from your code.Determines whether two object instances are equal.</summary>
+      <returns>true if the specified object is equal to the current object; otherwise, false.</returns>
+      <param name="obj">The object to compare with the current object.</param>
+    </member>
+    <member name="M:System.Web.Razor.Generator.AddImportCodeGenerator.GenerateCode(System.Web.Razor.Parser.SyntaxTree.Span,System.Web.Razor.Generator.CodeGeneratorContext)">
+      <summary>This type/member supports the .NET Framework infrastructure and is not intended to be used directly from your code.Generates code with the specified parameters using the added import code generator.</summary>
+      <param name="target">The target span.</param>
+      <param name="context">The code generator context.</param>
+    </member>
+    <member name="M:System.Web.Razor.Generator.AddImportCodeGenerator.GetHashCode">
+      <summary>This type/member supports the .NET Framework infrastructure and is not intended to be used directly from your code.Returns the hash code for this instance.</summary>
+      <returns>The hash code for this instance.</returns>
+    </member>
+    <member name="P:System.Web.Razor.Generator.AddImportCodeGenerator.Namespace">
+      <summary>This type/member supports the .NET Framework infrastructure and is not intended to be used directly from your code.Gets the string namespace of the generator to add import code generator.</summary>
+      <returns>The string namespace of the generator to add import code generator.</returns>
+    </member>
+    <member name="P:System.Web.Razor.Generator.AddImportCodeGenerator.NamespaceKeywordLength">
+      <summary>This type/member supports the .NET Framework infrastructure and is not intended to be used directly from your code.Gets or sets the length of keyword namespace for the code generator.</summary>
+      <returns>The length of keyword namespace for the code generator.</returns>
+    </member>
+    <member name="M:System.Web.Razor.Generator.AddImportCodeGenerator.ToString">
+      <summary>This type/member supports the .NET Framework infrastructure and is not intended to be used directly from your code.Returns a string that represents the current object.</summary>
+      <returns>A string that represents the current object.</returns>
+    </member>
+    <member name="T:System.Web.Razor.Generator.AttributeBlockCodeGenerator">
+      <summary>This type/member supports the .NET Framework infrastructure and is not intended to be used directly from your code.Represents the attributes of the block code generator.</summary>
+    </member>
+    <member name="M:System.Web.Razor.Generator.AttributeBlockCodeGenerator.#ctor(System.String,System.Web.Razor.Text.LocationTagged{System.String},System.Web.Razor.Text.LocationTagged{System.String})">
+      <summary>This type/member supports the .NET Framework infrastructure and is not intended to be used directly from your code.Initializes a new instance of the <see cref="T:System.Web.Razor.Generator.AttributeBlockCodeGenerator" /> class.</summary>
+      <param name="name">The name.</param>
+      <param name="prefix">The prefix string.</param>
+      <param name="suffix">The suffix string.</param>
+    </member>
+    <member name="M:System.Web.Razor.Generator.AttributeBlockCodeGenerator.Equals(System.Object)">
+      <summary>This type/member supports the .NET Framework infrastructure and is not intended to be used directly from your code.Determines whether the specified object is equal to the current object.</summary>
+      <returns>true if the specified object is equal to the current object; otherwise, false.</returns>
+      <param name="obj">The object to compare with the current object.</param>
+    </member>
+    <member name="M:System.Web.Razor.Generator.AttributeBlockCodeGenerator.GenerateEndBlockCode(System.Web.Razor.Parser.SyntaxTree.Block,System.Web.Razor.Generator.CodeGeneratorContext)">
+      <summary>This type/member supports the .NET Framework infrastructure and is not intended to be used directly from your code.Generates code to end the block using the specified parameters.</summary>
+      <param name="target">The target block.</param>
+      <param name="context">The code generator context.</param>
+    </member>
+    <member name="M:System.Web.Razor.Generator.AttributeBlockCodeGenerator.GenerateStartBlockCode(System.Web.Razor.Parser.SyntaxTree.Block,System.Web.Razor.Generator.CodeGeneratorContext)">
+      <summary>This type/member supports the .NET Framework infrastructure and is not intended to be used directly from your code.Generates code to start the block using the specified parameters.</summary>
+      <param name="target">The target block.</param>
+      <param name="context">The code generator context.</param>
+    </member>
+    <member name="M:System.Web.Razor.Generator.AttributeBlockCodeGenerator.GetHashCode">
+      <summary>This type/member supports the .NET Framework infrastructure and is not intended to be used directly from your code.Returns the hash code for this code generator.</summary>
+      <returns>The hash code for this code generator.</returns>
+    </member>
+    <member name="P:System.Web.Razor.Generator.AttributeBlockCodeGenerator.Name">
+      <summary>This type/member supports the .NET Framework infrastructure and is not intended to be used directly from your code.Gets the string name of the <see cref="T:System.Web.Razor.Generator.AttributeBlockCodeGenerator" />.</summary>
+      <returns>The string name of the <see cref="T:System.Web.Razor.Generator.AttributeBlockCodeGenerator" />.</returns>
+    </member>
+    <member name="P:System.Web.Razor.Generator.AttributeBlockCodeGenerator.Prefix">
+      <summary>This type/member supports the .NET Framework infrastructure and is not intended to be used directly from your code.Gets the prefix of the code generator.</summary>
+      <returns>The prefix of the code generator.</returns>
+    </member>
+    <member name="P:System.Web.Razor.Generator.AttributeBlockCodeGenerator.Suffix">
+      <summary>This type/member supports the .NET Framework infrastructure and is not intended to be used directly from your code.Gets the suffix for the code generator.</summary>
+      <returns>The suffix for the code generator.</returns>
+    </member>
+    <member name="M:System.Web.Razor.Generator.AttributeBlockCodeGenerator.ToString">
+      <summary>This type/member supports the .NET Framework infrastructure and is not intended to be used directly from your code.Returns a string that represents the current object.</summary>
+      <returns>A string that represents the current object.</returns>
+    </member>
+    <member name="T:System.Web.Razor.Generator.BlockCodeGenerator">
+      <summary>This type/member supports the .NET Framework infrastructure and is not intended to be used directly from your code.Represent the block code generator for this razor syntax.</summary>
+    </member>
+    <member name="M:System.Web.Razor.Generator.BlockCodeGenerator.#ctor">
+      <summary>This type/member supports the .NET Framework infrastructure and is not intended to be used directly from your code.Initializes a new instance of the <see cref="T:System.Web.Razor.Generator.BlockCodeGenerator" /> class.</summary>
+    </member>
+    <member name="M:System.Web.Razor.Generator.BlockCodeGenerator.Equals(System.Object)">
+      <summary>This type/member supports the .NET Framework infrastructure and is not intended to be used directly from your code.Determines whether the specified object is equal to the current object.</summary>
+      <returns>true if the specified object is equal to the current object; otherwise, false.</returns>
+      <param name="obj">The object to compare with the current object.</param>
+    </member>
+    <member name="M:System.Web.Razor.Generator.BlockCodeGenerator.GenerateEndBlockCode(System.Web.Razor.Parser.SyntaxTree.Block,System.Web.Razor.Generator.CodeGeneratorContext)">
+      <summary>This type/member supports the .NET Framework infrastructure and is not intended to be used directly from your code.Generates the end of the block code generator for this razor syntax.</summary>
+      <param name="target">The target block.</param>
+      <param name="context">The code generator context.</param>
+    </member>
+    <member name="M:System.Web.Razor.Generator.BlockCodeGenerator.GenerateStartBlockCode(System.Web.Razor.Parser.SyntaxTree.Block,System.Web.Razor.Generator.CodeGeneratorContext)">
+      <summary>This type/member supports the .NET Framework infrastructure and is not intended to be used directly from your code.Generates the start of the block code generator for this razor syntax.</summary>
+      <param name="target">The target block.</param>
+      <param name="context">The code generator context.</param>
+    </member>
+    <member name="M:System.Web.Razor.Generator.BlockCodeGenerator.GetHashCode">
+      <summary>This type/member supports the .NET Framework infrastructure and is not intended to be used directly from your code.Returns a hash code for the block code generator.</summary>
+      <returns>A hash code for the block code generator.</returns>
+    </member>
+    <member name="F:System.Web.Razor.Generator.BlockCodeGenerator.Null">
+      <summary>This type/member supports the .NET Framework infrastructure and is not intended to be used directly from your code.Represents a null value for the block code generator.</summary>
+    </member>
+    <member name="T:System.Web.Razor.Generator.CodeGenerationCompleteEventArgs">
+      <summary>This type/member supports the .NET Framework infrastructure and is not intended to be used directly from your code.Represents the completion of event arguments for the code generation.</summary>
+    </member>
+    <member name="M:System.Web.Razor.Generator.CodeGenerationCompleteEventArgs.#ctor(System.String,System.String,System.CodeDom.CodeCompileUnit)">
+      <summary>This type/member supports the .NET Framework infrastructure and is not intended to be used directly from your code.Initializes a new instance of the <see cref="T:System.Web.Razor.Generator.CodeGenerationCompleteEventArgs" /> class.</summary>
+      <param name="virtualPath">The virtual path string.</param>
+      <param name="physicalPath">The physical path string.</param>
+      <param name="generatedCode">The generated code compile unit.</param>
+    </member>
+    <member name="P:System.Web.Razor.Generator.CodeGenerationCompleteEventArgs.GeneratedCode">
+      <summary>This type/member supports the .NET Framework infrastructure and is not intended to be used directly from your code.Gets the generated code to complete the event argument.</summary>
+      <returns>The generated code to complete the event argument.</returns>
+    </member>
+    <member name="P:System.Web.Razor.Generator.CodeGenerationCompleteEventArgs.PhysicalPath">
+      <summary>This type/member supports the .NET Framework infrastructure and is not intended to be used directly from your code.Gets the physical path for the code generation.</summary>
+      <returns>The physical path for the code generation.</returns>
+    </member>
+    <member name="P:System.Web.Razor.Generator.CodeGenerationCompleteEventArgs.VirtualPath">
+      <summary>This type/member supports the .NET Framework infrastructure and is not intended to be used directly from your code.Gets the virtual path of the code generation.</summary>
+      <returns>The virtual path of the code generation.</returns>
+    </member>
+    <member name="T:System.Web.Razor.Generator.CodeGeneratorContext">
+      <summary>This type/member supports the .NET Framework infrastructure and is not intended to be used directly from your code. Represents the context of the code generator.</summary>
+    </member>
+    <member name="M:System.Web.Razor.Generator.CodeGeneratorContext.AddCodeMapping(System.Web.Razor.Text.SourceLocation,System.Int32,System.Int32)">
+      <summary>This type/member supports the .NET Framework infrastructure and is not intended to be used directly from your code. Adds a new generated code mapping to the collection.</summary>
+      <returns>The collection index of the newly added code mapping.</returns>
+      <param name="sourceLocation">The source location of the generated code mapping.</param>
+      <param name="generatedCodeStart">The code start of the generated code mapping.</param>
+      <param name="generatedCodeLength">The length of the generated code mapping.</param>
+    </member>
+    <member name="M:System.Web.Razor.Generator.CodeGeneratorContext.AddContextCall(System.Web.Razor.Parser.SyntaxTree.Span,System.String,System.Boolean)">
+      <summary>This type/member supports the .NET Framework infrastructure and is not intended to be used directly from your code. Adds a code statement for a context call on the specified method.</summary>
+      <param name="contentSpan">The content span.</param>
+      <param name="methodName">The name of the method to invoke a context call.</param>
+      <param name="isLiteral">true to specify that the method parameter is literal; otherwise, false.</param>
+    </member>
+    <member name="M:System.Web.Razor.Generator.CodeGeneratorContext.AddDesignTimeHelperStatement(System.CodeDom.CodeSnippetStatement)">
+      <summary>This type/member supports the .NET Framework infrastructure and is not intended to be used directly from your code. Adds a code statement that inserts the Razor design time helpers method in the specified code statement.</summary>
+      <param name="statement">The code statement that receives the code insertion.</param>
+    </member>
+    <member name="M:System.Web.Razor.Generator.CodeGeneratorContext.AddStatement(System.String)">
+      <summary>This type/member supports the .NET Framework infrastructure and is not intended to be used directly from your code. Adds the specified code statement to the body of the target method.</summary>
+      <param name="generatedCode">The code statement to add the target method.</param>
+    </member>
+    <member name="M:System.Web.Razor.Generator.CodeGeneratorContext.AddStatement(System.String,System.CodeDom.CodeLinePragma)">
+      <summary>This type/member supports the .NET Framework infrastructure and is not intended to be used directly from your code. Adds the specified code statement to the body of the target method.</summary>
+      <param name="body">The code statement to add the target method.</param>
+      <param name="pragma">The line pragma.</param>
+    </member>
+    <member name="M:System.Web.Razor.Generator.CodeGeneratorContext.BufferStatementFragment(System.String)">
+      <summary>This type/member supports the .NET Framework infrastructure and is not intended to be used directly from your code. Appends the specified fragment to the current buffered statement.</summary>
+      <param name="fragment">The fragment to add.</param>
+    </member>
+    <member name="M:System.Web.Razor.Generator.CodeGeneratorContext.BufferStatementFragment(System.String,System.Web.Razor.Parser.SyntaxTree.Span)">
+      <summary>This type/member supports the .NET Framework infrastructure and is not intended to be used directly from your code. Appends the specified fragment to the current buffered statement.</summary>
+      <param name="fragment">The fragment to add.</param>
+      <param name="sourceSpan">The source span for the <paramref name="fragment" />.</param>
+    </member>
+    <member name="M:System.Web.Razor.Generator.CodeGeneratorContext.BufferStatementFragment(System.Web.Razor.Parser.SyntaxTree.Span)">
+      <summary>This type/member supports the .NET Framework infrastructure and is not intended to be used directly from your code. Appends the content of the span to the current buffered statement.</summary>
+      <param name="sourceSpan">The source span whose content is to be added.</param>
+    </member>
+    <member name="M:System.Web.Razor.Generator.CodeGeneratorContext.ChangeStatementCollector(System.Action{System.String,System.CodeDom.CodeLinePragma})">
+      <summary>Assigns a new statement collector and returns a disposable action that restores the old statement collector.</summary>
+      <returns>A disposable action that restores the old statement collector.</returns>
+      <param name="collector">The new statement collector.</param>
+    </member>
+    <member name="P:System.Web.Razor.Generator.CodeGeneratorContext.CodeMappings">
+      <summary>This type/member supports the .NET Framework infrastructure and is not intended to be used directly from your code. Gets or sets the dictionary collection of generated code mapping.</summary>
+      <returns>The dictionary collection of generated code mapping.</returns>
+    </member>
+    <member name="P:System.Web.Razor.Generator.CodeGeneratorContext.CompileUnit">
+      <summary>This type/member supports the .NET Framework infrastructure and is not intended to be used directly from your code. Gets or sets the code compile unit that will hold the program graph.</summary>
+      <returns>The code compile unit that will hold the program graph.</returns>
+    </member>
+    <member name="M:System.Web.Razor.Generator.CodeGeneratorContext.Create(System.Web.Razor.RazorEngineHost,System.String,System.String,System.String,System.Boolean)">
+      <summary>This type/member supports the .NET Framework infrastructure and is not intended to be used directly from your code. Creates a new instance of the <see cref="T:System.Web.Razor.Generator.CodeGeneratorContext" /> class.</summary>
+      <returns>The newly created instance of the code generator context.</returns>
+      <param name="host">The Razor engine host.</param>
+      <param name="className">The class name for the generated class type declaration.</param>
+      <param name="rootNamespace">The name for the generated namespace declaration.</param>
+      <param name="sourceFile">The source file.</param>
+      <param name="shouldGenerateLinePragmas">true to enable the generation of line pragmas; otherwise, false.</param>
+    </member>
+    <member name="P:System.Web.Razor.Generator.CodeGeneratorContext.CurrentBufferedStatement">
+      <summary>This type/member supports the .NET Framework infrastructure and is not intended to be used directly from your code. Gets the current buffered statement.</summary>
+      <returns>The current buffered statement.</returns>
+    </member>
+    <member name="M:System.Web.Razor.Generator.CodeGeneratorContext.EnsureExpressionHelperVariable">
+      <summary>This type/member supports the .NET Framework infrastructure and is not intended to be used directly from your code. Adds the expression helper variable to the generated class if not yet added,</summary>
+    </member>
+    <member name="M:System.Web.Razor.Generator.CodeGeneratorContext.FlushBufferedStatement">
+      <summary>This type/member supports the .NET Framework infrastructure and is not intended to be used directly from your code. Flushes the current buffered statement.</summary>
+    </member>
+    <member name="P:System.Web.Razor.Generator.CodeGeneratorContext.GeneratedClass">
+      <summary>This type/member supports the .NET Framework infrastructure and is not intended to be used directly from your code. Gets or sets the generated class type declaration.</summary>
+      <returns>The generated class type declaration.</returns>
+    </member>
+    <member name="M:System.Web.Razor.Generator.CodeGeneratorContext.GenerateLinePragma(System.Web.Razor.Parser.SyntaxTree.Span)">
+      <summary>This type/member supports the .NET Framework infrastructure and is not intended to be used directly from your code. Generates the line pragma for the specified source.</summary>
+      <returns>The line pragma for the specified source.</returns>
+      <param name="target">The source span.</param>
+    </member>
+    <member name="M:System.Web.Razor.Generator.CodeGeneratorContext.GenerateLinePragma(System.Web.Razor.Parser.SyntaxTree.Span,System.Int32)">
+      <summary>This type/member supports the .NET Framework infrastructure and is not intended to be used directly from your code. Generates the line pragma for the source.</summary>
+      <returns>The line pragma for the specified source.</returns>
+      <param name="target">The source span.</param>
+      <param name="generatedCodeStart">The start index of code.</param>
+    </member>
+    <member name="M:System.Web.Razor.Generator.CodeGeneratorContext.GenerateLinePragma(System.Web.Razor.Parser.SyntaxTree.Span,System.Int32,System.Int32)">
+      <summary>This type/member supports the .NET Framework infrastructure and is not intended to be used directly from your code. Generates the line pragma for the source.</summary>
+      <returns>The line pragma for the specified source.</returns>
+      <param name="target">The source span.</param>
+      <param name="generatedCodeStart">The start index of code.</param>
+      <param name="codeLength">The length of code.</param>
+    </member>
+    <member name="M:System.Web.Razor.Generator.CodeGeneratorContext.GenerateLinePragma(System.Web.Razor.Text.SourceLocation,System.Int32,System.Int32)">
+      <summary>This type/member supports the .NET Framework infrastructure and is not intended to be used directly from your code. Generates the line pragma for the source.</summary>
+      <returns>The line pragma for the specified source.</returns>
+      <param name="start">The source location.</param>
+      <param name="generatedCodeStart">The start index of code.</param>
+      <param name="codeLength">The length of code.</param>
+    </member>
+    <member name="P:System.Web.Razor.Generator.CodeGeneratorContext.Host">
+      <summary>This type/member supports the .NET Framework infrastructure and is not intended to be used directly from your code. Gets or sets the Razor engine host.</summary>
+      <returns>The Razor engine host.</returns>
+    </member>
+    <member name="M:System.Web.Razor.Generator.CodeGeneratorContext.MarkEndOfGeneratedCode">
+      <summary>This type/member supports the .NET Framework infrastructure and is not intended to be used directly from your code. Marks the end of generated code.</summary>
+    </member>
+    <member name="M:System.Web.Razor.Generator.CodeGeneratorContext.MarkStartOfGeneratedCode">
+      <summary>This type/member supports the .NET Framework infrastructure and is not intended to be used directly from your code. Marks the start of generated code.</summary>
+    </member>
+    <member name="P:System.Web.Razor.Generator.CodeGeneratorContext.Namespace">
+      <summary>This type/member supports the .NET Framework infrastructure and is not intended to be used directly from your code. Gets or sets the generated namespace declaration.</summary>
+      <returns>The generated namespace declaration.</returns>
+    </member>
+    <member name="P:System.Web.Razor.Generator.CodeGeneratorContext.SourceFile">
+      <summary>This type/member supports the .NET Framework infrastructure and is not intended to be used directly from your code. Gets or sets the source file.</summary>
+      <returns>The source file.</returns>
+    </member>
+    <member name="P:System.Web.Razor.Generator.CodeGeneratorContext.TargetMethod">
+      <summary>This type/member supports the .NET Framework infrastructure and is not intended to be used directly from your code. Gets or sets the generated member method.</summary>
+      <returns>The generated member method.</returns>
+    </member>
+    <member name="P:System.Web.Razor.Generator.CodeGeneratorContext.TargetWriterName">
+      <summary>This type/member supports the .NET Framework infrastructure and is not intended to be used directly from your code. Gets or sets the name of text writer.</summary>
+      <returns>The name of text writer.</returns>
+    </member>
+    <member name="T:System.Web.Razor.Generator.CSharpRazorCodeGenerator">
+      <summary>This type/member supports the .NET Framework infrastructure and is not intended to be used directly from your code. Represents a Razor code generator for C# language.</summary>
+    </member>
+    <member name="M:System.Web.Razor.Generator.CSharpRazorCodeGenerator.#ctor(System.String,System.String,System.String,System.Web.Razor.RazorEngineHost)">
+      <summary>This type/member supports the .NET Framework infrastructure and is not intended to be used directly from your code. Initializes a new instance of the <see cref="T:System.Web.Razor.Generator.CSharpRazorCodeGenerator" /> class.</summary>
+      <param name="className">The class name for the generated class type declaration.</param>
+      <param name="rootNamespaceName">The name for the generated namespace declaration.</param>
+      <param name="sourceFileName">The source file.</param>
+      <param name="host">The Razor engine host.</param>
+    </member>
+    <member name="M:System.Web.Razor.Generator.CSharpRazorCodeGenerator.Initialize(System.Web.Razor.Generator.CodeGeneratorContext)">
+      <summary>This type/member supports the .NET Framework infrastructure and is not intended to be used directly from your code. Initializes the context for this code generator.</summary>
+      <param name="context">The context for this code generator.</param>
+    </member>
+    <member name="T:System.Web.Razor.Generator.DynamicAttributeBlockCodeGenerator">
+      <summary>This type/member supports the .NET Framework infrastructure and is not intended to be used directly from your code.Represents the dynamic attributes of the block code generator.</summary>
+    </member>
+    <member name="M:System.Web.Razor.Generator.DynamicAttributeBlockCodeGenerator.#ctor(System.Web.Razor.Text.LocationTagged{System.String},System.Int32,System.Int32,System.Int32)">
+      <summary>This type/member supports the .NET Framework infrastructure and is not intended to be used directly from your code.Initializes a new instances of the <see cref="T:System.Web.Razor.Generator.DynamicAttributeBlockCodeGenerator" /> class.</summary>
+      <param name="prefix">The prefix.</param>
+      <param name="offset">The offset values.</param>
+      <param name="line">The line values.</param>
+      <param name="col">The col.</param>
+    </member>
+    <member name="M:System.Web.Razor.Generator.DynamicAttributeBlockCodeGenerator.#ctor(System.Web.Razor.Text.LocationTagged{System.String},System.Web.Razor.Text.SourceLocation)">
+      <summary>This type/member supports the .NET Framework infrastructure and is not intended to be used directly from your code.Initializes a new instances of the <see cref="T:System.Web.Razor.Generator.DynamicAttributeBlockCodeGenerator" /> class.</summary>
+      <param name="prefix">The string prefix.</param>
+      <param name="valueStart">The value start.</param>
+    </member>
+    <member name="M:System.Web.Razor.Generator.DynamicAttributeBlockCodeGenerator.Equals(System.Object)">
+      <summary>This type/member supports the .NET Framework infrastructure and is not intended to be used directly from your code.Determines whether the specified object is equal to the current object.</summary>
+      <returns>true if the specified object is equal to the current object; otherwise, false.</returns>
+      <param name="obj">The object to compare with the current object.</param>
+    </member>
+    <member name="M:System.Web.Razor.Generator.DynamicAttributeBlockCodeGenerator.GenerateEndBlockCode(System.Web.Razor.Parser.SyntaxTree.Block,System.Web.Razor.Generator.CodeGeneratorContext)">
+      <summary>This type/member supports the .NET Framework infrastructure and is not intended to be used directly from your code.Generates code to end the block using the specified parameters.</summary>
+      <param name="target">The target block.</param>
+      <param name="context">The code generator context.</param>
+    </member>
+    <member name="M:System.Web.Razor.Generator.DynamicAttributeBlockCodeGenerator.GenerateStartBlockCode(System.Web.Razor.Parser.SyntaxTree.Block,System.Web.Razor.Generator.CodeGeneratorContext)">
+      <summary>This type/member supports the .NET Framework infrastructure and is not intended to be used directly from your code.Generates code to start the block using the specified parameters.</summary>
+      <param name="target">The target block.</param>
+      <param name="context">The code generator context.</param>
+    </member>
+    <member name="M:System.Web.Razor.Generator.DynamicAttributeBlockCodeGenerator.GetHashCode">
+      <summary>This type/member supports the .NET Framework infrastructure and is not intended to be used directly from your code.Returns the hash code for this instance.</summary>
+      <returns>The hash code for this instance.</returns>
+    </member>
+    <member name="P:System.Web.Razor.Generator.DynamicAttributeBlockCodeGenerator.Prefix">
+      <summary>This type/member supports the .NET Framework infrastructure and is not intended to be used directly from your code.Gets the namespace prefix of the code generator.</summary>
+      <returns>The namespace prefix of the code generator.</returns>
+    </member>
+    <member name="M:System.Web.Razor.Generator.DynamicAttributeBlockCodeGenerator.ToString">
+      <summary>This type/member supports the .NET Framework infrastructure and is not intended to be used directly from your code.Returns a string that represents the current object.</summary>
+      <returns>A string that represents the current object.</returns>
+    </member>
+    <member name="P:System.Web.Razor.Generator.DynamicAttributeBlockCodeGenerator.ValueStart">
+      <summary>This type/member supports the .NET Framework infrastructure and is not intended to be used directly from your code.Gets the value start for the dynamic attribute block code generator.</summary>
+      <returns>The value start for the dynamic attribute block code generator.</returns>
+    </member>
+    <member name="T:System.Web.Razor.Generator.ExpressionCodeGenerator">
+      <summary>This type/member supports the .NET Framework infrastructure and is not intended to be used directly from your code.Represents a code generator for expression.</summary>
+    </member>
+    <member name="M:System.Web.Razor.Generator.ExpressionCodeGenerator.#ctor">
+      <summary>This type/member supports the .NET Framework infrastructure and is not intended to be used directly from your code.Initializes a new instance of the <see cref="T:System.Web.Razor.Generator.ExpressionCodeGenerator" /> class.</summary>
+    </member>
+    <member name="M:System.Web.Razor.Generator.ExpressionCodeGenerator.Equals(System.Object)">
+      <summary>This type/member supports the .NET Framework infrastructure and is not intended to be used directly from your code.Indicates whether this instance and a specified object are equal.</summary>
+      <returns>true if <paramref name="obj" /> and this instance are the same type and represent the same value; otherwise, false.</returns>
+      <param name="obj">The object to compare with the current instance.</param>
+    </member>
+    <member name="M:System.Web.Razor.Generator.ExpressionCodeGenerator.GenerateCode(System.Web.Razor.Parser.SyntaxTree.Span,System.Web.Razor.Generator.CodeGeneratorContext)">
+      <summary>This type/member supports the .NET Framework infrastructure and is not intended to be used directly from your code.Generates code for the expression.</summary>
+      <param name="target">The source span whose content represents an expression.</param>
+      <param name="context">The code generator context.</param>
+    </member>
+    <member name="M:System.Web.Razor.Generator.ExpressionCodeGenerator.GenerateEndBlockCode(System.Web.Razor.Parser.SyntaxTree.Block,System.Web.Razor.Generator.CodeGeneratorContext)">
+      <summary>This type/member supports the .NET Framework infrastructure and is not intended to be used directly from your code.Generates the end code for the block.</summary>
+      <param name="target">The target block for the end code generation.</param>
+      <param name="context">The code generator context.</param>
+    </member>
+    <member name="M:System.Web.Razor.Generator.ExpressionCodeGenerator.GenerateStartBlockCode(System.Web.Razor.Parser.SyntaxTree.Block,System.Web.Razor.Generator.CodeGeneratorContext)">
+      <summary>This type/member supports the .NET Framework infrastructure and is not intended to be used directly from your code.Generates the start code for the block.</summary>
+      <param name="target">The target block the start code generation.</param>
+      <param name="context">The code generator context.</param>
+    </member>
+    <member name="M:System.Web.Razor.Generator.ExpressionCodeGenerator.GetHashCode">
+      <summary>This type/member supports the .NET Framework infrastructure and is not intended to be used directly from your code.Returns the hash code for this instance.</summary>
+      <returns>The hash code for this instance.</returns>
+    </member>
+    <member name="M:System.Web.Razor.Generator.ExpressionCodeGenerator.ToString">
+      <summary>This type/member supports the .NET Framework infrastructure and is not intended to be used directly from your code.Returns the string representation of this instance.</summary>
+      <returns>The string representation of this instance.</returns>
+    </member>
+    <member name="T:System.Web.Razor.Generator.ExpressionRenderingMode">
+      <summary>This type/member supports the .NET Framework infrastructure and is not intended to be used directly from your code.</summary>
+    </member>
+    <member name="F:System.Web.Razor.Generator.ExpressionRenderingMode.InjectCode">
+      <summary>This type/member supports the .NET Framework infrastructure and is not intended to be used directly from your code.</summary>
+    </member>
+    <member name="F:System.Web.Razor.Generator.ExpressionRenderingMode.WriteToOutput">
+      <summary>This type/member supports the .NET Framework infrastructure and is not intended to be used directly from your code.</summary>
+    </member>
+    <member name="T:System.Web.Razor.Generator.GeneratedClassContext">
+      <summary>This type/member supports the .NET Framework infrastructure and is not intended to be used directly from your code.Represents a generated class context.</summary>
+    </member>
+    <member name="M:System.Web.Razor.Generator.GeneratedClassContext.#ctor(System.String,System.String,System.String)">
+      <summary>This type/member supports the .NET Framework infrastructure and is not intended to be used directly from your code.Initializes a new instance of the <see cref="T:System.Web.Razor.Generator.GeneratedClassContext" /> class.</summary>
+      <param name="executeMethodName">The execute method name.</param>
+      <param name="writeMethodName">The write method name.</param>
+      <param name="writeLiteralMethodName">Write literal method name.</param>
+    </member>
+    <member name="M:System.Web.Razor.Generator.GeneratedClassContext.#ctor(System.String,System.String,System.String,System.String,System.String,System.String)">
+      <summary>This type/member supports the .NET Framework infrastructure and is not intended to be used directly from your code.Initializes a new instance of the <see cref="T:System.Web.Razor.Generator.GeneratedClassContext" /> class.</summary>
+      <param name="executeMethodName">Execute method name.</param>
+      <param name="writeMethodName">Write method name.</param>
+      <param name="writeLiteralMethodName">Write literal method name.</param>
+      <param name="writeToMethodName">Write to method name.</param>
+      <param name="writeLiteralToMethodName">Write literal to method name.</param>
+      <param name="templateTypeName">Template type name.</param>
+    </member>
+    <member name="M:System.Web.Razor.Generator.GeneratedClassContext.#ctor(System.String,System.String,System.String,System.String,System.String,System.String,System.String)">
+      <summary>This type/member supports the .NET Framework infrastructure and is not intended to be used directly from your code.Initializes a new instance of the <see cref="T:System.Web.Razor.Generator.GeneratedClassContext" /> class.</summary>
+      <param name="executeMethodName">Execute method name.</param>
+      <param name="writeMethodName">Write method name.</param>
+      <param name="writeLiteralMethodName">Write literal method name.</param>
+      <param name="writeToMethodName">Write to method name.</param>
+      <param name="writeLiteralToMethodName">Write literal to method name.</param>
+      <param name="templateTypeName">Template type name.</param>
+      <param name="defineSectionMethodName">Define section method name.</param>
+    </member>
+    <member name="M:System.Web.Razor.Generator.GeneratedClassContext.#ctor(System.String,System.String,System.String,System.String,System.String,System.String,System.String,System.String,System.String)">
+      <summary>This type/member supports the .NET Framework infrastructure and is not intended to be used directly from your code.Initializes a new instance of the <see cref="T:System.Web.Razor.Generator.GeneratedClassContext" /> class.</summary>
+      <param name="executeMethodName">Execute method name.</param>
+      <param name="writeMethodName">Write method name.</param>
+      <param name="writeLiteralMethodName">Write literal method name.</param>
+      <param name="writeToMethodName">Write to method name.</param>
+      <param name="writeLiteralToMethodName">Write literal to method name.</param>
+      <param name="templateTypeName">Template type name.</param>
+      <param name="defineSectionMethodName">Define section method name.</param>
+      <param name="beginContextMethodName">Begin context method name.</param>
+      <param name="endContextMethodName">End context method name.</param>
+    </member>
+    <member name="P:System.Web.Razor.Generator.GeneratedClassContext.AllowSections">
+      <summary>This type/member supports the .NET Framework infrastructure and is not intended to be used directly from your code.Gets a value indicating whether the context allows sections.</summary>
+      <returns>true if the context allows sections; otherwise, false.</returns>
+    </member>
+    <member name="P:System.Web.Razor.Generator.GeneratedClassContext.AllowTemplates">
+      <summary>This type/member supports the .NET Framework infrastructure and is not intended to be used directly from your code.Gets a value indicating whether the context allows templates.</summary>
+      <returns>true if the context allows templates; otherwise, false.</returns>
+    </member>
+    <member name="P:System.Web.Razor.Generator.GeneratedClassContext.BeginContextMethodName">
+      <summary>This type/member supports the .NET Framework infrastructure and is not intended to be used directly from your code.Gets or sets the name of the method before the generated context.</summary>
+      <returns>The name of the method before the generated context.</returns>
+    </member>
+    <member name="F:System.Web.Razor.Generator.GeneratedClassContext.Default">
+      <summary>This type/member supports the .NET Framework infrastructure and is not intended to be used directly from your code.Defines the default generated context.</summary>
+    </member>
+    <member name="F:System.Web.Razor.Generator.GeneratedClassContext.DefaultExecuteMethodName">
+      <summary>This type/member supports the .NET Framework infrastructure and is not intended to be used directly from your code.Defines the default name of the execute method.</summary>
+    </member>
+    <member name="F:System.Web.Razor.Generator.GeneratedClassContext.DefaultLayoutPropertyName">
+      <summary>This type/member supports the .NET Framework infrastructure and is not intended to be used directly from your code.Defines the default name of the layout property.</summary>
+    </member>
+    <member name="F:System.Web.Razor.Generator.GeneratedClassContext.DefaultWriteAttributeMethodName">
+      <summary>This type/member supports the .NET Framework infrastructure and is not intended to be used directly from your code.Defines the default name of the write attribute method.</summary>
+    </member>
+    <member name="F:System.Web.Razor.Generator.GeneratedClassContext.DefaultWriteAttributeToMethodName">
+      <summary>This type/member supports the .NET Framework infrastructure and is not intended to be used directly from your code.Defines the default name of the write to attribute to method.</summary>
+    </member>
+    <member name="F:System.Web.Razor.Generator.GeneratedClassContext.DefaultWriteLiteralMethodName">
+      <summary>This type/member supports the .NET Framework infrastructure and is not intended to be used directly from your code.Specifies the default name of the write literal method.</summary>
+    </member>
+    <member name="F:System.Web.Razor.Generator.GeneratedClassContext.DefaultWriteMethodName">
+      <summary>This type/member supports the .NET Framework infrastructure and is not intended to be used directly from your code.Specifies the default name of the write method.</summary>
+    </member>
+    <member name="P:System.Web.Razor.Generator.GeneratedClassContext.DefineSectionMethodName">
+      <summary>This type/member supports the .NET Framework infrastructure and is not intended to be used directly from your code.Gets or sets the name of the method that defines the section of the context.</summary>
+      <returns>The name of the method that defines the section of the context.</returns>
+    </member>
+    <member name="P:System.Web.Razor.Generator.GeneratedClassContext.EndContextMethodName">
+      <summary>This type/member supports the .NET Framework infrastructure and is not intended to be used directly from your code.Gets or sets the name of the method after the generated context.</summary>
+      <returns>The name of the method after the generated context.</returns>
+    </member>
+    <member name="M:System.Web.Razor.Generator.GeneratedClassContext.Equals(System.Object)">
+      <summary>This type/member supports the .NET Framework infrastructure and is not intended to be used directly from your code.Determines whether the specified object is equal to the current object.</summary>
+      <returns>true if the specified object is equal to the current object; otherwise, false.</returns>
+      <param name="obj">The object to compare to.</param>
+    </member>
+    <member name="P:System.Web.Razor.Generator.GeneratedClassContext.ExecuteMethodName">
+      <summary>This type/member supports the .NET Framework infrastructure and is not intended to be used directly from your code.Gets or sets the name of the method that will be invoked on the context.</summary>
+      <returns>The name of the method that will be invoked on the context.</returns>
+    </member>
+    <member name="M:System.Web.Razor.Generator.GeneratedClassContext.GetHashCode">
+      <summary>This type/member supports the .NET Framework infrastructure and is not intended to be used directly from your code.Returns the hash code for this current instance.</summary>
+      <returns>The hash code for this current instance.</returns>
+    </member>
+    <member name="P:System.Web.Razor.Generator.GeneratedClassContext.LayoutPropertyName">
+      <summary>This type/member supports the .NET Framework infrastructure and is not intended to be used directly from your code.Gets or sets the property name for the layout.</summary>
+      <returns>The property name for the layout.</returns>
+    </member>
+    <member name="M:System.Web.Razor.Generator.GeneratedClassContext.op_Equality(System.Web.Razor.Generator.GeneratedClassContext,System.Web.Razor.Generator.GeneratedClassContext)">
+      <summary>This type/member supports the .NET Framework infrastructure and is not intended to be used directly from your code.Determines whether the two <see cref="T:System.Web.Razor.Generator.GeneratedClassContext" /> object are equal.</summary>
+      <returns>true if the two <see cref="T:System.Web.Razor.Generator.GeneratedClassContext" /> object are equal; otherwise, false.</returns>
+      <param name="left">The first object to compare.</param>
+      <param name="right">The second object to compare.</param>
+    </member>
+    <member name="M:System.Web.Razor.Generator.GeneratedClassContext.op_Inequality(System.Web.Razor.Generator.GeneratedClassContext,System.Web.Razor.Generator.GeneratedClassContext)">
+      <summary>This type/member supports the .NET Framework infrastructure and is not intended to be used directly from your code.Determines whether the two <see cref="T:System.Web.Razor.Generator.GeneratedClassContext" /> object are not equal.</summary>
+      <returns>true if the two <see cref="T:System.Web.Razor.Generator.GeneratedClassContext" /> object are not equal; otherwise, false.</returns>
+      <param name="left">The first object to compare.</param>
+      <param name="right">The second object to compare.</param>
+    </member>
+    <member name="P:System.Web.Razor.Generator.GeneratedClassContext.ResolveUrlMethodName">
+      <summary>This type/member supports the .NET Framework infrastructure and is not intended to be used directly from your code.Gets or sets the name of the method that resolves a Url for the context.</summary>
+      <returns>The name of the method that resolves a Url for the context.</returns>
+    </member>
+    <member name="P:System.Web.Razor.Generator.GeneratedClassContext.SupportsInstrumentation">
+      <summary>This type/member supports the .NET Framework infrastructure and is not intended to be used directly from your code.Gets a value indicating whether the generated class supports instrumentation.</summary>
+      <returns>true if the generated class supports instrumentation; otherwise, false.</returns>
+    </member>
+    <member name="P:System.Web.Razor.Generator.GeneratedClassContext.TemplateTypeName">
+      <summary>This type/member supports the .NET Framework infrastructure and is not intended to be used directly from your code.Gets or sets the type name for the template.</summary>
+      <returns>The type name for the template.</returns>
+    </member>
+    <member name="P:System.Web.Razor.Generator.GeneratedClassContext.WriteAttributeMethodName">
+      <summary>This type/member supports the .NET Framework infrastructure and is not intended to be used directly from your code.Gets or sets the name of the method that writes an attribute.</summary>
+      <returns>The name of the method that writes an attribute.</returns>
+    </member>
+    <member name="P:System.Web.Razor.Generator.GeneratedClassContext.WriteAttributeToMethodName">
+      <summary>This type/member supports the .NET Framework infrastructure and is not intended to be used directly from your code.Gets or sets the name of the method where to write an attribute.</summary>
+      <returns>The name of the method where to write an attribute.</returns>
+    </member>
+    <member name="P:System.Web.Razor.Generator.GeneratedClassContext.WriteLiteralMethodName">
+      <summary>This type/member supports the .NET Framework infrastructure and is not intended to be used directly from your code.Gets or sets the name of the method where to write literal for the context.</summary>
+      <returns>The name of the method where to write literal for the context.</returns>
+    </member>
+    <member name="P:System.Web.Razor.Generator.GeneratedClassContext.WriteLiteralToMethodName">
+      <summary>This type/member supports the .NET Framework infrastructure and is not intended to be used directly from your code.Gets or sets the name of the method where to write literal for the context.</summary>
+      <returns>The name of the method where to write literal for the context.</returns>
+    </member>
+    <member name="P:System.Web.Razor.Generator.GeneratedClassContext.WriteMethodName">
+      <summary>This type/member supports the .NET Framework infrastructure and is not intended to be used directly from your code.Gets or sets the name of the method that will write on the context.</summary>
+      <returns>The name of the method that will write on the context.</returns>
+    </member>
+    <member name="P:System.Web.Razor.Generator.GeneratedClassContext.WriteToMethodName">
+      <summary>This type/member supports the .NET Framework infrastructure and is not intended to be used directly from your code.Gets or sets the name of the method that will write on the context.</summary>
+      <returns>The name of the method that will write on the context.</returns>
+    </member>
+    <member name="T:System.Web.Razor.Generator.GeneratedCodeMapping">
+      <summary>This type/member supports the .NET Framework infrastructure and is not intended to be used directly from your code.Represents the generated code mapping objects.</summary>
+    </member>
+    <member name="M:System.Web.Razor.Generator.GeneratedCodeMapping.#ctor(System.Int32,System.Int32,System.Int32,System.Int32)">
+      <summary>This type/member supports the .NET Framework infrastructure and is not intended to be used directly from your code.Initializes a new instance of the <see cref="T:System.Web.Razor.Generator.GeneratedCodeMapping" /> class.</summary>
+      <param name="startLine">The starting line.</param>
+      <param name="startColumn">The starting column.</param>
+      <param name="startGeneratedColumn">The start generated column.</param>
+      <param name="codeLength">The code length.</param>
+    </member>
+    <member name="M:System.Web.Razor.Generator.GeneratedCodeMapping.#ctor(System.Int32,System.Int32,System.Int32,System.Int32,System.Int32)">
+      <summary>This type/member supports the .NET Framework infrastructure and is not intended to be used directly from your code.Initializes a new instance of the <see cref="T:System.Web.Razor.Generator.GeneratedCodeMapping" /> class.</summary>
+      <param name="startOffset">The start offset.</param>
+      <param name="startLine">The starting line.</param>
+      <param name="startColumn">The starting column.</param>
+      <param name="startGeneratedColumn">The start generated column.</param>
+      <param name="codeLength">The code length.</param>
+    </member>
+    <member name="P:System.Web.Razor.Generator.GeneratedCodeMapping.CodeLength">
+      <summary>This type/member supports the .NET Framework infrastructure and is not intended to be used directly from your code.Gets or sets the length of the generated map codes.</summary>
+      <returns>The length of the generated map codes.</returns>
+    </member>
+    <member name="M:System.Web.Razor.Generator.GeneratedCodeMapping.Equals(System.Object)">
+      <summary>This type/member supports the .NET Framework infrastructure and is not intended to be used directly from your code.Determines whether the specified object is equal to the current generated code mapping object.</summary>
+      <returns>true if the specified object is equal to the current generated code mapping object; otherwise, false.</returns>
+      <param name="obj">The object to compare with the current object.</param>
+    </member>
+    <member name="M:System.Web.Razor.Generator.GeneratedCodeMapping.GetHashCode">
+      <summary>This type/member supports the .NET Framework infrastructure and is not intended to be used directly from your code.Returns the hash code for the generated code mapping object.</summary>
+      <returns>The hash code for the generated code mapping object.</returns>
+    </member>
+    <member name="M:System.Web.Razor.Generator.GeneratedCodeMapping.op_Equality(System.Web.Razor.Generator.GeneratedCodeMapping,System.Web.Razor.Generator.GeneratedCodeMapping)">
+      <summary>This type/member supports the .NET Framework infrastructure and is not intended to be used directly from your code.Determines whether the two specified generated code mapping objects have the same value.</summary>
+      <returns>true if the two specified generated code mapping objects have the same value; otherwise, false.</returns>
+      <param name="left">The left generated code mapping objects.</param>
+      <param name="right">The right generated code mapping objects.</param>
+    </member>
+    <member name="M:System.Web.Razor.Generator.GeneratedCodeMapping.op_Inequality(System.Web.Razor.Generator.GeneratedCodeMapping,System.Web.Razor.Generator.GeneratedCodeMapping)">
+      <summary>This type/member supports the .NET Framework infrastructure and is not intended to be used directly from your code.Determines whether the two specified generated code mapping objects have different values.</summary>
+      <returns>true the two specified generated code mapping objects have different values; otherwise, false.</returns>
+      <param name="right">The right generated code mapping objects.</param>
+    </member>
+    <member name="P:System.Web.Razor.Generator.GeneratedCodeMapping.StartColumn">
+      <summary>This type/member supports the .NET Framework infrastructure and is not intended to be used directly from your code.Gets or sets the starting column of the generated code maps.</summary>
+      <returns>The starting column of the generated code maps.</returns>
+    </member>
+    <member name="P:System.Web.Razor.Generator.GeneratedCodeMapping.StartGeneratedColumn">
+      <summary>This type/member supports the .NET Framework infrastructure and is not intended to be used directly from your code.Gets or sets the starting column of a generated code maps in the generated source file.</summary>
+      <returns>The starting column of a generated code maps in the generated source file.</returns>
+    </member>
+    <member name="P:System.Web.Razor.Generator.GeneratedCodeMapping.StartLine">
+      <summary>This type/member supports the .NET Framework infrastructure and is not intended to be used directly from your code.Gets the starting line of the generated code maps.</summary>
+      <returns>The starting line of the generated code maps.</returns>
+    </member>
+    <member name="P:System.Web.Razor.Generator.GeneratedCodeMapping.StartOffset">
+      <summary>This type/member supports the .NET Framework infrastructure and is not intended to be used directly from your code.Gets or sets the starting offset of the generated code maps.</summary>
+      <returns>The starting offset of the generated code maps.</returns>
+    </member>
+    <member name="M:System.Web.Razor.Generator.GeneratedCodeMapping.ToString">
+      <summary>Returns a string that represents the current object.</summary>
+      <returns>A string that represents the current object.</returns>
+    </member>
+    <member name="T:System.Web.Razor.Generator.HelperCodeGenerator">
+      <summary>This type/member supports the .NET Framework infrastructure and is not intended to be used directly from your code.Represents a helper code generator.</summary>
+    </member>
+    <member name="M:System.Web.Razor.Generator.HelperCodeGenerator.#ctor(System.Web.Razor.Text.LocationTagged{System.String},System.Boolean)">
+      <summary>This type/member supports the .NET Framework infrastructure and is not intended to be used directly from your code.Initializes a new instance of the <see cref="T:System.Web.Razor.Generator.HelperCodeGenerator" /> class.</summary>
+      <param name="signature">The signature.</param>
+      <param name="headerComplete">true to complete the header; otherwise, false.</param>
+    </member>
+    <member name="M:System.Web.Razor.Generator.HelperCodeGenerator.Equals(System.Object)">
+      <summary>This type/member supports the .NET Framework infrastructure and is not intended to be used directly from your code.Determines whether the specified object is equal to the current object.</summary>
+      <returns>true if the specified object is equal to the current object; otherwise, false.</returns>
+      <param name="obj">The object to compare to.</param>
+    </member>
+    <member name="P:System.Web.Razor.Generator.HelperCodeGenerator.Footer">
+      <summary>This type/member supports the .NET Framework infrastructure and is not intended to be used directly from your code.Gets or sets the footer for this code.</summary>
+      <returns>The footer for this code.</returns>
+    </member>
+    <member name="M:System.Web.Razor.Generator.HelperCodeGenerator.GenerateEndBlockCode(System.Web.Razor.Parser.SyntaxTree.Block,System.Web.Razor.Generator.CodeGeneratorContext)">
+      <summary>This type/member supports the .NET Framework infrastructure and is not intended to be used directly from your code.Generates a block after the code.</summary>
+      <param name="target">The block to generate.</param>
+      <param name="context">The code generator context.</param>
+    </member>
+    <member name="M:System.Web.Razor.Generator.HelperCodeGenerator.GenerateStartBlockCode(System.Web.Razor.Parser.SyntaxTree.Block,System.Web.Razor.Generator.CodeGeneratorContext)">
+      <summary>This type/member supports the .NET Framework infrastructure and is not intended to be used directly from your code.Generates a block before the code.</summary>
+      <param name="target">The block to generate.</param>
+      <param name="context">The code generator context.</param>
+    </member>
+    <member name="M:System.Web.Razor.Generator.HelperCodeGenerator.GetHashCode">
+      <summary>This type/member supports the .NET Framework infrastructure and is not intended to be used directly from your code.Gets the hash code for the current instance.</summary>
+      <returns>The hash code for the current instance.</returns>
+    </member>
+    <member name="P:System.Web.Razor.Generator.HelperCodeGenerator.HeaderComplete">
+      <summary>This type/member supports the .NET Framework infrastructure and is not intended to be used directly from your code.Gets or sets a value indicating whether the header for this code is complete.</summary>
+      <returns>true if the header for this code is complete; otherwise, false.</returns>
+    </member>
+    <member name="P:System.Web.Razor.Generator.HelperCodeGenerator.Signature">
+      <summary>This type/member supports the .NET Framework infrastructure and is not intended to be used directly from your code.Gets or sets the signature for this code.</summary>
+      <returns>The signature for this code.</returns>
+    </member>
+    <member name="M:System.Web.Razor.Generator.HelperCodeGenerator.ToString">
+      <summary>This type/member supports the .NET Framework infrastructure and is not intended to be used directly from your code.Returns a string representation of the current instance.</summary>
+      <returns>A string representation of the current instance.</returns>
+    </member>
+    <member name="T:System.Web.Razor.Generator.HybridCodeGenerator">
+      <summary>This type/member supports the .NET Framework infrastructure and is not intended to be used directly from your code.Represents a hybrid code generator.</summary>
+    </member>
+    <member name="M:System.Web.Razor.Generator.HybridCodeGenerator.#ctor">
+      <summary>This type/member supports the .NET Framework infrastructure and is not intended to be used directly from your code.Initializes a new instance of the <see cref="T:System.Web.Razor.Generator.HybridCodeGenerator" /> class.</summary>
+    </member>
+    <member name="M:System.Web.Razor.Generator.HybridCodeGenerator.GenerateCode(System.Web.Razor.Parser.SyntaxTree.Span,System.Web.Razor.Generator.CodeGeneratorContext)">
+      <summary>This type/member supports the .NET Framework infrastructure and is not intended to be used directly from your code.Generates code for the data model from switches identified by parameters.</summary>
+      <param name="target">The target object.</param>
+      <param name="context">The code generator context.</param>
+    </member>
+    <member name="M:System.Web.Razor.Generator.HybridCodeGenerator.GenerateEndBlockCode(System.Web.Razor.Parser.SyntaxTree.Block,System.Web.Razor.Generator.CodeGeneratorContext)">
+      <summary>This type/member supports the .NET Framework infrastructure and is not intended to be used directly from your code.Generates an end block code.</summary>
+      <param name="target">The target object.</param>
+      <param name="context">The code generator context.</param>
+    </member>
+    <member name="M:System.Web.Razor.Generator.HybridCodeGenerator.GenerateStartBlockCode(System.Web.Razor.Parser.SyntaxTree.Block,System.Web.Razor.Generator.CodeGeneratorContext)">
+      <summary>This type/member supports the .NET Framework infrastructure and is not intended to be used directly from your code.Generates the start block code.</summary>
+      <param name="target">The target object.</param>
+      <param name="context">The code generator context.</param>
+    </member>
+    <member name="T:System.Web.Razor.Generator.IBlockCodeGenerator">
+      <summary>This type/member supports the .NET Framework infrastructure and is not intended to be used directly from your code.Represents the <see cref="T:System.Web.Razor.Generator.IBlockCodeGenerator" /> for the webpages.</summary>
+    </member>
+    <member name="M:System.Web.Razor.Generator.IBlockCodeGenerator.GenerateEndBlockCode(System.Web.Razor.Parser.SyntaxTree.Block,System.Web.Razor.Generator.CodeGeneratorContext)">
+      <summary>This type/member supports the .NET Framework infrastructure and is not intended to be used directly from your code.Generates the end block code for the razor.</summary>
+      <param name="target">The target block.</param>
+      <param name="context">The code generator context.</param>
+    </member>
+    <member name="M:System.Web.Razor.Generator.IBlockCodeGenerator.GenerateStartBlockCode(System.Web.Razor.Parser.SyntaxTree.Block,System.Web.Razor.Generator.CodeGeneratorContext)">
+      <summary>This type/member supports the .NET Framework infrastructure and is not intended to be used directly from your code.Generates the start block code for the razor.</summary>
+      <param name="target">The target block.</param>
+      <param name="context">The code generator context.</param>
+    </member>
+    <member name="T:System.Web.Razor.Generator.ISpanCodeGenerator">
+      <summary>This type/member supports the .NET Framework infrastructure and is not intended to be used directly from your code.Represents a phase of the code generator.</summary>
+    </member>
+    <member name="M:System.Web.Razor.Generator.ISpanCodeGenerator.GenerateCode(System.Web.Razor.Parser.SyntaxTree.Span,System.Web.Razor.Generator.CodeGeneratorContext)">
+      <summary>This type/member supports the .NET Framework infrastructure and is not intended to be used directly from your code.Generates code for the data model with the specified target and context.</summary>
+      <param name="target">The target object.</param>
+      <param name="context">The code generator context.</param>
+    </member>
+    <member name="T:System.Web.Razor.Generator.LiteralAttributeCodeGenerator">
+      <summary>This type/member supports the .NET Framework infrastructure and is not intended to be used directly from your code. Represents a code generator for literal attribute.</summary>
+    </member>
+    <member name="M:System.Web.Razor.Generator.LiteralAttributeCodeGenerator.#ctor(System.Web.Razor.Text.LocationTagged{System.String},System.Web.Razor.Text.LocationTagged{System.String})">
+      <summary>This type/member supports the .NET Framework infrastructure and is not intended to be used directly from your code. . Initializes a new instance of the <see cref="T:System.Web.Razor.Generator.LiteralAttributeCodeGenerator" /> class.</summary>
+      <param name="prefix">The prefix of the literal attribute.</param>
+      <param name="value">The value of the literal attribute.</param>
+    </member>
+    <member name="M:System.Web.Razor.Generator.LiteralAttributeCodeGenerator.#ctor(System.Web.Razor.Text.LocationTagged{System.String},System.Web.Razor.Text.LocationTagged{System.Web.Razor.Generator.SpanCodeGenerator})">
+      <summary>This type/member supports the .NET Framework infrastructure and is not intended to be used directly from your code. . Initializes a new instance of the <see cref="T:System.Web.Razor.Generator.LiteralAttributeCodeGenerator" /> class.</summary>
+      <param name="prefix">The prefix of the literal attribute.</param>
+      <param name="valueGenerator">The value generator for the literal attribute.</param>
+    </member>
+    <member name="M:System.Web.Razor.Generator.LiteralAttributeCodeGenerator.Equals(System.Object)">
+      <summary>This type/member supports the .NET Framework infrastructure and is not intended to be used directly from your code. Determines whether the specified object is equal to this instance.</summary>
+      <returns>true if the specified object is equal to this instance; otherwise, false.</returns>
+      <param name="obj">The object to compare to this instance.</param>
+    </member>
+    <member name="M:System.Web.Razor.Generator.LiteralAttributeCodeGenerator.GenerateCode(System.Web.Razor.Parser.SyntaxTree.Span,System.Web.Razor.Generator.CodeGeneratorContext)">
+      <summary>This type/member supports the .NET Framework infrastructure and is not intended to be used directly from your code. Generates the code for the literal attribute.</summary>
+      <param name="target">The source span whose content represents the literal attribute.</param>
+      <param name="context">The context of the code generator.</param>
+    </member>
+    <member name="M:System.Web.Razor.Generator.LiteralAttributeCodeGenerator.GetHashCode">
+      <summary>This type/member supports the .NET Framework infrastructure and is not intended to be used directly from your code. Gets the hash code for the current instance.</summary>
+      <returns>The hash code for the current instance.</returns>
+    </member>
+    <member name="P:System.Web.Razor.Generator.LiteralAttributeCodeGenerator.Prefix">
+      <summary>This type/member supports the .NET Framework infrastructure and is not intended to be used directly from your code. Gets or sets the prefix of the literal attribute.</summary>
+      <returns>The prefix of the literal attribute.</returns>
+    </member>
+    <member name="M:System.Web.Razor.Generator.LiteralAttributeCodeGenerator.ToString">
+      <summary>This type/member supports the .NET Framework infrastructure and is not intended to be used directly from your code. Returns the string representation of this instance.</summary>
+      <returns>The string representation of this instance.</returns>
+    </member>
+    <member name="P:System.Web.Razor.Generator.LiteralAttributeCodeGenerator.Value">
+      <summary>This type/member supports the .NET Framework infrastructure and is not intended to be used directly from your code. Gets or sets the value of the literal attribute.</summary>
+      <returns>The value of the literal attribute.</returns>
+    </member>
+    <member name="P:System.Web.Razor.Generator.LiteralAttributeCodeGenerator.ValueGenerator">
+      <summary>This type/member supports the .NET Framework infrastructure and is not intended to be used directly from your code. Gets or sets the value generator for the literal attribute.</summary>
+      <returns>The value generator for the literal attribute.</returns>
+    </member>
+    <member name="T:System.Web.Razor.Generator.MarkupCodeGenerator">
+      <summary>This type/member supports the .NET Framework infrastructure and is not intended to be used directly from your code. Represents a code generator for markup.</summary>
+    </member>
+    <member name="M:System.Web.Razor.Generator.MarkupCodeGenerator.#ctor">
+      <summary>This type/member supports the .NET Framework infrastructure and is not intended to be used directly from your code. Initializes a new instance of the <see cref="T:System.Web.Razor.Generator.MarkupCodeGenerator" /> class.</summary>
+    </member>
+    <member name="M:System.Web.Razor.Generator.MarkupCodeGenerator.Equals(System.Object)">
+      <summary>This type/member supports the .NET Framework infrastructure and is not intended to be used directly from your code. Determines whether the specified object is equal to this instance.</summary>
+      <returns>true if the specified object is equal to this instance; otherwise, false.</returns>
+      <param name="obj">The object to compare to this instance.</param>
+    </member>
+    <member name="M:System.Web.Razor.Generator.MarkupCodeGenerator.GenerateCode(System.Web.Razor.Parser.SyntaxTree.Span,System.Web.Razor.Generator.CodeGeneratorContext)">
+      <summary>This type/member supports the .NET Framework infrastructure and is not intended to be used directly from your code. Generates the code for the markup.</summary>
+      <param name="target">The source span whose content represents the markup.</param>
+      <param name="context">The context of the code generator.</param>
+    </member>
+    <member name="M:System.Web.Razor.Generator.MarkupCodeGenerator.GetHashCode">
+      <summary>This type/member supports the .NET Framework infrastructure and is not intended to be used directly from your code. Returns the hash code for this instance.</summary>
+      <returns>The hash code for this instance.</returns>
+    </member>
+    <member name="M:System.Web.Razor.Generator.MarkupCodeGenerator.ToString">
+      <summary>This type/member supports the .NET Framework infrastructure and is not intended to be used directly from your code. Returns the string representation for this instance.</summary>
+      <returns>The string representation for this instance.</returns>
+    </member>
+    <member name="T:System.Web.Razor.Generator.RazorCodeGenerator">
+      <summary>This type/member supports the .NET Framework infrastructure and is not intended to be used directly from your code.Represents a Razor code generator.</summary>
+    </member>
+    <member name="M:System.Web.Razor.Generator.RazorCodeGenerator.#ctor(System.String,System.String,System.String,System.Web.Razor.RazorEngineHost)">
+      <summary>This type/member supports the .NET Framework infrastructure and is not intended to be used directly from your code.Initializes a new instance of the <see cref="T:System.Web.Razor.Generator.RazorCodeGenerator" /> class.</summary>
+      <param name="className">The class name.</param>
+      <param name="rootNamespaceName">The root namespace name.</param>
+      <param name="sourceFileName">The source file name.</param>
+      <param name="host">The host.</param>
+    </member>
+    <member name="P:System.Web.Razor.Generator.RazorCodeGenerator.ClassName">
+      <summary>This type/member supports the .NET Framework infrastructure and is not intended to be used directly from your code.Gets or sets the class name for this code.</summary>
+      <returns>The class name for this code.</returns>
+    </member>
+    <member name="P:System.Web.Razor.Generator.RazorCodeGenerator.Context">
+      <summary>This type/member supports the .NET Framework infrastructure and is not intended to be used directly from your code.Gets the context of this code generator.</summary>
+      <returns>The context of this code generator.</returns>
+    </member>
+    <member name="P:System.Web.Razor.Generator.RazorCodeGenerator.DesignTimeMode">
+      <summary>This type/member supports the .NET Framework infrastructure and is not intended to be used directly from your code.Gets or sets a value indicating whether the code generator is in design-time mode.</summary>
+      <returns>true if the code generator is in design-time mode; otherwise, false.</returns>
+    </member>
+    <member name="P:System.Web.Razor.Generator.RazorCodeGenerator.GenerateLinePragmas">
+      <summary>This type/member supports the .NET Framework infrastructure and is not intended to be used directly from your code.Gets or sets a value indicating whether the generator should generate line pragmas in the Razor code.</summary>
+      <returns>true if the generator should generate line pragmas in the Razor code; otherwise, false.</returns>
+    </member>
+    <member name="P:System.Web.Razor.Generator.RazorCodeGenerator.Host">
+      <summary>This type/member supports the .NET Framework infrastructure and is not intended to be used directly from your code.Gets or sets the razor engine host.</summary>
+      <returns>The razor engine host.</returns>
+    </member>
+    <member name="M:System.Web.Razor.Generator.RazorCodeGenerator.Initialize(System.Web.Razor.Generator.CodeGeneratorContext)">
+      <summary>This type/member supports the .NET Framework infrastructure and is not intended to be used directly from your code.Initializes the current <see cref="T:System.Web.Razor.Generator.RazorCodeGenerator" /> instance.</summary>
+      <param name="context">The context.</param>
+    </member>
+    <member name="M:System.Web.Razor.Generator.RazorCodeGenerator.OnComplete">
+      <summary>This type/member supports the .NET Framework infrastructure and is not intended to be used directly from your code.Raises the Complete event.</summary>
+    </member>
+    <member name="P:System.Web.Razor.Generator.RazorCodeGenerator.RootNamespaceName">
+      <summary>This type/member supports the .NET Framework infrastructure and is not intended to be used directly from your code.Gets or sets the name of the root namespace.</summary>
+      <returns>The name of the root namespace.</returns>
+    </member>
+    <member name="P:System.Web.Razor.Generator.RazorCodeGenerator.SourceFileName">
+      <summary>This type/member supports the .NET Framework infrastructure and is not intended to be used directly from your code.Gets or sets the name of the source file.</summary>
+      <returns>The name of the source file.</returns>
+    </member>
+    <member name="M:System.Web.Razor.Generator.RazorCodeGenerator.VisitEndBlock(System.Web.Razor.Parser.SyntaxTree.Block)">
+      <summary>This type/member supports the .NET Framework infrastructure and is not intended to be used directly from your code.Visits the end block.</summary>
+      <param name="block">The block to visit.</param>
+    </member>
+    <member name="M:System.Web.Razor.Generator.RazorCodeGenerator.VisitSpan(System.Web.Razor.Parser.SyntaxTree.Span)">
+      <summary>This type/member supports the .NET Framework infrastructure and is not intended to be used directly from your code.Visits the span.</summary>
+      <param name="span">The span to visit.</param>
+    </member>
+    <member name="M:System.Web.Razor.Generator.RazorCodeGenerator.VisitStartBlock(System.Web.Razor.Parser.SyntaxTree.Block)">
+      <summary>This type/member supports the .NET Framework infrastructure and is not intended to be used directly from your code.Visits the start block.</summary>
+      <param name="block">The block to visit.</param>
+    </member>
+    <member name="T:System.Web.Razor.Generator.RazorCommentCodeGenerator">
+      <summary>This type/member supports the .NET Framework infrastructure and is not intended to be used directly from your code.Represents the razor comment code generator for the webpages.</summary>
+    </member>
+    <member name="M:System.Web.Razor.Generator.RazorCommentCodeGenerator.#ctor">
+      <summary>This type/member supports the .NET Framework infrastructure and is not intended to be used directly from your code.Initializes a new instance of the <see cref="T:System.Web.Razor.Generator.RazorCommentCodeGenerator" /> class.</summary>
+    </member>
+    <member name="M:System.Web.Razor.Generator.RazorCommentCodeGenerator.GenerateStartBlockCode(System.Web.Razor.Parser.SyntaxTree.Block,System.Web.Razor.Generator.CodeGeneratorContext)">
+      <summary>This type/member supports the .NET Framework infrastructure and is not intended to be used directly from your code.Generates the start block code with the specified parameters.</summary>
+      <param name="target">The target block.</param>
+      <param name="context">The code generator context.</param>
+    </member>
+    <member name="T:System.Web.Razor.Generator.RazorDirectiveAttributeCodeGenerator">
+      <summary>This type/member supports the .NET Framework infrastructure and is not intended to be used directly from your code. Represents a code generator for Razor directive attribute. </summary>
+    </member>
+    <member name="M:System.Web.Razor.Generator.RazorDirectiveAttributeCodeGenerator.#ctor(System.String,System.String)">
+      <summary>This type/member supports the .NET Framework infrastructure and is not intended to be used directly from your code. Initializes a new instance of the <see cref="T:System.Web.Razor.Generator.RazorDirectiveAttributeCodeGenerator" /> class.</summary>
+      <param name="name">The name of the directive attribute.</param>
+      <param name="value">The value of the directive attribute.</param>
+    </member>
+    <member name="M:System.Web.Razor.Generator.RazorDirectiveAttributeCodeGenerator.Equals(System.Object)">
+      <summary>This type/member supports the .NET Framework infrastructure and is not intended to be used directly from your code. Determines whether the specified object is equal to this instance.</summary>
+      <returns>true if the specified object is equal to this instance; otherwise, false.</returns>
+      <param name="obj">The object to compare to this instance.</param>
+    </member>
+    <member name="M:System.Web.Razor.Generator.RazorDirectiveAttributeCodeGenerator.GenerateCode(System.Web.Razor.Parser.SyntaxTree.Span,System.Web.Razor.Generator.CodeGeneratorContext)">
+      <summary>This type/member supports the .NET Framework infrastructure and is not intended to be used directly from your code. Generates the code for the directive attribute.</summary>
+      <param name="target">The source span whose content represents the directive attribute to generate.</param>
+      <param name="context">The context of the code generator.</param>
+    </member>
+    <member name="M:System.Web.Razor.Generator.RazorDirectiveAttributeCodeGenerator.GetHashCode">
+      <summary>This type/member supports the .NET Framework infrastructure and is not intended to be used directly from your code. Returns the hash code for this instance.</summary>
+      <returns>The hash code for this instance.</returns>
+    </member>
+    <member name="P:System.Web.Razor.Generator.RazorDirectiveAttributeCodeGenerator.Name">
+      <summary>This type/member supports the .NET Framework infrastructure and is not intended to be used directly from your code. Gets or sets the name of the directive attribute.</summary>
+      <returns>The name of the directive attribute.</returns>
+    </member>
+    <member name="M:System.Web.Razor.Generator.RazorDirectiveAttributeCodeGenerator.ToString">
+      <summary>This type/member supports the .NET Framework infrastructure and is not intended to be used directly from your code. Returns the string representation for this instance.</summary>
+      <returns>The string representation for this instance.</returns>
+    </member>
+    <member name="P:System.Web.Razor.Generator.RazorDirectiveAttributeCodeGenerator.Value">
+      <summary>This type/member supports the .NET Framework infrastructure and is not intended to be used directly from your code. Gets or sets the value of the directive attribute.</summary>
+      <returns>The value of the directive attribute.</returns>
+    </member>
+    <member name="T:System.Web.Razor.Generator.ResolveUrlCodeGenerator">
+      <summary>This type/member supports the .NET Framework infrastructure and is not intended to be used directly from your code.Represents the resolve Url code generator.</summary>
+    </member>
+    <member name="M:System.Web.Razor.Generator.ResolveUrlCodeGenerator.#ctor">
+      <summary>This type/member supports the .NET Framework infrastructure and is not intended to be used directly from your code.Initializes a new instance of the <see cref="T:System.Web.Razor.Generator.ResolveUrlCodeGenerator" /> class.</summary>
+    </member>
+    <member name="M:System.Web.Razor.Generator.ResolveUrlCodeGenerator.Equals(System.Object)">
+      <summary>This type/member supports the .NET Framework infrastructure and is not intended to be used directly from your code.Indicates whether this instance and a specified object are equal.</summary>
+      <returns>true if <paramref name="obj" /> and this instance are the same type and represent the same value; otherwise, false.</returns>
+      <param name="obj">The object to compare with the current instance.</param>
+    </member>
+    <member name="M:System.Web.Razor.Generator.ResolveUrlCodeGenerator.GenerateCode(System.Web.Razor.Parser.SyntaxTree.Span,System.Web.Razor.Generator.CodeGeneratorContext)">
+      <summary>This type/member supports the .NET Framework infrastructure and is not intended to be used directly from your code.Generates code for the Url.</summary>
+      <param name="target">The target object.</param>
+      <param name="context">The code generator context.</param>
+    </member>
+    <member name="M:System.Web.Razor.Generator.ResolveUrlCodeGenerator.GetHashCode">
+      <summary>This type/member supports the .NET Framework infrastructure and is not intended to be used directly from your code.Returns the hash code for this instance.</summary>
+      <returns>The hash code for this instance.</returns>
+    </member>
+    <member name="M:System.Web.Razor.Generator.ResolveUrlCodeGenerator.ToString">
+      <summary>This type/member supports the .NET Framework infrastructure and is not intended to be used directly from your code.Returns the fully qualified type name of this instance.</summary>
+      <returns>The fully qualified type name.</returns>
+    </member>
+    <member name="T:System.Web.Razor.Generator.SectionCodeGenerator">
+      <summary>This type/member supports the .NET Framework infrastructure and is not intended to be used directly from your code.Represents a section code generator.</summary>
+    </member>
+    <member name="M:System.Web.Razor.Generator.SectionCodeGenerator.#ctor(System.String)">
+      <summary>This type/member supports the .NET Framework infrastructure and is not intended to be used directly from your code.Initializes a new instance of the <see cref="T:System.Web.Razor.Generator.SectionCodeGenerator" /> class.</summary>
+      <param name="sectionName">The name of the section code.</param>
+    </member>
+    <member name="M:System.Web.Razor.Generator.SectionCodeGenerator.Equals(System.Object)">
+      <summary>This type/member supports the .NET Framework infrastructure and is not intended to be used directly from your code.Determines whether the specified object is equal to the current object.</summary>
+      <returns>true if the specified object is equal to the current object; otherwise, false.</returns>
+      <param name="obj">The object to compare to.</param>
+    </member>
+    <member name="M:System.Web.Razor.Generator.SectionCodeGenerator.GenerateEndBlockCode(System.Web.Razor.Parser.SyntaxTree.Block,System.Web.Razor.Generator.CodeGeneratorContext)">
+      <summary>This type/member supports the .NET Framework infrastructure and is not intended to be used directly from your code.Generates a block after the section code.</summary>
+      <param name="target">The target to generate.</param>
+      <param name="context">The code generator context.</param>
+    </member>
+    <member name="M:System.Web.Razor.Generator.SectionCodeGenerator.GenerateStartBlockCode(System.Web.Razor.Parser.SyntaxTree.Block,System.Web.Razor.Generator.CodeGeneratorContext)">
+      <summary>This type/member supports the .NET Framework infrastructure and is not intended to be used directly from your code.Generates a block before the section code.</summary>
+      <param name="target">The target to generate.</param>
+      <param name="context">The code generator context.</param>
+    </member>
+    <member name="M:System.Web.Razor.Generator.SectionCodeGenerator.GetHashCode">
+      <summary>This type/member supports the .NET Framework infrastructure and is not intended to be used directly from your code.Retrieves the hash code for this current instance.</summary>
+      <returns>The hash code for this current instance.</returns>
+    </member>
+    <member name="P:System.Web.Razor.Generator.SectionCodeGenerator.SectionName">
+      <summary>This type/member supports the .NET Framework infrastructure and is not intended to be used directly from your code.Gets or sets the name of the section.</summary>
+      <returns>The name of the section.</returns>
+    </member>
+    <member name="M:System.Web.Razor.Generator.SectionCodeGenerator.ToString">
+      <summary>This type/member supports the .NET Framework infrastructure and is not intended to be used directly from your code.Returns a string representation of this current instance.</summary>
+      <returns>A string representation of this current instance.</returns>
+    </member>
+    <member name="T:System.Web.Razor.Generator.SetBaseTypeCodeGenerator">
+      <summary>This type/member supports the .NET Framework infrastructure and is not intended to be used directly from your code. Represents a code generator for set base type.</summary>
+    </member>
+    <member name="M:System.Web.Razor.Generator.SetBaseTypeCodeGenerator.#ctor(System.String)">
+      <summary>This type/member supports the .NET Framework infrastructure and is not intended to be used directly from your code. Initializes a new instance of the <see cref="T:System.Web.Razor.Generator.SetBaseTypeCodeGenerator" /> class.</summary>
+      <param name="baseType">The set base type.</param>
+    </member>
+    <member name="P:System.Web.Razor.Generator.SetBaseTypeCodeGenerator.BaseType">
+      <summary>This type/member supports the .NET Framework infrastructure and is not intended to be used directly from your code. Gets or sets the set base type.</summary>
+      <returns>The set base type.</returns>
+    </member>
+    <member name="M:System.Web.Razor.Generator.SetBaseTypeCodeGenerator.Equals(System.Object)">
+      <summary>This type/member supports the .NET Framework infrastructure and is not intended to be used directly from your code. Determines whether the specified object is equal to this instance.</summary>
+      <returns>true if the specified object is equal to this instance; otherwise, false.</returns>
+      <param name="obj">The object to compare to this instance.</param>
+    </member>
+    <member name="M:System.Web.Razor.Generator.SetBaseTypeCodeGenerator.GenerateCode(System.Web.Razor.Parser.SyntaxTree.Span,System.Web.Razor.Generator.CodeGeneratorContext)">
+      <summary>This type/member supports the .NET Framework infrastructure and is not intended to be used directly from your code. Generates the code for this set base type.</summary>
+      <param name="target">The source span that contains the set base type to generate code.</param>
+      <param name="context">The context of the code generator.</param>
+    </member>
+    <member name="M:System.Web.Razor.Generator.SetBaseTypeCodeGenerator.GetHashCode">
+      <summary>This type/member supports the .NET Framework infrastructure and is not intended to be used directly from your code. Gets the hash code for this current instance.</summary>
+      <returns>The hash code for this current instance.</returns>
+    </member>
+    <member name="M:System.Web.Razor.Generator.SetBaseTypeCodeGenerator.ResolveType(System.Web.Razor.Generator.CodeGeneratorContext,System.String)">
+      <summary>This type/member supports the .NET Framework infrastructure and is not intended to be used directly from your code. Resolves the given set base type.</summary>
+      <returns>The resolved set base type.</returns>
+      <param name="context">The context of the code generator.</param>
+      <param name="baseType">The set base type to resolve.</param>
+    </member>
+    <member name="M:System.Web.Razor.Generator.SetBaseTypeCodeGenerator.ToString">
+      <summary>This type/member supports the .NET Framework infrastructure and is not intended to be used directly from your code. Returns the string representation for this instance.</summary>
+      <returns>The string representation for this instance.</returns>
+    </member>
+    <member name="T:System.Web.Razor.Generator.SetLayoutCodeGenerator">
+      <summary>This type/member supports the .NET Framework infrastructure and is not intended to be used directly from your code.Represents a code generator that sets layout for the web Razor.</summary>
+    </member>
+    <member name="M:System.Web.Razor.Generator.SetLayoutCodeGenerator.#ctor(System.String)">
+      <summary>This type/member supports the .NET Framework infrastructure and is not intended to be used directly from your code.Initializes a new instance of the <see cref="T:System.Web.Razor.Generator.SetLayoutCodeGenerator" /> class.</summary>
+      <param name="layoutPath">The layout path.</param>
+    </member>
+    <member name="M:System.Web.Razor.Generator.SetLayoutCodeGenerator.Equals(System.Object)">
+      <summary>This type/member supports the .NET Framework infrastructure and is not intended to be used directly from your code.Determines whether the specified object is equal to the current object.</summary>
+      <returns>true if the specified object is equal to the current object; otherwise, false.</returns>
+      <param name="obj">The object to compare to.</param>
+    </member>
+    <member name="M:System.Web.Razor.Generator.SetLayoutCodeGenerator.GenerateCode(System.Web.Razor.Parser.SyntaxTree.Span,System.Web.Razor.Generator.CodeGeneratorContext)">
+      <summary>This type/member supports the .NET Framework infrastructure and is not intended to be used directly from your code.Generates a layout code.</summary>
+      <param name="target">The target where to generate the code.</param>
+      <param name="context">The code generator context.</param>
+    </member>
+    <member name="M:System.Web.Razor.Generator.SetLayoutCodeGenerator.GetHashCode">
+      <summary>This type/member supports the .NET Framework infrastructure and is not intended to be used directly from your code.Retrieves a hash code for this current instance.</summary>
+      <returns>A hash code for this current instance.</returns>
+    </member>
+    <member name="P:System.Web.Razor.Generator.SetLayoutCodeGenerator.LayoutPath">
+      <summary>This type/member supports the .NET Framework infrastructure and is not intended to be used directly from your code.Gets or sets the path of the layout code.</summary>
+      <returns>The path of the layout code.</returns>
+    </member>
+    <member name="M:System.Web.Razor.Generator.SetLayoutCodeGenerator.ToString">
+      <summary>This type/member supports the .NET Framework infrastructure and is not intended to be used directly from your code.Returns a string representation of this current instance.</summary>
+      <returns>A string representation of this current instance.</returns>
+    </member>
+    <member name="T:System.Web.Razor.Generator.SetVBOptionCodeGenerator">
+      <summary>This type/member supports the .NET Framework infrastructure and is not intended to be used directly from your code.Represents the conversion of the SetVBOptionCodeGenerator of the value.</summary>
+    </member>
+    <member name="M:System.Web.Razor.Generator.SetVBOptionCodeGenerator.#ctor(System.String,System.Boolean)">
+      <summary>This type/member supports the .NET Framework infrastructure and is not intended to be used directly from your code.Initializes a new instance of the <see cref="T:System.Web.Razor.Generator.SetVBOptionCodeGenerator" /> class.</summary>
+      <param name="optionName">The option name.</param>
+      <param name="value">true if the object has a value; otherwise, false.</param>
+    </member>
+    <member name="M:System.Web.Razor.Generator.SetVBOptionCodeGenerator.Explicit(System.Boolean)">
+      <summary>This type/member supports the .NET Framework infrastructure and is not intended to be used directly from your code.Converts the <see cref="T:System.Web.Razor.Generator.SetVBOptionCodeGenerator" /> explicitly to the on and off value.</summary>
+      <returns>The <see cref="T:System.Web.Razor.Generator.SetVBOptionCodeGenerator" /> explicitly converts to the on and off value.</returns>
+      <param name="onOffValue">true if the <see cref="T:System.Web.Razor.Generator.SetVBOptionCodeGenerator" /> converts to on and off value; otherwise, false.</param>
+    </member>
+    <member name="F:System.Web.Razor.Generator.SetVBOptionCodeGenerator.ExplicitCodeDomOptionName">
+      <summary>This type/member supports the .NET Framework infrastructure and is not intended to be used directly from your code.Represents the explicit code Dom option name.</summary>
+    </member>
+    <member name="M:System.Web.Razor.Generator.SetVBOptionCodeGenerator.GenerateCode(System.Web.Razor.Parser.SyntaxTree.Span,System.Web.Razor.Generator.CodeGeneratorContext)">
+      <summary>This type/member supports the .NET Framework infrastructure and is not intended to be used directly from your code.Generates code for the specified parameters.</summary>
+      <param name="target">The target.</param>
+      <param name="context">The code generator context.</param>
+    </member>
+    <member name="P:System.Web.Razor.Generator.SetVBOptionCodeGenerator.OptionName">
+      <summary>This type/member supports the .NET Framework infrastructure and is not intended to be used directly from your code.Gets the option name for the code generator.</summary>
+      <returns>The option name for the code generator.</returns>
+    </member>
+    <member name="M:System.Web.Razor.Generator.SetVBOptionCodeGenerator.Strict(System.Boolean)">
+      <summary>This type/member supports the .NET Framework infrastructure and is not intended to be used directly from your code.Strictly converts the <see cref="T:System.Web.Razor.Generator.SetVBOptionCodeGenerator" /> to the on and off value.</summary>
+      <returns>The <see cref="T:System.Web.Razor.Generator.SetVBOptionCodeGenerator" /> strictly converts to the on and off value.</returns>
+      <param name="onOffValue">true if the <see cref="T:System.Web.Razor.Generator.SetVBOptionCodeGenerator" /> strictly converts to the on and off value; otherwise, false.</param>
+    </member>
+    <member name="F:System.Web.Razor.Generator.SetVBOptionCodeGenerator.StrictCodeDomOptionName">
+      <summary>This type/member supports the .NET Framework infrastructure and is not intended to be used directly from your code.Represents the strict code Dom option name.</summary>
+    </member>
+    <member name="M:System.Web.Razor.Generator.SetVBOptionCodeGenerator.ToString">
+      <summary>This type/member supports the .NET Framework infrastructure and is not intended to be used directly from your code.Returns a String that represents the current Object.</summary>
+      <returns>A String that represents the current Object.</returns>
+    </member>
+    <member name="P:System.Web.Razor.Generator.SetVBOptionCodeGenerator.Value">
+      <summary>This type/member supports the .NET Framework infrastructure and is not intended to be used directly from your code.Gets a value that indicates whether the <see cref="T:System.Web.Razor.Generator.SetVBOptionCodeGenerator" /> has a value.</summary>
+      <returns>true if the <see cref="T:System.Web.Razor.Generator.SetVBOptionCodeGenerator" /> has a value; otherwise, false.</returns>
+    </member>
+    <member name="T:System.Web.Razor.Generator.SpanCodeGenerator">
+      <summary>This type/member supports the .NET Framework infrastructure and is not intended to be used directly from your code.Represents the span code generator for the razor.</summary>
+    </member>
+    <member name="M:System.Web.Razor.Generator.SpanCodeGenerator.#ctor">
+      <summary>This type/member supports the .NET Framework infrastructure and is not intended to be used directly from your code.Initializes a new instance of the <see cref="T:System.Web.Razor.Generator.SpanCodeGenerator" /> class.</summary>
+    </member>
+    <member name="M:System.Web.Razor.Generator.SpanCodeGenerator.Equals(System.Object)">
+      <summary>This type/member supports the .NET Framework infrastructure and is not intended to be used directly from your code.Determines whether the specified object is equal to the current object.</summary>
+      <returns>true if the specified object is equal to the current object; otherwise, false.</returns>
+      <param name="obj">The object to compare with the current object.</param>
+    </member>
+    <member name="M:System.Web.Razor.Generator.SpanCodeGenerator.GenerateCode(System.Web.Razor.Parser.SyntaxTree.Span,System.Web.Razor.Generator.CodeGeneratorContext)">
+      <summary>This type/member supports the .NET Framework infrastructure and is not intended to be used directly from your code.Generates a code for the specified target and context parameters.</summary>
+      <param name="target">The target span.</param>
+      <param name="context">The code generator context.</param>
+    </member>
+    <member name="M:System.Web.Razor.Generator.SpanCodeGenerator.GetHashCode">
+      <summary>This type/member supports the .NET Framework infrastructure and is not intended to be used directly from your code.Returns a hash code for the span code generator.</summary>
+      <returns>A hash code for the span code generator.</returns>
+    </member>
+    <member name="F:System.Web.Razor.Generator.SpanCodeGenerator.Null">
+      <summary>This type/member supports the .NET Framework infrastructure and is not intended to be used directly from your code.Represents a null value for the span code generator.</summary>
+    </member>
+    <member name="T:System.Web.Razor.Generator.StatementCodeGenerator">
+      <summary>This type/member supports the .NET Framework infrastructure and is not intended to be used directly from your code. Represents a code generator for the statement.</summary>
+    </member>
+    <member name="M:System.Web.Razor.Generator.StatementCodeGenerator.#ctor">
+      <summary>This type/member supports the .NET Framework infrastructure and is not intended to be used directly from your code. Initializes a new instance of the <see cref="T:System.Web.Razor.Generator.StatementCodeGenerator" /> class.</summary>
+    </member>
+    <member name="M:System.Web.Razor.Generator.StatementCodeGenerator.Equals(System.Object)">
+      <summary>This type/member supports the .NET Framework infrastructure and is not intended to be used directly from your code. Determines whether the specified object is equal to this instance.</summary>
+      <returns>true if the specified object is equal to this instance; otherwise, false.</returns>
+      <param name="obj">The object to compare to this instance.</param>
+    </member>
+    <member name="M:System.Web.Razor.Generator.StatementCodeGenerator.GenerateCode(System.Web.Razor.Parser.SyntaxTree.Span,System.Web.Razor.Generator.CodeGeneratorContext)">
+      <summary>This type/member supports the .NET Framework infrastructure and is not intended to be used directly from your code. Generates the code for the statement.</summary>
+      <param name="target">The span source whose content contains the statement to generate.</param>
+      <param name="context">The context of the code generator.</param>
+    </member>
+    <member name="M:System.Web.Razor.Generator.StatementCodeGenerator.GetHashCode">
+      <summary>This type/member supports the .NET Framework infrastructure and is not intended to be used directly from your code. Returns the hash code for this current instance.</summary>
+      <returns>The hash code for this current instance.</returns>
+    </member>
+    <member name="M:System.Web.Razor.Generator.StatementCodeGenerator.ToString">
+      <summary>This type/member supports the .NET Framework infrastructure and is not intended to be used directly from your code. Returns the string representation for this instance.</summary>
+      <returns>The string representation for this instance.</returns>
+    </member>
+    <member name="T:System.Web.Razor.Generator.TemplateBlockCodeGenerator">
+      <summary>This type/member supports the .NET Framework infrastructure and is not intended to be used directly from your code.Represents the template block code generator of the razor.</summary>
+    </member>
+    <member name="M:System.Web.Razor.Generator.TemplateBlockCodeGenerator.#ctor">
+      <summary>This type/member supports the .NET Framework infrastructure and is not intended to be used directly from your code.Initializes a new instance of the <see cref="T:System.Web.Razor.Generator.TemplateBlockCodeGenerator" /> class.</summary>
+    </member>
+    <member name="M:System.Web.Razor.Generator.TemplateBlockCodeGenerator.GenerateEndBlockCode(System.Web.Razor.Parser.SyntaxTree.Block,System.Web.Razor.Generator.CodeGeneratorContext)">
+      <summary>This type/member supports the .NET Framework infrastructure and is not intended to be used directly from your code.Generates code to end the block of the template block code generator.</summary>
+      <param name="target">The target block.</param>
+      <param name="context">The code generator context.</param>
+    </member>
+    <member name="M:System.Web.Razor.Generator.TemplateBlockCodeGenerator.GenerateStartBlockCode(System.Web.Razor.Parser.SyntaxTree.Block,System.Web.Razor.Generator.CodeGeneratorContext)">
+      <summary>This type/member supports the .NET Framework infrastructure and is not intended to be used directly from your code.Generates code to start the block for the template block code generator.</summary>
+      <param name="target">The target block.</param>
+      <param name="context">The code generator context.</param>
+    </member>
+    <member name="T:System.Web.Razor.Generator.TypeMemberCodeGenerator">
+      <summary>This type/member supports the .NET Framework infrastructure and is not intended to be used directly from your code.Represents a type member code generator.</summary>
+    </member>
+    <member name="M:System.Web.Razor.Generator.TypeMemberCodeGenerator.#ctor">
+      <summary>This type/member supports the .NET Framework infrastructure and is not intended to be used directly from your code.Initializes a new instance of the <see cref="T:System.Web.Razor.Generator.TypeMemberCodeGenerator" /> class.</summary>
+    </member>
+    <member name="M:System.Web.Razor.Generator.TypeMemberCodeGenerator.Equals(System.Object)">
+      <summary>This type/member supports the .NET Framework infrastructure and is not intended to be used directly from your code.Determines whether the specified object is equal to the current object.</summary>
+      <returns>true if the specified object is equal to the current object; otherwise, false.</returns>
+      <param name="obj">The object to compare to.</param>
+    </member>
+    <member name="M:System.Web.Razor.Generator.TypeMemberCodeGenerator.GenerateCode(System.Web.Razor.Parser.SyntaxTree.Span,System.Web.Razor.Generator.CodeGeneratorContext)">
+      <summary>This type/member supports the .NET Framework infrastructure and is not intended to be used directly from your code.Generates code with a given target and context.</summary>
+      <param name="target">The target where to generate the code.</param>
+      <param name="context">The code generator context.</param>
+    </member>
+    <member name="M:System.Web.Razor.Generator.TypeMemberCodeGenerator.GetHashCode">
+      <summary>This type/member supports the .NET Framework infrastructure and is not intended to be used directly from your code.Retrieves the hash code for this current instance.</summary>
+      <returns>The hash code for this current instance.</returns>
+    </member>
+    <member name="M:System.Web.Razor.Generator.TypeMemberCodeGenerator.ToString">
+      <summary>This type/member supports the .NET Framework infrastructure and is not intended to be used directly from your code.Returns a string representation of this code.</summary>
+      <returns>A string representation of this code.</returns>
+    </member>
+    <member name="T:System.Web.Razor.Generator.VBRazorCodeGenerator">
+      <summary>This type/member supports the .NET Framework infrastructure and is not intended to be used directly from your code.Represents the razor code generator for VB.</summary>
+    </member>
+    <member name="M:System.Web.Razor.Generator.VBRazorCodeGenerator.#ctor(System.String,System.String,System.String,System.Web.Razor.RazorEngineHost)">
+      <summary>This type/member supports the .NET Framework infrastructure and is not intended to be used directly from your code.Initializes a new instance of the <see cref="T:System.Web.Razor.Generator.VBRazorCodeGenerator" /> class.</summary>
+      <param name="className">The name of the class.</param>
+      <param name="rootNamespaceName">The root namespace.</param>
+      <param name="sourceFileName">The file name of the asset source.</param>
+      <param name="host">The host.</param>
+    </member>
+    <member name="T:System.Web.Razor.Parser.BalancingModes">
+      <summary>This type/member supports the .NET Framework infrastructure and is not intended to be used directly from your code.</summary>
+    </member>
+    <member name="F:System.Web.Razor.Parser.BalancingModes.AllowCommentsAndTemplates">
+      <summary>This type/member supports the .NET Framework infrastructure and is not intended to be used directly from your code.</summary>
+    </member>
+    <member name="F:System.Web.Razor.Parser.BalancingModes.AllowEmbeddedTransitions">
+      <summary>This type/member supports the .NET Framework infrastructure and is not intended to be used directly from your code.</summary>
+    </member>
+    <member name="F:System.Web.Razor.Parser.BalancingModes.BacktrackOnFailure">
+      <summary>This type/member supports the .NET Framework infrastructure and is not intended to be used directly from your code.</summary>
+    </member>
+    <member name="F:System.Web.Razor.Parser.BalancingModes.NoErrorOnFailure">
+      <summary>This type/member supports the .NET Framework infrastructure and is not intended to be used directly from your code.</summary>
+    </member>
+    <member name="F:System.Web.Razor.Parser.BalancingModes.None">
+      <summary>This type/member supports the .NET Framework infrastructure and is not intended to be used directly from your code.</summary>
+    </member>
+    <member name="T:System.Web.Razor.Parser.CallbackVisitor">
+      <summary>This type/member supports the .NET Framework infrastructure and is not intended to be used directly from your code. Represents a visitor that executes a callback upon the completion of a visit.</summary>
+    </member>
+    <member name="M:System.Web.Razor.Parser.CallbackVisitor.#ctor(System.Action{System.Web.Razor.Parser.SyntaxTree.Span})">
+      <summary>This type/member supports the .NET Framework infrastructure and is not intended to be used directly from your code. Initializes a new instance of the <see cref="T:System.Web.Razor.Parser.CallbackVisitor" /> class.</summary>
+      <param name="spanCallback">The delegate for the span visit.</param>
+    </member>
+    <member name="M:System.Web.Razor.Parser.CallbackVisitor.#ctor(System.Action{System.Web.Razor.Parser.SyntaxTree.Span},System.Action{System.Web.Razor.Parser.SyntaxTree.RazorError})">
+      <summary>This type/member supports the .NET Framework infrastructure and is not intended to be used directly from your code. Initializes a new instance of the <see cref="T:System.Web.Razor.Parser.CallbackVisitor" /> class.</summary>
+      <param name="spanCallback">The delegate for the span visit.</param>
+      <param name="errorCallback">The delegate for the error visit.</param>
+    </member>
+    <member name="M:System.Web.Razor.Parser.CallbackVisitor.#ctor(System.Action{System.Web.Razor.Parser.SyntaxTree.Span},System.Action{System.Web.Razor.Parser.SyntaxTree.RazorError},System.Action{System.Web.Razor.Parser.SyntaxTree.BlockType},System.Action{System.Web.Razor.Parser.SyntaxTree.BlockType})">
+      <summary>This type/member supports the .NET Framework infrastructure and is not intended to be used directly from your code. Initializes a new instance of the <see cref="T:System.Web.Razor.Parser.CallbackVisitor" /> class.</summary>
+      <param name="spanCallback">The delegate for the span visit.</param>
+      <param name="errorCallback">The delegate for the error visit.</param>
+      <param name="startBlockCallback">The delegate for the start block visit.</param>
+      <param name="endBlockCallback">The delegate for the end block visit.</param>
+    </member>
+    <member name="M:System.Web.Razor.Parser.CallbackVisitor.#ctor(System.Action{System.Web.Razor.Parser.SyntaxTree.Span},System.Action{System.Web.Razor.Parser.SyntaxTree.RazorError},System.Action{System.Web.Razor.Parser.SyntaxTree.BlockType},System.Action{System.Web.Razor.Parser.SyntaxTree.BlockType},System.Action)">
+      <summary>This type/member supports the .NET Framework infrastructure and is not intended to be used directly from your code. Initializes a new instance of the <see cref="T:System.Web.Razor.Parser.CallbackVisitor" /> class.</summary>
+      <param name="spanCallback">The delegate for the span visit.</param>
+      <param name="errorCallback">The delegate for the error visit.</param>
+      <param name="startBlockCallback">The delegate for the start block visit.</param>
+      <param name="endBlockCallback">The delegate for the end block visit.</param>
+      <param name="completeCallback">The delegate to execute for the complete event.</param>
+    </member>
+    <member name="M:System.Web.Razor.Parser.CallbackVisitor.OnComplete">
+      <summary>This type/member supports the .NET Framework infrastructure and is not intended to be used directly from your code. </summary>
+    </member>
+    <member name="P:System.Web.Razor.Parser.CallbackVisitor.SynchronizationContext">
+      <summary>This type/member supports the .NET Framework infrastructure and is not intended to be used directly from your code. Gets or sets the synchronization context for this callback visitor.</summary>
+      <returns>The synchronization context for this callback visitor.</returns>
+    </member>
+    <member name="M:System.Web.Razor.Parser.CallbackVisitor.VisitEndBlock(System.Web.Razor.Parser.SyntaxTree.Block)">
+      <summary>This type/member supports the .NET Framework infrastructure and is not intended to be used directly from your code. Executes the visitor callback to visit the end block.</summary>
+      <param name="block">The end block to visit.</param>
+    </member>
+    <member name="M:System.Web.Razor.Parser.CallbackVisitor.VisitError(System.Web.Razor.Parser.SyntaxTree.RazorError)">
+      <summary>This type/member supports the .NET Framework infrastructure and is not intended to be used directly from your code. Executes the visitor callback to visit the error.</summary>
+      <param name="err">The Razor error to visit.</param>
+    </member>
+    <member name="M:System.Web.Razor.Parser.CallbackVisitor.VisitSpan(System.Web.Razor.Parser.SyntaxTree.Span)">
+      <summary>This type/member supports the .NET Framework infrastructure and is not intended to be used directly from your code. Executes the visitor callback to visit the span.</summary>
+      <param name="span">The span to visit.</param>
+    </member>
+    <member name="M:System.Web.Razor.Parser.CallbackVisitor.VisitStartBlock(System.Web.Razor.Parser.SyntaxTree.Block)">
+      <summary>This type/member supports the .NET Framework infrastructure and is not intended to be used directly from your code. Executes the visitor callback to visit the start block.</summary>
+      <param name="block">The start block to visit.</param>
+    </member>
+    <member name="T:System.Web.Razor.Parser.CSharpCodeParser">
+      <summary>This type/member supports the .NET Framework infrastructure and is not intended to be used directly from your code.Represents a C sharp code parser.</summary>
+    </member>
+    <member name="M:System.Web.Razor.Parser.CSharpCodeParser.#ctor">
+      <summary>This type/member supports the .NET Framework infrastructure and is not intended to be used directly from your code.Initializes a new instance of the <see cref="T:System.Web.Razor.Parser.CSharpCodeParser" /> class.</summary>
+    </member>
+    <member name="M:System.Web.Razor.Parser.CSharpCodeParser.AcceptIf(System.Web.Razor.Tokenizer.Symbols.CSharpKeyword)">
+      <summary>This type/member supports the .NET Framework infrastructure and is not intended to be used directly from your code.Determines whether the parser accepts the ‘IF’ keyword.</summary>
+      <returns>true if the parser accepts the ‘IF’ keyword; otherwise, false.</returns>
+      <param name="keyword">The keyword to accept.</param>
+    </member>
+    <member name="M:System.Web.Razor.Parser.CSharpCodeParser.AssertDirective(System.String)">
+      <summary>This type/member supports the .NET Framework infrastructure and is not intended to be used directly from your code.Asserts a directive code.</summary>
+      <param name="directive">The directive code to assert.</param>
+    </member>
+    <member name="M:System.Web.Razor.Parser.CSharpCodeParser.At(System.Web.Razor.Tokenizer.Symbols.CSharpKeyword)">
+      <summary>This type/member supports the .NET Framework infrastructure and is not intended to be used directly from your code.Determines whether the code contains ‘AT’ keyword.</summary>
+      <returns>true if the code contains ‘AT’ keyword; otherwise, false.</returns>
+      <param name="keyword">The keyword.</param>
+    </member>
+    <member name="M:System.Web.Razor.Parser.CSharpCodeParser.BaseTypeDirective(System.String,System.Func{System.String,System.Web.Razor.Generator.SpanCodeGenerator})">
+      <summary>Indicates the base type directive.</summary>
+      <param name="noTypeNameError">The no type name error.</param>
+      <param name="createCodeGenerator">The create code generator.</param>
+    </member>
+    <member name="M:System.Web.Razor.Parser.CSharpCodeParser.FunctionsDirective">
+      <summary>This type/member supports the .NET Framework infrastructure and is not intended to be used directly from your code.Indicates the functions directive.</summary>
+    </member>
+    <member name="M:System.Web.Razor.Parser.CSharpCodeParser.HandleEmbeddedTransition">
+      <summary>This type/member supports the .NET Framework infrastructure and is not intended to be used directly from your code.Indicates the code that handles embedded transition.</summary>
+    </member>
+    <member name="M:System.Web.Razor.Parser.CSharpCodeParser.HelperDirective">
+      <summary>This type/member supports the .NET Framework infrastructure and is not intended to be used directly from your code.Indicates a helper directive.</summary>
+    </member>
+    <member name="M:System.Web.Razor.Parser.CSharpCodeParser.InheritsDirective">
+      <summary>This type/member supports the .NET Framework infrastructure and is not intended to be used directly from your code.Indicates which class the application will derive the view from, and can therefore ensure proper type checking.</summary>
+    </member>
+    <member name="M:System.Web.Razor.Parser.CSharpCodeParser.InheritsDirectiveCore">
+      <summary>This type/member supports the .NET Framework infrastructure and is not intended to be used directly from your code.Inherits a directive core.</summary>
+    </member>
+    <member name="M:System.Web.Razor.Parser.CSharpCodeParser.IsAtEmbeddedTransition(System.Boolean,System.Boolean)">
+      <summary>This type/member supports the .NET Framework infrastructure and is not intended to be used directly from your code.Determines whether the code is at embedded transition.</summary>
+      <returns>true if the code is at embedded transition; otherwise, false.</returns>
+      <param name="allowTemplatesAndComments">true to allow templates and comments; otherwise, false.</param>
+      <param name="allowTransitions">true to allow transitions; otherwise, false.</param>
+    </member>
+    <member name="P:System.Web.Razor.Parser.CSharpCodeParser.IsNested">
+      <summary>This type/member supports the .NET Framework infrastructure and is not intended to be used directly from your code.Gets or sets a value that indicates whether the code is nested.</summary>
+      <returns>true if the code is nested; otherwise, false.</returns>
+    </member>
+    <member name="M:System.Web.Razor.Parser.CSharpCodeParser.IsSpacingToken(System.Boolean,System.Boolean)">
+      <summary>Indicates whether the lines and comments is spacing token.</summary>
+      <returns>The function that indicates the spacing token.</returns>
+      <param name="includeNewLines">true to include new lines; otherwise, false.</param>
+      <param name="includeComments">true to include comments; otherwise, false.</param>
+    </member>
+    <member name="P:System.Web.Razor.Parser.CSharpCodeParser.Keywords">
+      <summary>This type/member supports the .NET Framework infrastructure and is not intended to be used directly from your code.Gets or sets the C sharp language keywords.</summary>
+      <returns>The C sharp language keywords.</returns>
+    </member>
+    <member name="P:System.Web.Razor.Parser.CSharpCodeParser.Language">
+      <summary>This type/member supports the .NET Framework infrastructure and is not intended to be used directly from your code.Gets the specific language for parsing.</summary>
+      <returns>The specific language for parsing.</returns>
+    </member>
+    <member name="M:System.Web.Razor.Parser.CSharpCodeParser.LayoutDirective">
+      <summary>This type/member supports the .NET Framework infrastructure and is not intended to be used directly from your code.Indicates the layout directive.</summary>
+    </member>
+    <member name="M:System.Web.Razor.Parser.CSharpCodeParser.MapDirectives(System.Action,System.String[])">
+      <summary>This type/member supports the .NET Framework infrastructure and is not intended to be used directly from your code.Maps the given directives.</summary>
+      <param name="handler">The handler.</param>
+      <param name="directives">The directives.</param>
+    </member>
+    <member name="P:System.Web.Razor.Parser.CSharpCodeParser.OtherParser">
+      <summary>This type/member supports the .NET Framework infrastructure and is not intended to be used directly from your code.Gets the other parser used for the code.</summary>
+      <returns>The other parser used for the code.</returns>
+    </member>
+    <member name="M:System.Web.Razor.Parser.CSharpCodeParser.OutputSpanBeforeRazorComment">
+      <summary>This type/member supports the .NET Framework infrastructure and is not intended to be used directly from your code.Spans the output of the parsing before the comment.</summary>
+    </member>
+    <member name="M:System.Web.Razor.Parser.CSharpCodeParser.ParseBlock">
+      <summary>This type/member supports the .NET Framework infrastructure and is not intended to be used directly from your code.Blocks the parsing.</summary>
+    </member>
+    <member name="M:System.Web.Razor.Parser.CSharpCodeParser.ReservedDirective(System.Boolean)">
+      <summary>This type/member supports the .NET Framework infrastructure and is not intended to be used directly from your code.Indicates the reserved directive.</summary>
+      <param name="topLevel">Determines whether the directive is a top level.</param>
+    </member>
+    <member name="M:System.Web.Razor.Parser.CSharpCodeParser.SectionDirective">
+      <summary>This type/member supports the .NET Framework infrastructure and is not intended to be used directly from your code.Indicates a section directive.</summary>
+    </member>
+    <member name="M:System.Web.Razor.Parser.CSharpCodeParser.SessionStateDirective">
+      <summary>This type/member supports the .NET Framework infrastructure and is not intended to be used directly from your code.Indicates a session state directive.</summary>
+    </member>
+    <member name="M:System.Web.Razor.Parser.CSharpCodeParser.SessionStateDirectiveCore">
+      <summary>This type/member supports the .NET Framework infrastructure and is not intended to be used directly from your code.Indicates the session state directive core.</summary>
+    </member>
+    <member name="M:System.Web.Razor.Parser.CSharpCodeParser.SessionStateTypeDirective(System.String,System.Func{System.String,System.String,System.Web.Razor.Generator.SpanCodeGenerator})">
+      <summary>Indicates the directive for session state type.</summary>
+      <param name="noValueError">The no value error.</param>
+      <param name="createCodeGenerator">The create code generator.</param>
+    </member>
+    <member name="M:System.Web.Razor.Parser.CSharpCodeParser.TryGetDirectiveHandler(System.String,System.Action@)">
+      <summary>This type/member supports the .NET Framework infrastructure and is not intended to be used directly from your code.Gets a directive handler.</summary>
+      <returns>true if successful; otherwise, false.</returns>
+      <param name="directive">The directive.</param>
+      <param name="handler">The handler.</param>
+    </member>
+    <member name="M:System.Web.Razor.Parser.CSharpCodeParser.ValidSessionStateValue">
+      <summary>This type/member supports the .NET Framework infrastructure and is not intended to be used directly from your code.Determines whether the value of the session state is valid.</summary>
+      <returns>true if the value of the session state is valid; otherwise, false.</returns>
+    </member>
+    <member name="T:System.Web.Razor.Parser.CSharpCodeParser.Block">
+      <summary>This type/member supports the .NET Framework infrastructure and is not intended to be used directly from your code.Represents the block for this CSharpCode parser.</summary>
+    </member>
+    <member name="M:System.Web.Razor.Parser.CSharpCodeParser.Block.#ctor(System.String,System.Web.Razor.Text.SourceLocation)">
+      <summary>This type/member supports the .NET Framework infrastructure and is not intended to be used directly from your code.Initializes a new instance of the <see cref="T:System.Web.Razor.Parser.CSharpCodeParser.Block" /> class.</summary>
+      <param name="name">The string name.</param>
+      <param name="start">The start of the source location.</param>
+    </member>
+    <member name="M:System.Web.Razor.Parser.CSharpCodeParser.Block.#ctor(System.Web.Razor.Tokenizer.Symbols.CSharpSymbol)">
+      <summary>This type/member supports the .NET Framework infrastructure and is not intended to be used directly from your code.Initializes a new instance of the <see cref="T:System.Web.Razor.Parser.CSharpCodeParser.Block" /> class.</summary>
+      <param name="symbol">The CSharp symbol.</param>
+    </member>
+    <member name="P:System.Web.Razor.Parser.CSharpCodeParser.Block.Name">
+      <summary>This type/member supports the .NET Framework infrastructure and is not intended to be used directly from your code.Gets or sets the string name for the block.</summary>
+      <returns>The string name for the block.</returns>
+    </member>
+    <member name="P:System.Web.Razor.Parser.CSharpCodeParser.Block.Start">
+      <summary>This type/member supports the .NET Framework infrastructure and is not intended to be used directly from your code.Gets or sets the source location to start the block.</summary>
+      <returns>The source location to start the block.</returns>
+    </member>
+    <member name="T:System.Web.Razor.Parser.CSharpLanguageCharacteristics">
+      <summary>This type/member supports the .NET Framework infrastructure and is not intended to be used directly from your code.Represents the different language characteristics in a CSharp language.</summary>
+    </member>
+    <member name="M:System.Web.Razor.Parser.CSharpLanguageCharacteristics.CreateMarkerSymbol(System.Web.Razor.Text.SourceLocation)">
+      <summary>This type/member supports the .NET Framework infrastructure and is not intended to be used directly from your code.Creates a marker symbol in the code.</summary>
+      <returns>A marker symbol in the code.</returns>
+      <param name="location">The source location.</param>
+    </member>
+    <member name="M:System.Web.Razor.Parser.CSharpLanguageCharacteristics.CreateSymbol(System.Web.Razor.Text.SourceLocation,System.String,System.Web.Razor.Tokenizer.Symbols.CSharpSymbolType,System.Collections.Generic.IEnumerable{System.Web.Razor.Parser.SyntaxTree.RazorError})">
+      <summary>This type/member supports the .NET Framework infrastructure and is not intended to be used directly from your code.Creates a symbol in the code.</summary>
+      <returns>A symbol in the code.</returns>
+      <param name="location">The source location.</param>
+      <param name="content">The content value.</param>
+      <param name="type">The html symbol type.</param>
+      <param name="errors">List of errors.</param>
+    </member>
+    <member name="M:System.Web.Razor.Parser.CSharpLanguageCharacteristics.CreateTokenizer(System.Web.Razor.Text.ITextDocument)">
+      <summary>This type/member supports the .NET Framework infrastructure and is not intended to be used directly from your code.Creates a language tokenizer.</summary>
+      <returns>A language tokenizer.</returns>
+      <param name="source">The source of the text document.</param>
+    </member>
+    <member name="M:System.Web.Razor.Parser.CSharpLanguageCharacteristics.FlipBracket(System.Web.Razor.Tokenizer.Symbols.CSharpSymbolType)">
+      <summary>This type/member supports the .NET Framework infrastructure and is not intended to be used directly from your code.Flips the bracket symbol in the code.</summary>
+      <returns>The bracket symbol in the code.</returns>
+      <param name="bracket">The symbol bracket.</param>
+    </member>
+    <member name="M:System.Web.Razor.Parser.CSharpLanguageCharacteristics.GetKeyword(System.Web.Razor.Tokenizer.Symbols.CSharpKeyword)">
+      <summary>This type/member supports the .NET Framework infrastructure and is not intended to be used directly from your code.Gets the keyword in the code.</summary>
+      <returns>The keyword in the code.</returns>
+      <param name="keyword">The keyword.</param>
+    </member>
+    <member name="M:System.Web.Razor.Parser.CSharpLanguageCharacteristics.GetKnownSymbolType(System.Web.Razor.Tokenizer.Symbols.KnownSymbolType)">
+      <summary>This type/member supports the .NET Framework infrastructure and is not intended to be used directly from your code.Gets the <see cref="T:System.Web.Razor.Parser.CSharpLanguageCharacteristics" /> in the code.</summary>
+      <returns>The <see cref="T:System.Web.Razor.Parser.CSharpLanguageCharacteristics" /> in the code.</returns>
+      <param name="type">The <see cref="T:System.Web.Razor.Tokenizer.Symbols.KnownSymbolType" />.</param>
+    </member>
+    <member name="M:System.Web.Razor.Parser.CSharpLanguageCharacteristics.GetSample(System.Web.Razor.Tokenizer.Symbols.CSharpSymbolType)">
+      <summary>This type/member supports the .NET Framework infrastructure and is not intended to be used directly from your code.Gets a sample symbol in the code.</summary>
+      <returns>A sample symbol in the code.</returns>
+      <param name="type">The <see cref="T:System.Web.Razor.Tokenizer.Symbols.HtmlSymbolType" />.</param>
+    </member>
+    <member name="M:System.Web.Razor.Parser.CSharpLanguageCharacteristics.GetSymbolSample(System.Web.Razor.Tokenizer.Symbols.CSharpSymbolType)">
+      <summary>This type/member supports the .NET Framework infrastructure and is not intended to be used directly from your code.Gets a sample symbol in the code.</summary>
+      <returns>A sample symbol in the code.</returns>
+      <param name="type">The <see cref="T:System.Web.Razor.Tokenizer.Symbols.HtmlSymbolType" />.</param>
+    </member>
+    <member name="P:System.Web.Razor.Parser.CSharpLanguageCharacteristics.Instance">
+      <summary>This type/member supports the .NET Framework infrastructure and is not intended to be used directly from your code.Gets the instance for the <see cref="T:System.Web.Razor.Parser.CSharpLanguageCharacteristics" /> class.</summary>
+      <returns>The instance for the <see cref="T:System.Web.Razor.Parser.CSharpLanguageCharacteristics" /> class.</returns>
+    </member>
+    <member name="T:System.Web.Razor.Parser.HtmlLanguageCharacteristics">
+      <summary>This type/member supports the .NET Framework infrastructure and is not intended to be used directly from your code.Represents the different language characteristics in an html.</summary>
+    </member>
+    <member name="M:System.Web.Razor.Parser.HtmlLanguageCharacteristics.CreateMarkerSymbol(System.Web.Razor.Text.SourceLocation)">
+      <summary>This type/member supports the .NET Framework infrastructure and is not intended to be used directly from your code.Creates a marker symbol in the Html.</summary>
+      <returns>A marker symbol in the Html.</returns>
+      <param name="location">The source location.</param>
+    </member>
+    <member name="M:System.Web.Razor.Parser.HtmlLanguageCharacteristics.CreateSymbol(System.Web.Razor.Text.SourceLocation,System.String,System.Web.Razor.Tokenizer.Symbols.HtmlSymbolType,System.Collections.Generic.IEnumerable{System.Web.Razor.Parser.SyntaxTree.RazorError})">
+      <summary>This type/member supports the .NET Framework infrastructure and is not intended to be used directly from your code.Creates a symbol in the Html.</summary>
+      <returns>A symbol in the Html.</returns>
+      <param name="location">The source location.</param>
+      <param name="content">The content value.</param>
+      <param name="type">The html symbol type.</param>
+      <param name="errors">List of errors.</param>
+    </member>
+    <member name="M:System.Web.Razor.Parser.HtmlLanguageCharacteristics.CreateTokenizer(System.Web.Razor.Text.ITextDocument)">
+      <summary>This type/member supports the .NET Framework infrastructure and is not intended to be used directly from your code.Creates an html tokenizer.</summary>
+      <returns>An html tokenizer.</returns>
+      <param name="source">The source of the text document.</param>
+    </member>
+    <member name="M:System.Web.Razor.Parser.HtmlLanguageCharacteristics.FlipBracket(System.Web.Razor.Tokenizer.Symbols.HtmlSymbolType)">
+      <summary>This type/member supports the .NET Framework infrastructure and is not intended to be used directly from your code.Flips the bracket symbol in the html.</summary>
+      <returns>The bracket symbol in the html.</returns>
+      <param name="bracket">The symbol bracket.</param>
+    </member>
+    <member name="M:System.Web.Razor.Parser.HtmlLanguageCharacteristics.GetKnownSymbolType(System.Web.Razor.Tokenizer.Symbols.KnownSymbolType)">
+      <summary>This type/member supports the .NET Framework infrastructure and is not intended to be used directly from your code.Gets the <see cref="T:System.Web.Razor.Tokenizer.Symbols.KnownSymbolType" /> in the html.</summary>
+      <returns>The <see cref="T:System.Web.Razor.Tokenizer.Symbols.KnownSymbolType" /> in the html.</returns>
+      <param name="type">The <see cref="T:System.Web.Razor.Tokenizer.Symbols.KnownSymbolType" />.</param>
+    </member>
+    <member name="M:System.Web.Razor.Parser.HtmlLanguageCharacteristics.GetSample(System.Web.Razor.Tokenizer.Symbols.HtmlSymbolType)">
+      <summary>This type/member supports the .NET Framework infrastructure and is not intended to be used directly from your code.Gets a sample symbol in the html.</summary>
+      <returns>A sample symbol in the html.</returns>
+      <param name="type">The <see cref="T:System.Web.Razor.Tokenizer.Symbols.HtmlSymbolType" />.</param>
+    </member>
+    <member name="P:System.Web.Razor.Parser.HtmlLanguageCharacteristics.Instance">
+      <summary>This type/member supports the .NET Framework infrastructure and is not intended to be used directly from your code.Gets the instance for the <see cref="T:System.Web.Razor.Parser.HtmlLanguageCharacteristics" /> class.</summary>
+      <returns>The instance for the <see cref="T:System.Web.Razor.Parser.HtmlLanguageCharacteristics" /> class.</returns>
+    </member>
+    <member name="T:System.Web.Razor.Parser.HtmlMarkupParser">
+      <summary>This type/member supports the .NET Framework infrastructure and is not intended to be used directly from your code. Represents a parser specifically for parsing HTML markup.</summary>
+    </member>
+    <member name="M:System.Web.Razor.Parser.HtmlMarkupParser.#ctor">
+      <summary>This type/member supports the .NET Framework infrastructure and is not intended to be used directly from your code. Initializes a new instance of the <see cref="T:System.Web.Razor.Parser.HtmlMarkupParser" /> class.</summary>
+    </member>
+    <member name="M:System.Web.Razor.Parser.HtmlMarkupParser.BuildSpan(System.Web.Razor.Parser.SyntaxTree.SpanBuilder,System.Web.Razor.Text.SourceLocation,System.String)">
+      <summary>This type/member supports the .NET Framework infrastructure and is not intended to be used directly from your code. Builds the span for the given content using the specified span builder.</summary>
+      <param name="span">The span builder used to build the span.</param>
+      <param name="start">The start location.</param>
+      <param name="content">The span content.</param>
+    </member>
+    <member name="M:System.Web.Razor.Parser.HtmlMarkupParser.IsSpacingToken(System.Boolean)">
+      <summary>This type/member supports the .NET Framework infrastructure and is not intended to be used directly from your code. Returns the function delegate used to determine the token used for HTML spacing.</summary>
+      <returns>The function delegate used to determine the token used for HTML spacing.</returns>
+      <param name="includeNewLines">true to indicate that new lines are considered as spacing token; otherwise, false.</param>
+    </member>
+    <member name="P:System.Web.Razor.Parser.HtmlMarkupParser.Language">
+      <summary>This type/member supports the .NET Framework infrastructure and is not intended to be used directly from your code. Gets the instance that defines the characteristics of HTML language.</summary>
+      <returns>The instance that defines the characteristics of HTML language.</returns>
+    </member>
+    <member name="P:System.Web.Razor.Parser.HtmlMarkupParser.OtherParser">
+      <summary>This type/member supports the .NET Framework infrastructure and is not intended to be used directly from your code. Gets the other parser for parsing HTML markup.</summary>
+      <returns>The other parser for parsing HTML markup.</returns>
+    </member>
+    <member name="M:System.Web.Razor.Parser.HtmlMarkupParser.OutputSpanBeforeRazorComment">
+      <summary>This type/member supports the .NET Framework infrastructure and is not intended to be used directly from your code. Builds the span before the Razor comment.</summary>
+    </member>
+    <member name="M:System.Web.Razor.Parser.HtmlMarkupParser.ParseBlock">
+      <summary>This type/member supports the .NET Framework infrastructure and is not intended to be used directly from your code. Parses the next HTML block.</summary>
+    </member>
+    <member name="M:System.Web.Razor.Parser.HtmlMarkupParser.ParseDocument">
+      <summary>This type/member supports the .NET Framework infrastructure and is not intended to be used directly from your code. Parses the HTML document.</summary>
+    </member>
+    <member name="M:System.Web.Razor.Parser.HtmlMarkupParser.ParseSection(System.Tuple{System.String,System.String},System.Boolean)">
+      <summary>Parses a section with markups given by the nesting sequences.</summary>
+      <param name="nestingSequences">A tuple that specifies the markup nesting sequences.</param>
+      <param name="caseSensitive">true to indicate case-sensitive parsing; otherwise, false.</param>
+    </member>
+    <member name="M:System.Web.Razor.Parser.HtmlMarkupParser.SkipToAndParseCode(System.Func{System.Web.Razor.Tokenizer.Symbols.HtmlSymbol,System.Boolean})">
+      <summary>This type/member supports the .NET Framework infrastructure and is not intended to be used directly from your code. Skips the parse until the specified condition is meet.</summary>
+      <param name="condition">A function delegate that defines the condition.</param>
+    </member>
+    <member name="M:System.Web.Razor.Parser.HtmlMarkupParser.SkipToAndParseCode(System.Web.Razor.Tokenizer.Symbols.HtmlSymbolType)">
+      <summary>This type/member supports the .NET Framework infrastructure and is not intended to be used directly from your code. Skips the parse until the specified HTML symbol type is encountered.</summary>
+      <param name="type">The HTML symbol type.</param>
+    </member>
+    <member name="P:System.Web.Razor.Parser.HtmlMarkupParser.VoidElements">
+      <summary>This type/member supports the .NET Framework infrastructure and is not intended to be used directly from your code. Gets the HTML tags that are considered as void.</summary>
+      <returns>The HTML tags that are considered as void.</returns>
+    </member>
+    <member name="T:System.Web.Razor.Parser.LanguageCharacteristics`3">
+      <summary>This type/member supports the .NET Framework infrastructure and is not intended to be used directly from your code. Provides methods that define the behavior of a Razor code language.</summary>
+      <typeparam name="TTokenizer">The type of the code tokenizer for the Razor language.</typeparam>
+      <typeparam name="TSymbol">The type for the language symbol.</typeparam>
+      <typeparam name="TSymbolType">The enumeration type for the language symbol.</typeparam>
+    </member>
+    <member name="M:System.Web.Razor.Parser.LanguageCharacteristics`3.#ctor">
+      <summary>This type/member supports the .NET Framework infrastructure and is not intended to be used directly from your code. Initializes a new instance of the <see cref="T:System.Web.Razor.Parser.LanguageCharacteristics`3" /> class.</summary>
+    </member>
+    <member name="M:System.Web.Razor.Parser.LanguageCharacteristics`3.CreateMarkerSymbol(System.Web.Razor.Text.SourceLocation)">
+      <summary>This type/member supports the .NET Framework infrastructure and is not intended to be used directly from your code. Creates a code language symbol with the specified source location as the start marker.</summary>
+      <returns>The symbol for the code language.</returns>
+      <param name="location">The source location as the start marker.</param>
+    </member>
+    <member name="M:System.Web.Razor.Parser.LanguageCharacteristics`3.CreateSymbol(System.Web.Razor.Text.SourceLocation,System.String,`2,System.Collections.Generic.IEnumerable{System.Web.Razor.Parser.SyntaxTree.RazorError})">
+      <summary>This type/member supports the .NET Framework infrastructure and is not intended to be used directly from your code. Creates a code language symbol with the specified source location with the specified source location as the start marker.</summary>
+      <returns>The symbol for the code language.</returns>
+      <param name="location">The source location as the start marker.</param>
+      <param name="content">The content.</param>
+      <param name="type">The enumeration type for the language symbol.</param>
+      <param name="errors">The collection of error.</param>
+    </member>
+    <member name="M:System.Web.Razor.Parser.LanguageCharacteristics`3.CreateTokenizer(System.Web.Razor.Text.ITextDocument)">
+      <summary>This type/member supports the .NET Framework infrastructure and is not intended to be used directly from your code. Creates a Razor code language tokenizer for the specified source document.</summary>
+      <returns>A Razor code language tokenizer for the specified source document.</returns>
+      <param name="source">The source document.</param>
+    </member>
+    <member name="M:System.Web.Razor.Parser.LanguageCharacteristics`3.FlipBracket(`2)">
+      <summary>This type/member supports the .NET Framework infrastructure and is not intended to be used directly from your code. Returns the opposite bracket symbol for the specified bracket symbol.</summary>
+      <returns>The opposite bracket symbol for the specified bracket symbol.</returns>
+      <param name="bracket">The bracket symbol to flip.</param>
+    </member>
+    <member name="M:System.Web.Razor.Parser.LanguageCharacteristics`3.GetKnownSymbolType(System.Web.Razor.Tokenizer.Symbols.KnownSymbolType)">
+      <summary>This type/member supports the .NET Framework infrastructure and is not intended to be used directly from your code. Gets the specific language symbol type for the given symbol type.</summary>
+      <returns>The specific language symbol type for the given symbol type.</returns>
+      <param name="type">The symbol type to get.</param>
+    </member>
+    <member name="M:System.Web.Razor.Parser.LanguageCharacteristics`3.GetSample(`2)">
+      <summary>This type/member supports the .NET Framework infrastructure and is not intended to be used directly from your code. Gets the actual symbol for the given language symbol type.</summary>
+      <returns>The actual symbol for the given language symbol type.</returns>
+      <param name="type">The language symbol type to get.</param>
+    </member>
+    <member name="M:System.Web.Razor.Parser.LanguageCharacteristics`3.IsCommentBody(`1)">
+      <summary>This type/member supports the .NET Framework infrastructure and is not intended to be used directly from your code. Determines whether the symbol is a comment body type.</summary>
+      <returns>true if the symbol is a comment body type; otherwise, false.</returns>
+      <param name="symbol">The symbol to check.</param>
+    </member>
+    <member name="M:System.Web.Razor.Parser.LanguageCharacteristics`3.IsCommentStar(`1)">
+      <summary>This type/member supports the .NET Framework infrastructure and is not intended to be used directly from your code. Determines whether the symbol is a comment star type.</summary>
+      <returns>true if the symbol is a comment star type; otherwise, false.</returns>
+      <param name="symbol">The symbol to check.</param>
+    </member>
+    <member name="M:System.Web.Razor.Parser.LanguageCharacteristics`3.IsCommentStart(`1)">
+      <summary>This type/member supports the .NET Framework infrastructure and is not intended to be used directly from your code. Determines whether the symbol is a comment start type.</summary>
+      <returns>true if the symbol is a comment start type; otherwise, false. </returns>
+      <param name="symbol">The symbol to check.</param>
+    </member>
+    <member name="M:System.Web.Razor.Parser.LanguageCharacteristics`3.IsIdentifier(`1)">
+      <summary>This type/member supports the .NET Framework infrastructure and is not intended to be used directly from your code. Determines whether the symbol is an identifier type.</summary>
+      <returns>true if the symbol is an identifier type; otherwise, false.</returns>
+      <param name="symbol">The symbol to check.</param>
+    </member>
+    <member name="M:System.Web.Razor.Parser.LanguageCharacteristics`3.IsKeyword(`1)">
+      <summary>This type/member supports the .NET Framework infrastructure and is not intended to be used directly from your code. Determines whether the symbol is a keyword type.</summary>
+      <returns>true if the symbol is a keyword type; otherwise, false.</returns>
+      <param name="symbol">The symbol to check.</param>
+    </member>
+    <member name="M:System.Web.Razor.Parser.LanguageCharacteristics`3.IsKnownSymbolType(`1,System.Web.Razor.Tokenizer.Symbols.KnownSymbolType)">
+      <summary>This type/member supports the .NET Framework infrastructure and is not intended to be used directly from your code. Determines whether the symbol type is a known symbol type.</summary>
+      <returns>true if the symbol type is a known symbol type; otherwise, false.</returns>
+      <param name="symbol">The symbol whose type is to be checked.</param>
+      <param name="type">The known type of the symbol.</param>
+    </member>
+    <member name="M:System.Web.Razor.Parser.LanguageCharacteristics`3.IsNewLine(`1)">
+      <summary>This type/member supports the .NET Framework infrastructure and is not intended to be used directly from your code. Determines whether the symbol is a new line type.</summary>
+      <returns>true if the symbol is a new line type; otherwise, false.</returns>
+      <param name="symbol">The symbol to check.</param>
+    </member>
+    <member name="M:System.Web.Razor.Parser.LanguageCharacteristics`3.IsTransition(`1)">
+      <summary>This type/member supports the .NET Framework infrastructure and is not intended to be used directly from your code. Determines whether the symbol is a transition type.</summary>
+      <returns>true if the symbol is a transition type; otherwise, false.</returns>
+      <param name="symbol">The symbol to check.</param>
+    </member>
+    <member name="M:System.Web.Razor.Parser.LanguageCharacteristics`3.IsUnknown(`1)">
+      <summary>This type/member supports the .NET Framework infrastructure and is not intended to be used directly from your code. Determines whether the symbol is an unknown type.</summary>
+      <returns>true if the symbol is an unknown type; otherwise, false.</returns>
+      <param name="symbol">The symbol to check.</param>
+    </member>
+    <member name="M:System.Web.Razor.Parser.LanguageCharacteristics`3.IsWhiteSpace(`1)">
+      <summary>This type/member supports the .NET Framework infrastructure and is not intended to be used directly from your code. Determines whether the symbol is a whitespace type.</summary>
+      <returns>true if the symbol is a whitespace type; otherwise, false.</returns>
+      <param name="symbol">The symbol to check.</param>
+    </member>
+    <member name="M:System.Web.Razor.Parser.LanguageCharacteristics`3.KnowsSymbolType(System.Web.Razor.Tokenizer.Symbols.KnownSymbolType)">
+      <summary>This type/member supports the .NET Framework infrastructure and is not intended to be used directly from your code. Determines whether the symbol is an unknown type.</summary>
+      <returns>true if the symbol is an unknown type; otherwise, false.</returns>
+      <param name="type">The known type of the symbol.</param>
+    </member>
+    <member name="M:System.Web.Razor.Parser.LanguageCharacteristics`3.SplitSymbol(`1,System.Int32,`2)">
+      <summary>Splits the content of the code language symbol at the specified index.</summary>
+      <returns>A tuple of code language symbol.</returns>
+      <param name="symbol">The symbol whose content is to be splitted.</param>
+      <param name="splitAt">The index where the split occurs.</param>
+      <param name="leftType">The enumeration type for the language symbol.</param>
+    </member>
+    <member name="M:System.Web.Razor.Parser.LanguageCharacteristics`3.TokenizeString(System.String)">
+      <summary>This type/member supports the .NET Framework infrastructure and is not intended to be used directly from your code. Splits the specified string into tokens.</summary>
+      <returns>The collection of token.</returns>
+      <param name="content">The string to tokenize.</param>
+    </member>
+    <member name="M:System.Web.Razor.Parser.LanguageCharacteristics`3.TokenizeString(System.Web.Razor.Text.SourceLocation,System.String)">
+      <summary>This type/member supports the .NET Framework infrastructure and is not intended to be used directly from your code. Splits the specified string into tokens.</summary>
+      <returns>The collection of token.</returns>
+      <param name="start">The source location as the start marker for the tokenizer.</param>
+      <param name="input">The string to tokenize.</param>
+    </member>
+    <member name="T:System.Web.Razor.Parser.ParserBase">
+      <summary>This type/member supports the .NET Framework infrastructure and is not intended to be used directly from your code.Represents the parser base class for the razor.</summary>
+    </member>
+    <member name="M:System.Web.Razor.Parser.ParserBase.#ctor">
+      <summary>This type/member supports the .NET Framework infrastructure and is not intended to be used directly from your code.Initializes a new instance of the <see cref="T:System.Web.Razor.Parser.ParserBase" /> class.</summary>
+    </member>
+    <member name="M:System.Web.Razor.Parser.ParserBase.BuildSpan(System.Web.Razor.Parser.SyntaxTree.SpanBuilder,System.Web.Razor.Text.SourceLocation,System.String)">
+      <summary>This type/member supports the .NET Framework infrastructure and is not intended to be used directly from your code.Builds a span for the parser base.</summary>
+      <param name="span">The span builder.</param>
+      <param name="start">The beginning of the source location.</param>
+      <param name="content">The content.</param>
+    </member>
+    <member name="P:System.Web.Razor.Parser.ParserBase.Context">
+      <summary>This type/member supports the .NET Framework infrastructure and is not intended to be used directly from your code.Gets or sets the <see cref="T:System.Web.Razor.Parser.ParserContext" />.</summary>
+      <returns>The <see cref="T:System.Web.Razor.Parser.ParserContext" />.</returns>
+    </member>
+    <member name="P:System.Web.Razor.Parser.ParserBase.IsMarkupParser">
+      <summary>This type/member supports the .NET Framework infrastructure and is not intended to be used directly from your code.Gets a value that indicates whether the parser is a markup parser.</summary>
+      <returns>true if the parser is a markup parser; otherwise, false.</returns>
+    </member>
+    <member name="P:System.Web.Razor.Parser.ParserBase.OtherParser">
+      <summary>This type/member supports the .NET Framework infrastructure and is not intended to be used directly from your code.Gets the other parser <see cref="T:System.Web.Razor.Parser.ParserBase" />.</summary>
+      <returns>The other parser <see cref="T:System.Web.Razor.Parser.ParserBase" />.</returns>
+    </member>
+    <member name="M:System.Web.Razor.Parser.ParserBase.ParseBlock">
+      <summary>This type/member supports the .NET Framework infrastructure and is not intended to be used directly from your code.Blocks the parser.</summary>
+    </member>
+    <member name="M:System.Web.Razor.Parser.ParserBase.ParseDocument">
+      <summary>This type/member supports the .NET Framework infrastructure and is not intended to be used directly from your code.Creates documentation for the parse.</summary>
+    </member>
+    <member name="M:System.Web.Razor.Parser.ParserBase.ParseSection(System.Tuple{System.String,System.String},System.Boolean)">
+      <summary>Parses the section in ordered list of the elements.</summary>
+      <param name="nestingSequences">The pair of nesting sequences.</param>
+      <param name="caseSensitive">true if the case is sensitive; otherwise, false.</param>
+    </member>
+    <member name="T:System.Web.Razor.Parser.ParserContext">
+      <summary>This type/member supports the .NET Framework infrastructure and is not intended to be used directly from your code. Represents a parser whose context can be switched to either a code or a markup.</summary>
+    </member>
+    <member name="M:System.Web.Razor.Parser.ParserContext.#ctor(System.Web.Razor.Text.ITextDocument,System.Web.Razor.Parser.ParserBase,System.Web.Razor.Parser.ParserBase,System.Web.Razor.Parser.ParserBase)">
+      <summary>This type/member supports the .NET Framework infrastructure and is not intended to be used directly from your code. Initializes a new instance of the <see cref="T:System.Web.Razor.Parser.ParserContext" /> class.</summary>
+      <param name="source">The source document.</param>
+      <param name="codeParser">The code parser for the context.</param>
+      <param name="markupParser">The markup parser for the context.</param>
+      <param name="activeParser">The active parser for the context.</param>
+    </member>
+    <member name="P:System.Web.Razor.Parser.ParserContext.ActiveParser">
+      <summary>This type/member supports the .NET Framework infrastructure and is not intended to be used directly from your code. Gets or sets the active parser for the context.</summary>
+      <returns>The active parser for the context.</returns>
+    </member>
+    <member name="M:System.Web.Razor.Parser.ParserContext.AddSpan(System.Web.Razor.Parser.SyntaxTree.Span)">
+      <summary>This type/member supports the .NET Framework infrastructure and is not intended to be used directly from your code. Adds the specified span at the end of the block builder stack.</summary>
+      <param name="span">The span to add.</param>
+    </member>
+    <member name="P:System.Web.Razor.Parser.ParserContext.CodeParser">
+      <summary>This type/member supports the .NET Framework infrastructure and is not intended to be used directly from your code. Gets or sets the code parser for the context.</summary>
+      <returns>The code parser for the context.</returns>
+    </member>
+    <member name="M:System.Web.Razor.Parser.ParserContext.CompleteParse">
+      <summary>This type/member supports the .NET Framework infrastructure and is not intended to be used directly from your code. Parses the last span and returns the parse results that contain the newly built block.</summary>
+      <returns>The parse results that contain the newly built block.</returns>
+    </member>
+    <member name="P:System.Web.Razor.Parser.ParserContext.CurrentBlock">
+      <summary>This type/member supports the .NET Framework infrastructure and is not intended to be used directly from your code. Gets the current block builder.</summary>
+      <returns>The current block builder.</returns>
+    </member>
+    <member name="P:System.Web.Razor.Parser.ParserContext.CurrentCharacter">
+      <summary>This type/member supports the .NET Framework infrastructure and is not intended to be used directly from your code. Gets the current character available from the source.</summary>
+      <returns>The current character available from the source.</returns>
+    </member>
+    <member name="P:System.Web.Razor.Parser.ParserContext.DesignTimeMode">
+      <summary>This type/member supports the .NET Framework infrastructure and is not intended to be used directly from your code. Gets or sets a value that indicates whether the parser is in design mode.</summary>
+      <returns>true if the parser is in design mode; otherwise, false.</returns>
+    </member>
+    <member name="M:System.Web.Razor.Parser.ParserContext.EndBlock">
+      <summary>This type/member supports the .NET Framework infrastructure and is not intended to be used directly from your code. Creates an end block from the last item of the block builder stack. </summary>
+    </member>
+    <member name="P:System.Web.Razor.Parser.ParserContext.EndOfFile">
+      <summary>This type/member supports the .NET Framework infrastructure and is not intended to be used directly from your code. Gets a value that indicates whether the source status is end of file.</summary>
+      <returns>true if the source status is end of file; otherwise, false.</returns>
+    </member>
+    <member name="P:System.Web.Razor.Parser.ParserContext.Errors">
+      <summary>This type/member supports the .NET Framework infrastructure and is not intended to be used directly from your code. Gets or sets the list of errors during parsing.</summary>
+      <returns>The list of errors.</returns>
+    </member>
+    <member name="M:System.Web.Razor.Parser.ParserContext.IsWithin(System.Web.Razor.Parser.SyntaxTree.BlockType)">
+      <summary>This type/member supports the .NET Framework infrastructure and is not intended to be used directly from your code. Determines whether the specified block type exists in the block builder list.</summary>
+      <returns>true if the specified block type exists in the block builder list; otherwise, false.</returns>
+      <param name="type">The block type to check.</param>
+    </member>
+    <member name="P:System.Web.Razor.Parser.ParserContext.LastAcceptedCharacters">
+      <summary>This type/member supports the .NET Framework infrastructure and is not intended to be used directly from your code. Gets the last accepted characters.</summary>
+      <returns>One of the values of the <see cref="T:System.Web.Razor.Parser.SyntaxTree.AcceptedCharacters" /> enumeration.</returns>
+    </member>
+    <member name="P:System.Web.Razor.Parser.ParserContext.LastSpan">
+      <summary>This type/member supports the .NET Framework infrastructure and is not intended to be used directly from your code. Gets or sets the last span.</summary>
+      <returns>The last span.</returns>
+    </member>
+    <member name="P:System.Web.Razor.Parser.ParserContext.MarkupParser">
+      <summary>This type/member supports the .NET Framework infrastructure and is not intended to be used directly from your code. Gets or sets the markup parser for the context.</summary>
+      <returns>The markup parser for the context.</returns>
+    </member>
+    <member name="M:System.Web.Razor.Parser.ParserContext.OnError(System.Web.Razor.Text.SourceLocation,System.String)">
+      <summary>This type/member supports the .NET Framework infrastructure and is not intended to be used directly from your code. Occurs when parse encountered error.</summary>
+      <param name="location">The source location.</param>
+      <param name="message">The error message.</param>
+    </member>
+    <member name="M:System.Web.Razor.Parser.ParserContext.OnError(System.Web.Razor.Text.SourceLocation,System.String,System.Object[])">
+      <summary>This type/member supports the .NET Framework infrastructure and is not intended to be used directly from your code. Occurs when parse encountered an error.</summary>
+      <param name="location">The source location.</param>
+      <param name="message">The error message.</param>
+      <param name="args">The other information about the source location.</param>
+    </member>
+    <member name="P:System.Web.Razor.Parser.ParserContext.Source">
+      <summary>This type/member supports the .NET Framework infrastructure and is not intended to be used directly from your code. Gets or sets the text reader for the source document.</summary>
+      <returns>The text reader for the source document.</returns>
+    </member>
+    <member name="M:System.Web.Razor.Parser.ParserContext.StartBlock">
+      <summary>This type/member supports the .NET Framework infrastructure and is not intended to be used directly from your code. Adds a new block builder at the end of the block builder stack and returns a disposable action that returns an end block.</summary>
+      <returns>A disposable action that returns an end block.</returns>
+    </member>
+    <member name="M:System.Web.Razor.Parser.ParserContext.StartBlock(System.Web.Razor.Parser.SyntaxTree.BlockType)">
+      <summary>This type/member supports the .NET Framework infrastructure and is not intended to be used directly from your code. Adds a new block builder at the end of the block builder stack and returns a disposable action that returns an end block.</summary>
+      <returns>A disposable action that returns an end block.</returns>
+      <param name="blockType">The type for the new block builder.</param>
+    </member>
+    <member name="M:System.Web.Razor.Parser.ParserContext.SwitchActiveParser">
+      <summary>This type/member supports the .NET Framework infrastructure and is not intended to be used directly from your code. Alternately switches the code parser or markup parser as the active parser.</summary>
+    </member>
+    <member name="P:System.Web.Razor.Parser.ParserContext.WhiteSpaceIsSignificantToAncestorBlock">
+      <summary>This type/member supports the .NET Framework infrastructure and is not intended to be used directly from your code. Gets or sets a value that indicates whether white space is significant to ancestor block.</summary>
+      <returns>true is white space is significant to ancestor block; otherwise, false.</returns>
+    </member>
+    <member name="T:System.Web.Razor.Parser.ParserHelpers">
+      <summary>This type/member supports the .NET Framework infrastructure and is not intended to be used directly from your code. Provides helper methods for the parser.</summary>
+    </member>
+    <member name="M:System.Web.Razor.Parser.ParserHelpers.IsCombining(System.Char)">
+      <summary>This type/member supports the .NET Framework infrastructure and is not intended to be used directly from your code. Determines whether the specified character value is a spacing combining mark or a non-spacing mark.</summary>
+      <returns>true if the specified character value is a spacing combining mark or a non-spacing mark; otherwise, false.</returns>
+      <param name="value">The value to check.</param>
+    </member>
+    <member name="M:System.Web.Razor.Parser.ParserHelpers.IsConnecting(System.Char)">
+      <summary>This type/member supports the .NET Framework infrastructure and is not intended to be used directly from your code. Determines whether the specified character value is a connector punctuation.</summary>
+      <returns>true if the specified character value is a connector punctuation; otherwise, false.</returns>
+      <param name="value">The value to check.</param>
+    </member>
+    <member name="M:System.Web.Razor.Parser.ParserHelpers.IsDecimalDigit(System.Char)">
+      <summary>This type/member supports the .NET Framework infrastructure and is not intended to be used directly from your code. Determines whether the specified character value is a decimal digit number.</summary>
+      <returns>true if the specified character value is a decimal digit number; otherwise, false.</returns>
+      <param name="value">The value to check.</param>
+    </member>
+    <member name="M:System.Web.Razor.Parser.ParserHelpers.IsEmailPart(System.Char)">
+      <summary>This type/member supports the .NET Framework infrastructure and is not intended to be used directly from your code. Determines whether the specified character value is valid for use in email address.</summary>
+      <returns>true if the specified character value is valid for use in email address; otherwise, false.</returns>
+      <param name="character">The value to check.</param>
+    </member>
+    <member name="M:System.Web.Razor.Parser.ParserHelpers.IsFormatting(System.Char)">
+      <summary>This type/member supports the .NET Framework infrastructure and is not intended to be used directly from your code. Determines whether the specified character value is used for formatting text layout or formatting text operation.</summary>
+      <returns>true if the specified character value is used for formatting text layout or formatting text operation.; otherwise, false.</returns>
+      <param name="value">The value to check.</param>
+    </member>
+    <member name="M:System.Web.Razor.Parser.ParserHelpers.IsHexDigit(System.Char)">
+      <summary>This type/member supports the .NET Framework infrastructure and is not intended to be used directly from your code. Determines whether the specified character value is a hexadecimal digit number.</summary>
+      <returns>true if the specified character is a hexadecimal digit number; otherwise, false.</returns>
+      <param name="value">The value to check.</param>
+    </member>
+    <member name="M:System.Web.Razor.Parser.ParserHelpers.IsIdentifier(System.String)">
+      <summary>This type/member supports the .NET Framework infrastructure and is not intended to be used directly from your code. Determines whether the specified string value is an identifier.</summary>
+      <returns>true if the specified string value is an identifier; otherwise, false.</returns>
+      <param name="value">The value to check.</param>
+    </member>
+    <member name="M:System.Web.Razor.Parser.ParserHelpers.IsIdentifier(System.String,System.Boolean)">
+      <summary>This type/member supports the .NET Framework infrastructure and is not intended to be used directly from your code. Determines whether the specified string value is an identifier.</summary>
+      <returns>true if the specified string value is an identifier; otherwise, false.</returns>
+      <param name="value">The value to check.</param>
+      <param name="requireIdentifierStart">true to require that the identifier starts with a letter or underscore (_); otherwise, false.</param>
+    </member>
+    <member name="M:System.Web.Razor.Parser.ParserHelpers.IsIdentifierPart(System.Char)">
+      <summary>This type/member supports the .NET Framework infrastructure and is not intended to be used directly from your code. Determines whether the specified character value is valid for use in identifier.</summary>
+      <returns>true if the specified character is valid for use in identifier; otherwise, false.</returns>
+      <param name="value">The value to check.</param>
+    </member>
+    <member name="M:System.Web.Razor.Parser.ParserHelpers.IsIdentifierStart(System.Char)">
+      <summary>This type/member supports the .NET Framework infrastructure and is not intended to be used directly from your code. Determines whether the specified character value is valid for use as start character of an identifier.</summary>
+      <returns>true if the specified character value is valid for use as start character of an identifier; otherwise, false.</returns>
+      <param name="value">The value to check.</param>
+    </member>
+    <member name="M:System.Web.Razor.Parser.ParserHelpers.IsLetter(System.Char)">
+      <summary>This type/member supports the .NET Framework infrastructure and is not intended to be used directly from your code. Determines whether the specified character value is a letter.</summary>
+      <returns>true if the specified character is a letter; otherwise, false.</returns>
+      <param name="value">The value to check.</param>
+    </member>
+    <member name="M:System.Web.Razor.Parser.ParserHelpers.IsLetterOrDecimalDigit(System.Char)">
+      <summary>This type/member supports the .NET Framework infrastructure and is not intended to be used directly from your code. Determines whether the specified character value is a letter or a decimal digit number.</summary>
+      <returns>true if the specified character is a letter or a decimal digit number; otherwise, false.</returns>
+      <param name="value">The value to check.</param>
+    </member>
+    <member name="M:System.Web.Razor.Parser.ParserHelpers.IsNewLine(System.Char)">
+      <summary>This type/member supports the .NET Framework infrastructure and is not intended to be used directly from your code. Determines whether the specified value is a newline.</summary>
+      <returns>true if the specified character is a newline; otherwise, false.</returns>
+      <param name="value">The value to check.</param>
+    </member>
+    <member name="M:System.Web.Razor.Parser.ParserHelpers.IsNewLine(System.String)">
+      <summary>This type/member supports the .NET Framework infrastructure and is not intended to be used directly from your code. Determines whether the specified value is a newline.</summary>
+      <returns>true if the specified character is a newline; otherwise, false.</returns>
+      <param name="value">The value to check.</param>
+    </member>
+    <member name="M:System.Web.Razor.Parser.ParserHelpers.IsTerminatingCharToken(System.Char)">
+      <summary>This type/member supports the .NET Framework infrastructure and is not intended to be used directly from your code. Determines whether the specified character value is a terminating character token.</summary>
+      <returns>true if the specified character value is a terminating character token; otherwise, false.</returns>
+      <param name="value">The value to check.</param>
+    </member>
+    <member name="M:System.Web.Razor.Parser.ParserHelpers.IsTerminatingQuotedStringToken(System.Char)">
+      <summary>This type/member supports the .NET Framework infrastructure and is not intended to be used directly from your code. Determines whether the specified character value is a terminating quoted string.</summary>
+      <returns>true if the specified character value is a terminating quoted string; otherwise, false.</returns>
+      <param name="value">The value to check.</param>
+    </member>
+    <member name="M:System.Web.Razor.Parser.ParserHelpers.IsWhitespace(System.Char)">
+      <summary>This type/member supports the .NET Framework infrastructure and is not intended to be used directly from your code. Determines whether the specified character value is a whitespace.</summary>
+      <returns>true if the specified character value is a whitespace; otherwise, false.</returns>
+      <param name="value">The value to check.</param>
+    </member>
+    <member name="M:System.Web.Razor.Parser.ParserHelpers.IsWhitespaceOrNewLine(System.Char)">
+      <summary>This type/member supports the .NET Framework infrastructure and is not intended to be used directly from your code. Determines whether the specified character value is a whitespace or newline.</summary>
+      <returns>true if the specified character value is a whitespace or newline; otherwise, false.</returns>
+      <param name="value">The value to check.</param>
+    </member>
+    <member name="M:System.Web.Razor.Parser.ParserHelpers.SanitizeClassName(System.String)">
+      <summary>This type/member supports the .NET Framework infrastructure and is not intended to be used directly from your code. Sanitizes the specified input name to conform as a valid value for class name.</summary>
+      <returns>The sanitized class name.</returns>
+      <param name="inputName">The value to check.</param>
+    </member>
+    <member name="T:System.Web.Razor.Parser.ParserVisitor">
+      <summary>This type/member supports the .NET Framework infrastructure and is not intended to be used directly from your code.Represents a parser visitor.</summary>
+    </member>
+    <member name="M:System.Web.Razor.Parser.ParserVisitor.#ctor">
+      <summary>This type/member supports the .NET Framework infrastructure and is not intended to be used directly from your code.Initializes a new instance of the <see cref="T:System.Web.Razor.Parser.ParserVisitor" /> class.</summary>
+    </member>
+    <member name="P:System.Web.Razor.Parser.ParserVisitor.CancelToken">
+      <summary>This type/member supports the .NET Framework infrastructure and is not intended to be used directly from your code.Gets or sets the cancellation token.</summary>
+      <returns>The cancellation token.</returns>
+    </member>
+    <member name="M:System.Web.Razor.Parser.ParserVisitor.OnComplete">
+      <summary>This type/member supports the .NET Framework infrastructure and is not intended to be used directly from your code.Indicates that a visitor method has completed execution.</summary>
+    </member>
+    <member name="M:System.Web.Razor.Parser.ParserVisitor.ThrowIfCanceled">
+      <summary>This type/member supports the .NET Framework infrastructure and is not intended to be used directly from your code.</summary>
+    </member>
+    <member name="M:System.Web.Razor.Parser.ParserVisitor.VisitBlock(System.Web.Razor.Parser.SyntaxTree.Block)">
+      <summary>This type/member supports the .NET Framework infrastructure and is not intended to be used directly from your code.Visits the specified block.</summary>
+      <param name="block">The block to visit.</param>
+    </member>
+    <member name="M:System.Web.Razor.Parser.ParserVisitor.VisitEndBlock(System.Web.Razor.Parser.SyntaxTree.Block)">
+      <summary>This type/member supports the .NET Framework infrastructure and is not intended to be used directly from your code.Visits the specified black after parsing.</summary>
+      <param name="block">The block to visit.</param>
+    </member>
+    <member name="M:System.Web.Razor.Parser.ParserVisitor.VisitError(System.Web.Razor.Parser.SyntaxTree.RazorError)">
+      <summary>This type/member supports the .NET Framework infrastructure and is not intended to be used directly from your code.Visits the given razor error.</summary>
+      <param name="err">The error to visit.</param>
+    </member>
+    <member name="M:System.Web.Razor.Parser.ParserVisitor.VisitSpan(System.Web.Razor.Parser.SyntaxTree.Span)">
+      <summary>This type/member supports the .NET Framework infrastructure and is not intended to be used directly from your code.Visits the specified span.</summary>
+      <param name="span">The span to visit.</param>
+    </member>
+    <member name="M:System.Web.Razor.Parser.ParserVisitor.VisitStartBlock(System.Web.Razor.Parser.SyntaxTree.Block)">
+      <summary>This type/member supports the .NET Framework infrastructure and is not intended to be used directly from your code.Visits the specified block before parsing.</summary>
+      <param name="block">The block to visit.</param>
+    </member>
+    <member name="T:System.Web.Razor.Parser.ParserVisitorExtensions">
+      <summary>This type/member supports the .NET Framework infrastructure and is not intended to be used directly from your code.Provides extension methods for parser visitor.</summary>
+    </member>
+    <member name="M:System.Web.Razor.Parser.ParserVisitorExtensions.Visit(System.Web.Razor.Parser.ParserVisitor,System.Web.Razor.ParserResults)"></member>
+    <member name="T:System.Web.Razor.Parser.RazorParser">
+      <summary>This type/member supports the .NET Framework infrastructure and is not intended to be used directly from your code.Represents a Razor parser.</summary>
+    </member>
+    <member name="M:System.Web.Razor.Parser.RazorParser.#ctor(System.Web.Razor.Parser.ParserBase,System.Web.Razor.Parser.ParserBase)">
+      <summary>This type/member supports the .NET Framework infrastructure and is not intended to be used directly from your code.Initializes a new instance of the <see cref="T:System.Web.Razor.Parser.RazorParser" /> class.</summary>
+      <param name="codeParser">The code parser.</param>
+      <param name="markupParser">The markup parser.</param>
+    </member>
+    <member name="M:System.Web.Razor.Parser.RazorParser.CreateParseTask(System.IO.TextReader,System.Action{System.Web.Razor.Parser.SyntaxTree.Span},System.Action{System.Web.Razor.Parser.SyntaxTree.RazorError})">
+      <summary>This type/member supports the .NET Framework infrastructure and is not intended to be used directly from your code.Creates a task that parses a specified object.</summary>
+      <returns>The created <see cref="T:System.Threading.Tasks.Task" />.</returns>
+      <param name="input">The object to parse.</param>
+      <param name="spanCallback">The span callback.</param>
+      <param name="errorCallback">The error callback.</param>
+    </member>
+    <member name="M:System.Web.Razor.Parser.RazorParser.CreateParseTask(System.IO.TextReader,System.Action{System.Web.Razor.Parser.SyntaxTree.Span},System.Action{System.Web.Razor.Parser.SyntaxTree.RazorError},System.Threading.CancellationToken)">
+      <summary>This type/member supports the .NET Framework infrastructure and is not intended to be used directly from your code.Creates a task that parses a specified object.</summary>
+      <returns>The created <see cref="T:System.Threading.Tasks.Task" />.</returns>
+      <param name="input">The object to parse.</param>
+      <param name="spanCallback">The span callback.</param>
+      <param name="errorCallback">The error callback.</param>
+      <param name="cancelToken">The cancellation token.</param>
+    </member>
+    <member name="M:System.Web.Razor.Parser.RazorParser.CreateParseTask(System.IO.TextReader,System.Action{System.Web.Razor.Parser.SyntaxTree.Span},System.Action{System.Web.Razor.Parser.SyntaxTree.RazorError},System.Threading.SynchronizationContext)">
+      <summary>This type/member supports the .NET Framework infrastructure and is not intended to be used directly from your code.Creates a task that parses a specified object.</summary>
+      <returns>The created <see cref="T:System.Threading.Tasks.Task" />.</returns>
+      <param name="input">The object to parse.</param>
+      <param name="spanCallback">The span callback.</param>
+      <param name="errorCallback">The error callback.</param>
+      <param name="context">The context.</param>
+    </member>
+    <member name="M:System.Web.Razor.Parser.RazorParser.CreateParseTask(System.IO.TextReader,System.Action{System.Web.Razor.Parser.SyntaxTree.Span},System.Action{System.Web.Razor.Parser.SyntaxTree.RazorError},System.Threading.SynchronizationContext,System.Threading.CancellationToken)">
+      <summary>This type/member supports the .NET Framework infrastructure and is not intended to be used directly from your code.Creates a task that parses a specified object.</summary>
+      <returns>The created <see cref="T:System.Threading.Tasks.Task" />.</returns>
+      <param name="input">The object to parse.</param>
+      <param name="spanCallback">The span callback.</param>
+      <param name="errorCallback">The error callback.</param>
+      <param name="context">The context.</param>
+      <param name="cancelToken">The cancellation token.</param>
+    </member>
+    <member name="M:System.Web.Razor.Parser.RazorParser.CreateParseTask(System.IO.TextReader,System.Web.Razor.Parser.ParserVisitor)">
+      <summary>This type/member supports the .NET Framework infrastructure and is not intended to be used directly from your code.Creates a task that parses a specified object.</summary>
+      <returns>The created <see cref="T:System.Threading.Tasks.Task" />.</returns>
+      <param name="input">The object to parse.</param>
+      <param name="consumer">The consumer.</param>
+    </member>
+    <member name="P:System.Web.Razor.Parser.RazorParser.DesignTimeMode">
+      <summary>This type/member supports the .NET Framework infrastructure and is not intended to be used directly from your code.Gets the design time mode.</summary>
+      <returns>The design time mode.</returns>
+    </member>
+    <member name="M:System.Web.Razor.Parser.RazorParser.Parse(System.IO.TextReader)">
+      <summary>This type/member supports the .NET Framework infrastructure and is not intended to be used directly from your code.Parses the specified object.</summary>
+      <returns>The parser result.</returns>
+      <param name="input">The object to parse.</param>
+    </member>
+    <member name="M:System.Web.Razor.Parser.RazorParser.Parse(System.IO.TextReader,System.Web.Razor.Parser.ParserVisitor)">
+      <summary>This type/member supports the .NET Framework infrastructure and is not intended to be used directly from your code.Parses the specified object.</summary>
+      <param name="input">The object to parse.</param>
+      <param name="visitor">The visitor.</param>
+    </member>
+    <member name="M:System.Web.Razor.Parser.RazorParser.Parse(System.Web.Razor.Text.ITextDocument)">
+      <summary>This type/member supports the .NET Framework infrastructure and is not intended to be used directly from your code.Parses the specified object.</summary>
+      <returns>The parser result.</returns>
+      <param name="input">The object to parse.</param>
+    </member>
+    <member name="M:System.Web.Razor.Parser.RazorParser.Parse(System.Web.Razor.Text.LookaheadTextReader)">
+      <summary>This type/member supports the .NET Framework infrastructure and is not intended to be used directly from your code.Parses the specified object.</summary>
+      <returns>The parser result.</returns>
+      <param name="input">The object to parse.</param>
+    </member>
+    <member name="M:System.Web.Razor.Parser.RazorParser.Parse(System.Web.Razor.Text.LookaheadTextReader,System.Web.Razor.Parser.ParserVisitor)">
+      <summary>This type/member supports the .NET Framework infrastructure and is not intended to be used directly from your code.Parses the specified object.</summary>
+      <param name="input">The object to parse.</param>
+      <param name="visitor">The visitor.</param>
+    </member>
+    <member name="T:System.Web.Razor.Parser.SyntaxConstants">
+      <summary>This type/member supports the .NET Framework infrastructure and is not intended to be used directly from your code.</summary>
+    </member>
+    <member name="F:System.Web.Razor.Parser.SyntaxConstants.EndCommentSequence">
+      <summary>This type/member supports the .NET Framework infrastructure and is not intended to be used directly from your code.</summary>
+    </member>
+    <member name="F:System.Web.Razor.Parser.SyntaxConstants.StartCommentSequence">
+      <summary>This type/member supports the .NET Framework infrastructure and is not intended to be used directly from your code.</summary>
+    </member>
+    <member name="F:System.Web.Razor.Parser.SyntaxConstants.TextTagName">
+      <summary>This type/member supports the .NET Framework infrastructure and is not intended to be used directly from your code.</summary>
+    </member>
+    <member name="F:System.Web.Razor.Parser.SyntaxConstants.TransitionCharacter">
+      <summary>This type/member supports the .NET Framework infrastructure and is not intended to be used directly from your code.</summary>
+    </member>
+    <member name="F:System.Web.Razor.Parser.SyntaxConstants.TransitionString">
+      <summary>This type/member supports the .NET Framework infrastructure and is not intended to be used directly from your code.</summary>
+    </member>
+    <member name="T:System.Web.Razor.Parser.SyntaxConstants.CSharp">
+      <summary>This type/member supports the .NET Framework infrastructure and is not intended to be used directly from your code.</summary>
+    </member>
+    <member name="F:System.Web.Razor.Parser.SyntaxConstants.CSharp.ClassKeyword">
+      <summary>This type/member supports the .NET Framework infrastructure and is not intended to be used directly from your code.</summary>
+    </member>
+    <member name="F:System.Web.Razor.Parser.SyntaxConstants.CSharp.ElseIfKeyword">
+      <summary>This type/member supports the .NET Framework infrastructure and is not intended to be used directly from your code.</summary>
+    </member>
+    <member name="F:System.Web.Razor.Parser.SyntaxConstants.CSharp.FunctionsKeyword">
+      <summary>This type/member supports the .NET Framework infrastructure and is not intended to be used directly from your code.</summary>
+    </member>
+    <member name="F:System.Web.Razor.Parser.SyntaxConstants.CSharp.HelperKeyword">
+      <summary>This type/member supports the .NET Framework infrastructure and is not intended to be used directly from your code.</summary>
+    </member>
+    <member name="F:System.Web.Razor.Parser.SyntaxConstants.CSharp.InheritsKeyword">
+      <summary>This type/member supports the .NET Framework infrastructure and is not intended to be used directly from your code.</summary>
+    </member>
+    <member name="F:System.Web.Razor.Parser.SyntaxConstants.CSharp.LayoutKeyword">
+      <summary>This type/member supports the .NET Framework infrastructure and is not intended to be used directly from your code.</summary>
+    </member>
+    <member name="F:System.Web.Razor.Parser.SyntaxConstants.CSharp.NamespaceKeyword">
+      <summary>This type/member supports the .NET Framework infrastructure and is not intended to be used directly from your code.</summary>
+    </member>
+    <member name="F:System.Web.Razor.Parser.SyntaxConstants.CSharp.SectionKeyword">
+      <summary>This type/member supports the .NET Framework infrastructure and is not intended to be used directly from your code.</summary>
+    </member>
+    <member name="F:System.Web.Razor.Parser.SyntaxConstants.CSharp.SessionStateKeyword">
+      <summary>This type/member supports the .NET Framework infrastructure and is not intended to be used directly from your code.</summary>
+    </member>
+    <member name="F:System.Web.Razor.Parser.SyntaxConstants.CSharp.UsingKeywordLength">
+      <summary>This type/member supports the .NET Framework infrastructure and is not intended to be used directly from your code.</summary>
+    </member>
+    <member name="T:System.Web.Razor.Parser.SyntaxConstants.VB">
+      <summary>This type/member supports the .NET Framework infrastructure and is not intended to be used directly from your code.</summary>
+    </member>
+    <member name="F:System.Web.Razor.Parser.SyntaxConstants.VB.CodeKeyword">
+      <summary>This type/member supports the .NET Framework infrastructure and is not intended to be used directly from your code.</summary>
+    </member>
+    <member name="F:System.Web.Razor.Parser.SyntaxConstants.VB.EndCodeKeyword">
+      <summary>This type/member supports the .NET Framework infrastructure and is not intended to be used directly from your code.</summary>
+    </member>
+    <member name="F:System.Web.Razor.Parser.SyntaxConstants.VB.EndFunctionsKeyword">
+      <summary>This type/member supports the .NET Framework infrastructure and is not intended to be used directly from your code.</summary>
+    </member>
+    <member name="F:System.Web.Razor.Parser.SyntaxConstants.VB.EndHelperKeyword">
+      <summary>This type/member supports the .NET Framework infrastructure and is not intended to be used directly from your code.</summary>
+    </member>
+    <member name="F:System.Web.Razor.Parser.SyntaxConstants.VB.EndKeyword">
+      <summary>This type/member supports the .NET Framework infrastructure and is not intended to be used directly from your code.</summary>
+    </member>
+    <member name="F:System.Web.Razor.Parser.SyntaxConstants.VB.EndSectionKeyword">
+      <summary>This type/member supports the .NET Framework infrastructure and is not intended to be used directly from your code.</summary>
+    </member>
+    <member name="F:System.Web.Razor.Parser.SyntaxConstants.VB.ExplicitKeyword">
+      <summary>This type/member supports the .NET Framework infrastructure and is not intended to be used directly from your code.</summary>
+    </member>
+    <member name="F:System.Web.Razor.Parser.SyntaxConstants.VB.FunctionsKeyword">
+      <summary>This type/member supports the .NET Framework infrastructure and is not intended to be used directly from your code.</summary>
+    </member>
+    <member name="F:System.Web.Razor.Parser.SyntaxConstants.VB.HelperKeyword">
+      <summary>This type/member supports the .NET Framework infrastructure and is not intended to be used directly from your code.</summary>
+    </member>
+    <member name="F:System.Web.Razor.Parser.SyntaxConstants.VB.ImportsKeywordLength">
+      <summary>This type/member supports the .NET Framework infrastructure and is not intended to be used directly from your code.</summary>
+    </member>
+    <member name="F:System.Web.Razor.Parser.SyntaxConstants.VB.LayoutKeyword">
+      <summary>This type/member supports the .NET Framework infrastructure and is not intended to be used directly from your code.</summary>
+    </member>
+    <member name="F:System.Web.Razor.Parser.SyntaxConstants.VB.OffKeyword">
+      <summary>This type/member supports the .NET Framework infrastructure and is not intended to be used directly from your code.</summary>
+    </member>
+    <member name="F:System.Web.Razor.Parser.SyntaxConstants.VB.SectionKeyword">
+      <summary>This type/member supports the .NET Framework infrastructure and is not intended to be used directly from your code.</summary>
+    </member>
+    <member name="F:System.Web.Razor.Parser.SyntaxConstants.VB.SelectCaseKeyword">
+      <summary>This type/member supports the .NET Framework infrastructure and is not intended to be used directly from your code.</summary>
+    </member>
+    <member name="F:System.Web.Razor.Parser.SyntaxConstants.VB.SessionStateKeyword">
+      <summary>This type/member supports the .NET Framework infrastructure and is not intended to be used directly from your code.</summary>
+    </member>
+    <member name="F:System.Web.Razor.Parser.SyntaxConstants.VB.StrictKeyword">
+      <summary>This type/member supports the .NET Framework infrastructure and is not intended to be used directly from your code.</summary>
+    </member>
+    <member name="T:System.Web.Razor.Parser.TokenizerBackedParser`3">
+      <summary>This type/member supports the .NET Framework infrastructure and is not intended to be used directly from your code.Represents a tokenizer backed parser.</summary>
+      <typeparam name="TTokenizer">The type of tokenizer.</typeparam>
+      <typeparam name="TSymbol">The type of symbol.</typeparam>
+      <typeparam name="TSymbolType">The type of SymbolType.</typeparam>
+    </member>
+    <member name="M:System.Web.Razor.Parser.TokenizerBackedParser`3.#ctor">
+      <summary>This type/member supports the .NET Framework infrastructure and is not intended to be used directly from your code.Initializes a new instance of the <see cref="T:System.Web.Razor.Parser.TokenizerBackedParser`3" /> class.</summary>
+    </member>
+    <member name="M:System.Web.Razor.Parser.TokenizerBackedParser`3.Accept(System.Collections.Generic.IEnumerable{`1})">
+      <summary>This type/member supports the .NET Framework infrastructure and is not intended to be used directly from your code.Accepts the list of symbols</summary>
+      <param name="symbols">The list of symbols.</param>
+    </member>
+    <member name="M:System.Web.Razor.Parser.TokenizerBackedParser`3.Accept(`1)">
+      <summary>This type/member supports the .NET Framework infrastructure and is not intended to be used directly from your code.Accepts the specified symbol.</summary>
+      <param name="symbol">The symbol to accept.</param>
+    </member>
+    <member name="M:System.Web.Razor.Parser.TokenizerBackedParser`3.AcceptAll(`2[])">
+      <summary>This type/member supports the .NET Framework infrastructure and is not intended to be used directly from your code.Determines whether the parser accepts all types of tokenizer.</summary>
+      <returns>true of the parser accepts all types of tokenizer; otherwise, false.</returns>
+      <param name="types">The types.</param>
+    </member>
+    <member name="M:System.Web.Razor.Parser.TokenizerBackedParser`3.AcceptAndMoveNext">
+      <summary>This type/member supports the .NET Framework infrastructure and is not intended to be used directly from your code.Determines whether the parser accepts and moves to the next tokenizer.</summary>
+      <returns>true if the parser accepts and moves to the next tokenizer; otherwise, false.</returns>
+    </member>
+    <member name="M:System.Web.Razor.Parser.TokenizerBackedParser`3.AcceptSingleWhiteSpaceCharacter">
+      <summary>This type/member supports the .NET Framework infrastructure and is not intended to be used directly from your code.Determines whether the parser accepts single whitespace character.</summary>
+      <returns>true if the parser accepts single whitespace character; otherwise, false.</returns>
+    </member>
+    <member name="M:System.Web.Razor.Parser.TokenizerBackedParser`3.AcceptUntil(`2)">
+      <summary>This type/member supports the .NET Framework infrastructure and is not intended to be used directly from your code.Accepts token until a token of the given type is found.</summary>
+      <param name="type">The type of the token.</param>
+    </member>
+    <member name="M:System.Web.Razor.Parser.TokenizerBackedParser`3.AcceptUntil(`2,`2)">
+      <summary>This type/member supports the .NET Framework infrastructure and is not intended to be used directly from your code.Accepts token until a token of the given type is found and it will backup so that the next token is of the given type.</summary>
+      <param name="type1">The type of the first token.</param>
+      <param name="type2">The type of the second token.</param>
+    </member>
+    <member name="M:System.Web.Razor.Parser.TokenizerBackedParser`3.AcceptUntil(`2,`2,`2)">
+      <summary>This type/member supports the .NET Framework infrastructure and is not intended to be used directly from your code.Accepts the given tokens until a token of the given type is found.</summary>
+      <param name="type1">The type of the first token.</param>
+      <param name="type2">The type of the second token.</param>
+      <param name="type3">The type of the third token.</param>
+    </member>
+    <member name="M:System.Web.Razor.Parser.TokenizerBackedParser`3.AcceptUntil(`2[])">
+      <summary>This type/member supports the .NET Framework infrastructure and is not intended to be used directly from your code.Accepts token until a token of the given types is found.</summary>
+      <param name="types">The types of the token.</param>
+    </member>
+    <member name="M:System.Web.Razor.Parser.TokenizerBackedParser`3.AcceptWhile(System.Func{`1,System.Boolean})">
+      <summary>Accepts token while the condition has been reached.</summary>
+      <param name="condition">The condition.</param>
+    </member>
+    <member name="M:System.Web.Razor.Parser.TokenizerBackedParser`3.AcceptWhile(`2)">
+      <summary>This type/member supports the .NET Framework infrastructure and is not intended to be used directly from your code.Accepts the token while a token of the given type is not found.</summary>
+      <param name="type">The type of the token.</param>
+    </member>
+    <member name="M:System.Web.Razor.Parser.TokenizerBackedParser`3.AcceptWhile(`2,`2)">
+      <summary>This type/member supports the .NET Framework infrastructure and is not intended to be used directly from your code.Accepts token while the token of the given type has been reached.</summary>
+      <param name="type1">The type of the first token.</param>
+      <param name="type2">The type of the second token.</param>
+    </member>
+    <member name="M:System.Web.Razor.Parser.TokenizerBackedParser`3.AcceptWhile(`2,`2,`2)">
+      <summary>This type/member supports the .NET Framework infrastructure and is not intended to be used directly from your code.Accepts token while the token of the given type has been reached.</summary>
+      <param name="type1">The type of the first token.</param>
+      <param name="type2">The type of the second token.</param>
+      <param name="type3">The type of the third token.</param>
+    </member>
+    <member name="M:System.Web.Razor.Parser.TokenizerBackedParser`3.AcceptWhile(`2[])">
+      <summary>This type/member supports the .NET Framework infrastructure and is not intended to be used directly from your code.Accepts token while the token of the given types has been reached.</summary>
+      <param name="types">The types.</param>
+    </member>
+    <member name="M:System.Web.Razor.Parser.TokenizerBackedParser`3.AcceptWhiteSpaceInLines">
+      <summary>This type/member supports the .NET Framework infrastructure and is not intended to be used directly from your code.Determines whether the parser accepts whitespace in lines.</summary>
+      <returns>true if the parser accepts whitespace in lines; otherwise, false.</returns>
+    </member>
+    <member name="M:System.Web.Razor.Parser.TokenizerBackedParser`3.AddMarkerSymbolIfNecessary">
+      <summary>This type/member supports the .NET Framework infrastructure and is not intended to be used directly from your code.Adds a marker symbol if necessary.</summary>
+    </member>
+    <member name="M:System.Web.Razor.Parser.TokenizerBackedParser`3.AddMarkerSymbolIfNecessary(System.Web.Razor.Text.SourceLocation)">
+      <summary>This type/member supports the .NET Framework infrastructure and is not intended to be used directly from your code.Adds a marker symbol if necessary.</summary>
+      <param name="location">The location where to add the symbol.</param>
+    </member>
+    <member name="M:System.Web.Razor.Parser.TokenizerBackedParser`3.At(`2)">
+      <summary>This type/member supports the .NET Framework infrastructure and is not intended to be used directly from your code.Determines whether the token is at the specified type.</summary>
+      <returns>true if the token is at the specified type; otherwise, false.</returns>
+      <param name="type">The type.</param>
+    </member>
+    <member name="M:System.Web.Razor.Parser.TokenizerBackedParser`3.AtIdentifier(System.Boolean)">
+      <summary>This type/member supports the .NET Framework infrastructure and is not intended to be used directly from your code.Determines whether the token is at the specified identifier.</summary>
+      <returns>true if the token is at the specified identifier; otherwise, false.</returns>
+      <param name="allowKeywords">true to allow keywords; otherwise, false.</param>
+    </member>
+    <member name="M:System.Web.Razor.Parser.TokenizerBackedParser`3.Balance(System.Web.Razor.Parser.BalancingModes)">
+      <summary>This type/member supports the .NET Framework infrastructure and is not intended to be used directly from your code.Determines whether the parsing is balance.</summary>
+      <returns>true if the parsing is balance; otherwise, false.</returns>
+      <param name="mode">The balancing mode.</param>
+    </member>
+    <member name="M:System.Web.Razor.Parser.TokenizerBackedParser`3.Balance(System.Web.Razor.Parser.BalancingModes,`2,`2,System.Web.Razor.Text.SourceLocation)">
+      <summary>This type/member supports the .NET Framework infrastructure and is not intended to be used directly from your code.Determines whether the parsing is balance.</summary>
+      <returns>true if the parsing is balance; otherwise, false.</returns>
+      <param name="mode">The balancing mode.</param>
+      <param name="left">The left parse.</param>
+      <param name="right">The right parse.</param>
+      <param name="start">The start of the mode.</param>
+    </member>
+    <member name="M:System.Web.Razor.Parser.TokenizerBackedParser`3.BuildSpan(System.Web.Razor.Parser.SyntaxTree.SpanBuilder,System.Web.Razor.Text.SourceLocation,System.String)">
+      <summary>This type/member supports the .NET Framework infrastructure and is not intended to be used directly from your code.Builds a specified span.</summary>
+      <param name="span">The span to build.</param>
+      <param name="start">The start location to build the span.</param>
+      <param name="content">The content of the span.</param>
+    </member>
+    <member name="M:System.Web.Razor.Parser.TokenizerBackedParser`3.ConfigureSpan(System.Action{System.Web.Razor.Parser.SyntaxTree.SpanBuilder})">
+      <summary>This type/member supports the .NET Framework infrastructure and is not intended to be used directly from your code.Configures the span.</summary>
+      <param name="config">The configuration.</param>
+    </member>
+    <member name="M:System.Web.Razor.Parser.TokenizerBackedParser`3.ConfigureSpan(System.Action{System.Web.Razor.Parser.SyntaxTree.SpanBuilder,System.Action{System.Web.Razor.Parser.SyntaxTree.SpanBuilder}})">
+      <summary>Configures the span.</summary>
+      <param name="config">The configuration.</param>
+    </member>
+    <member name="P:System.Web.Razor.Parser.TokenizerBackedParser`3.CurrentLocation">
+      <summary>This type/member supports the .NET Framework infrastructure and is not intended to be used directly from your code.Gets the current location of the current instance.</summary>
+      <returns>The current location of the current instance.</returns>
+    </member>
+    <member name="P:System.Web.Razor.Parser.TokenizerBackedParser`3.CurrentSymbol">
+      <summary>This type/member supports the .NET Framework infrastructure and is not intended to be used directly from your code.Gets the current symbol of this instance.</summary>
+      <returns>The current symbol of this instance.</returns>
+    </member>
+    <member name="P:System.Web.Razor.Parser.TokenizerBackedParser`3.EndOfFile">
+      <summary>This type/member supports the .NET Framework infrastructure and is not intended to be used directly from your code.Gets a value indicating whether the tokenizer is in the end of file.</summary>
+      <returns>true if the tokenizer is in the end of file; otherwise, false.</returns>
+    </member>
+    <member name="M:System.Web.Razor.Parser.TokenizerBackedParser`3.EnsureCurrent">
+      <summary>This type/member supports the .NET Framework infrastructure and is not intended to be used directly from your code.Determines whether to ensure the current parser.</summary>
+      <returns>true if to ensure the current parser; otherwise, false.</returns>
+    </member>
+    <member name="M:System.Web.Razor.Parser.TokenizerBackedParser`3.Expected(System.Web.Razor.Tokenizer.Symbols.KnownSymbolType)">
+      <summary>This type/member supports the .NET Framework infrastructure and is not intended to be used directly from your code.Indicates the expected token with the given type.</summary>
+      <param name="type">The type.</param>
+    </member>
+    <member name="M:System.Web.Razor.Parser.TokenizerBackedParser`3.Expected(`2[])">
+      <summary>This type/member supports the .NET Framework infrastructure and is not intended to be used directly from your code.Indicates the expected token with the given types.</summary>
+      <param name="types">The types.</param>
+    </member>
+    <member name="M:System.Web.Razor.Parser.TokenizerBackedParser`3.HandleEmbeddedTransition">
+      <summary>This type/member supports the .NET Framework infrastructure and is not intended to be used directly from your code.Handles the embedded transition.</summary>
+    </member>
+    <member name="M:System.Web.Razor.Parser.TokenizerBackedParser`3.Initialize(System.Web.Razor.Parser.SyntaxTree.SpanBuilder)">
+      <summary>This type/member supports the .NET Framework infrastructure and is not intended to be used directly from your code.Initializes a specified span.</summary>
+      <param name="span">The span to initialize.</param>
+    </member>
+    <member name="M:System.Web.Razor.Parser.TokenizerBackedParser`3.IsAtEmbeddedTransition(System.Boolean,System.Boolean)">
+      <summary>This type/member supports the .NET Framework infrastructure and is not intended to be used directly from your code.Determines whether this instance is at embedded transition.</summary>
+      <returns>true if this instance is at embedded transition; otherwise, false.</returns>
+      <param name="allowTemplatesAndComments">true to allow templates and comments; otherwise, false.</param>
+      <param name="allowTransitions">true to allow transitions; otherwise, false.</param>
+    </member>
+    <member name="P:System.Web.Razor.Parser.TokenizerBackedParser`3.Language">
+      <summary>This type/member supports the .NET Framework infrastructure and is not intended to be used directly from your code.Gets the language used for parsing.</summary>
+      <returns>The language used for parsing.</returns>
+    </member>
+    <member name="M:System.Web.Razor.Parser.TokenizerBackedParser`3.NextIs(System.Func{`1,System.Boolean})">
+      <summary>Determines whether the token with the given condition would pass.</summary>
+      <returns>true if the token with the given condition would pass; otherwise, false.</returns>
+      <param name="condition">The condition.</param>
+    </member>
+    <member name="M:System.Web.Razor.Parser.TokenizerBackedParser`3.NextIs(`2)">
+      <summary>This type/member supports the .NET Framework infrastructure and is not intended to be used directly from your code.Determines whether the token with the given type would pass.</summary>
+      <returns>true if the token with the give type would pass; otherwise, false.</returns>
+      <param name="type">The type of the token.</param>
+    </member>
+    <member name="M:System.Web.Razor.Parser.TokenizerBackedParser`3.NextIs(`2[])">
+      <summary>This type/member supports the .NET Framework infrastructure and is not intended to be used directly from your code.Determines whether the token with the given types would pass.</summary>
+      <returns>true if the token with the given types would pass; otherwise, false.</returns>
+      <param name="types">The types.</param>
+    </member>
+    <member name="M:System.Web.Razor.Parser.TokenizerBackedParser`3.NextToken">
+      <summary>This type/member supports the .NET Framework infrastructure and is not intended to be used directly from your code.Determines whether the parser advances to the next token.</summary>
+      <returns>true if the parser advances to the next token; otherwise, false.</returns>
+    </member>
+    <member name="M:System.Web.Razor.Parser.TokenizerBackedParser`3.Optional(System.Web.Razor.Tokenizer.Symbols.KnownSymbolType)">
+      <summary>This type/member supports the .NET Framework infrastructure and is not intended to be used directly from your code.Determines whether parsing a token with the given type is optional.</summary>
+      <returns>true if parsing a token with the given type is optional; otherwise, false.</returns>
+      <param name="type">The type of the token.</param>
+    </member>
+    <member name="M:System.Web.Razor.Parser.TokenizerBackedParser`3.Optional(`2)">
+      <summary>This type/member supports the .NET Framework infrastructure and is not intended to be used directly from your code.Determines whether parsing a token with the given type is optional.</summary>
+      <returns>true if parsing a token with the given type is optional; otherwise, false.</returns>
+      <param name="type">The type of the token.</param>
+    </member>
+    <member name="M:System.Web.Razor.Parser.TokenizerBackedParser`3.Output(System.Web.Razor.Parser.SyntaxTree.AcceptedCharacters)">
+      <summary>This type/member supports the .NET Framework infrastructure and is not intended to be used directly from your code.Outputs a token with accepted characters.</summary>
+      <param name="accepts">The accepted characters.</param>
+    </member>
+    <member name="M:System.Web.Razor.Parser.TokenizerBackedParser`3.Output(System.Web.Razor.Parser.SyntaxTree.SpanKind)">
+      <summary>This type/member supports the .NET Framework infrastructure and is not intended to be used directly from your code.Outputs a token with span kind.</summary>
+      <param name="kind">The span kind.</param>
+    </member>
+    <member name="M:System.Web.Razor.Parser.TokenizerBackedParser`3.Output(System.Web.Razor.Parser.SyntaxTree.SpanKind,System.Web.Razor.Parser.SyntaxTree.AcceptedCharacters)">
+      <summary>This type/member supports the .NET Framework infrastructure and is not intended to be used directly from your code.Outputs a token with a given span kind and accepted characters.</summary>
+      <param name="kind">The span kind.</param>
+      <param name="accepts">The accepted characters.</param>
+    </member>
+    <member name="M:System.Web.Razor.Parser.TokenizerBackedParser`3.OutputSpanBeforeRazorComment">
+      <summary>This type/member supports the .NET Framework infrastructure and is not intended to be used directly from your code.Outputs a span before the razor comment.</summary>
+    </member>
+    <member name="P:System.Web.Razor.Parser.TokenizerBackedParser`3.PreviousSymbol">
+      <summary>This type/member supports the .NET Framework infrastructure and is not intended to be used directly from your code..Gets the previous symbol of this instance.</summary>
+    </member>
+    <member name="M:System.Web.Razor.Parser.TokenizerBackedParser`3.PushSpanConfig">
+      <summary>This type/member supports the .NET Framework infrastructure and is not intended to be used directly from your code.Pushes the span configuration.</summary>
+      <returns>An <see cref="T:System.IDisposable" /> that shuts down the configuration.</returns>
+    </member>
+    <member name="M:System.Web.Razor.Parser.TokenizerBackedParser`3.PushSpanConfig(System.Action{System.Web.Razor.Parser.SyntaxTree.SpanBuilder})">
+      <summary>This type/member supports the .NET Framework infrastructure and is not intended to be used directly from your code.Pushes the span configuration.</summary>
+      <returns>An <see cref="T:System.IDisposable" /> that shuts down the configuration.</returns>
+      <param name="newConfig">The new configuration.</param>
+    </member>
+    <member name="M:System.Web.Razor.Parser.TokenizerBackedParser`3.PushSpanConfig(System.Action{System.Web.Razor.Parser.SyntaxTree.SpanBuilder,System.Action{System.Web.Razor.Parser.SyntaxTree.SpanBuilder}})">
+      <summary>Pushes the span configuration.</summary>
+      <returns>An <see cref="T:System.IDisposable" /> that shuts down the configuration.</returns>
+      <param name="newConfig">The new configuration.</param>
+    </member>
+    <member name="M:System.Web.Razor.Parser.TokenizerBackedParser`3.PutBack(System.Collections.Generic.IEnumerable{`1})">
+      <summary>This type/member supports the .NET Framework infrastructure and is not intended to be used directly from your code.Puts the transition back.</summary>
+      <param name="symbols">The symbols.</param>
+    </member>
+    <member name="M:System.Web.Razor.Parser.TokenizerBackedParser`3.PutBack(`1)">
+      <summary>This type/member supports the .NET Framework infrastructure and is not intended to be used directly from your code.Puts the transition back.</summary>
+      <param name="symbol">The symbol.</param>
+    </member>
+    <member name="M:System.Web.Razor.Parser.TokenizerBackedParser`3.PutCurrentBack">
+      <summary>This type/member supports the .NET Framework infrastructure and is not intended to be used directly from your code.Puts the current transition back.</summary>
+    </member>
+    <member name="M:System.Web.Razor.Parser.TokenizerBackedParser`3.RazorComment">
+      <summary>This type/member supports the .NET Framework infrastructure and is not intended to be used directly from your code.Displays the razor comment.</summary>
+    </member>
+    <member name="M:System.Web.Razor.Parser.TokenizerBackedParser`3.ReadWhile(System.Func{`1,System.Boolean})">
+      <summary>Reads a token while the condition is not reached.</summary>
+      <returns>The token to read.</returns>
+      <param name="condition">The condition.</param>
+    </member>
+    <member name="M:System.Web.Razor.Parser.TokenizerBackedParser`3.Required(`2,System.Boolean,System.String)">
+      <summary>This type/member supports the .NET Framework infrastructure and is not intended to be used directly from your code.Determines whether the expected token is required.</summary>
+      <returns>true if the expected token is required; otherwise, false.</returns>
+      <param name="expected">The expected token.</param>
+      <param name="errorIfNotFound">true to display an error if not found; otherwise, false.</param>
+      <param name="errorBase">The error base.</param>
+    </member>
+    <member name="P:System.Web.Razor.Parser.TokenizerBackedParser`3.Span">
+      <summary>This type/member supports the .NET Framework infrastructure and is not intended to be used directly from your code.Gets or sets the <see cref="T:System.Web.Razor.Parser.SyntaxTree.SpanBuilder" /> associated with this instance.</summary>
+      <returns>The <see cref="T:System.Web.Razor.Parser.SyntaxTree.SpanBuilder" /> associated with this instance.</returns>
+    </member>
+    <member name="P:System.Web.Razor.Parser.TokenizerBackedParser`3.SpanConfig">
+      <summary>This type/member supports the .NET Framework infrastructure and is not intended to be used directly from your code.Gets or sets the span configuration.</summary>
+      <returns>The span configuration.</returns>
+    </member>
+    <member name="P:System.Web.Razor.Parser.TokenizerBackedParser`3.Tokenizer">
+      <summary>This type/member supports the .NET Framework infrastructure and is not intended to be used directly from your code.Gets the tokenizer.</summary>
+      <returns>The tokenizer.</returns>
+    </member>
+    <member name="M:System.Web.Razor.Parser.TokenizerBackedParser`3.Was(`2)">
+      <summary>This type/member supports the .NET Framework infrastructure and is not intended to be used directly from your code.Determines whether the token with the given type was parsed.</summary>
+      <returns>true if the token with the given type was parsed; otherwise, false.</returns>
+      <param name="type">The type of the token.</param>
+    </member>
+    <member name="T:System.Web.Razor.Parser.VBCodeParser">
+      <summary>This type/member supports the .NET Framework infrastructure and is not intended to be used directly from your code.Represents a Visual Basic code parser.</summary>
+    </member>
+    <member name="M:System.Web.Razor.Parser.VBCodeParser.#ctor">
+      <summary>This type/member supports the .NET Framework infrastructure and is not intended to be used directly from your code.Initializes a new instance of the <see cref="T:System.Web.Razor.Parser.VBCodeParser" /> class.</summary>
+    </member>
+    <member name="M:System.Web.Razor.Parser.VBCodeParser.AcceptVBSpaces">
+      <summary>This type/member supports the .NET Framework infrastructure and is not intended to be used directly from your code.Accepts spaces in the VB code.</summary>
+    </member>
+    <member name="M:System.Web.Razor.Parser.VBCodeParser.Assert(System.Web.Razor.Tokenizer.Symbols.VBKeyword)">
+      <summary>This type/member supports the .NET Framework infrastructure and is not intended to be used directly from your code.Checks for a condition and displays a keyword in the code.</summary>
+      <param name="keyword">The keyword.</param>
+    </member>
+    <member name="M:System.Web.Razor.Parser.VBCodeParser.AssertDirective(System.String)">
+      <summary>This type/member supports the .NET Framework infrastructure and is not intended to be used directly from your code.Asserts the given directive.</summary>
+      <param name="directive">The directive to assert.</param>
+    </member>
+    <member name="M:System.Web.Razor.Parser.VBCodeParser.At(System.String)">
+      <summary>This type/member supports the .NET Framework infrastructure and is not intended to be used directly from your code.Determines whether the directive is ‘AT’ directive.</summary>
+      <returns>true if the directive is an ‘AT’ directive; otherwise, false.</returns>
+      <param name="directive">The directive.</param>
+    </member>
+    <member name="M:System.Web.Razor.Parser.VBCodeParser.At(System.Web.Razor.Tokenizer.Symbols.VBKeyword)">
+      <summary>This type/member supports the .NET Framework infrastructure and is not intended to be used directly from your code.Determines whether the given keyword is ‘AT’.</summary>
+      <returns>true if the given keyword is ‘AT’; otherwise, false.</returns>
+      <param name="keyword">The keyword.</param>
+    </member>
+    <member name="M:System.Web.Razor.Parser.VBCodeParser.EndTerminatedDirective(System.String,System.Web.Razor.Parser.SyntaxTree.BlockType,System.Web.Razor.Generator.SpanCodeGenerator,System.Boolean)">
+      <summary>This type/member supports the .NET Framework infrastructure and is not intended to be used directly from your code.Ends a terminated directive.</summary>
+      <returns>The function that ends the terminated directive.</returns>
+      <param name="directive">The directive.</param>
+      <param name="blockType">The block type.</param>
+      <param name="codeGenerator">The code generator.</param>
+      <param name="allowMarkup">true to allow markup; otherwise, false.</param>
+    </member>
+    <member name="M:System.Web.Razor.Parser.VBCodeParser.EndTerminatedDirectiveBody(System.String,System.Web.Razor.Text.SourceLocation,System.Boolean)">
+      <summary>This type/member supports the .NET Framework infrastructure and is not intended to be used directly from your code.Determines whether the termination of directive body is ended.</summary>
+      <returns>true if the termination of directive body is ended; otherwise, false.</returns>
+      <param name="directive">The directive.</param>
+      <param name="blockStart">The block start.</param>
+      <param name="allowAllTransitions">true to allow all transitions; otherwise, false.</param>
+    </member>
+    <member name="M:System.Web.Razor.Parser.VBCodeParser.EndTerminatedStatement(System.Web.Razor.Tokenizer.Symbols.VBKeyword,System.Boolean,System.Boolean)">
+      <summary>Ends a termination of statement.</summary>
+      <returns>The function that ends the termination.</returns>
+      <param name="keyword">The keyword.</param>
+      <param name="supportsExit">true if the termination supports exit; otherwise, false.</param>
+      <param name="supportsContinue">true if the termination supports continue; otherwise, false.</param>
+    </member>
+    <member name="M:System.Web.Razor.Parser.VBCodeParser.EndTerminatedStatement(System.Web.Razor.Tokenizer.Symbols.VBKeyword,System.Boolean,System.Boolean,System.String)">
+      <summary>Ends a termination of statement.</summary>
+      <returns>The function that ends the termination.</returns>
+      <param name="keyword">The keyword.</param>
+      <param name="supportsExit">true if the termination supports exit; otherwise, false.</param>
+      <param name="supportsContinue">true if the termination supports continue; otherwise, false.</param>
+      <param name="blockName">The block name.</param>
+    </member>
+    <member name="M:System.Web.Razor.Parser.VBCodeParser.HandleEmbeddedTransition">
+      <summary>This type/member supports the .NET Framework infrastructure and is not intended to be used directly from your code.Handles the embedded transition.</summary>
+    </member>
+    <member name="M:System.Web.Razor.Parser.VBCodeParser.HandleEmbeddedTransition(System.Web.Razor.Tokenizer.Symbols.VBSymbol)">
+      <summary>This type/member supports the .NET Framework infrastructure and is not intended to be used directly from your code.Handles the embedded transition.</summary>
+      <param name="lastWhiteSpace">The last white space.</param>
+    </member>
+    <member name="M:System.Web.Razor.Parser.VBCodeParser.HandleExitOrContinue(System.Web.Razor.Tokenizer.Symbols.VBKeyword)">
+      <summary>This type/member supports the .NET Framework infrastructure and is not intended to be used directly from your code.Indicates the code that handles the Exit or Continue keyword.</summary>
+      <param name="keyword">The keyword.</param>
+    </member>
+    <member name="M:System.Web.Razor.Parser.VBCodeParser.HandleTransition(System.Web.Razor.Tokenizer.Symbols.VBSymbol)">
+      <summary>This type/member supports the .NET Framework infrastructure and is not intended to be used directly from your code.Indicates a code that handles a transition.</summary>
+      <param name="lastWhiteSpace">The last white space.</param>
+    </member>
+    <member name="M:System.Web.Razor.Parser.VBCodeParser.HelperDirective">
+      <summary>This type/member supports the .NET Framework infrastructure and is not intended to be used directly from your code.Indicates whether the code is a helper directive.</summary>
+      <returns>true if the code is a helper directive; otherwise, false.</returns>
+    </member>
+    <member name="M:System.Web.Razor.Parser.VBCodeParser.ImportsStatement">
+      <summary>This type/member supports the .NET Framework infrastructure and is not intended to be used directly from your code.Determines whether the code imports a statement.</summary>
+      <returns>true if the code imports a statement; otherwise, false.</returns>
+    </member>
+    <member name="M:System.Web.Razor.Parser.VBCodeParser.InheritsStatement">
+      <summary>This type/member supports the .NET Framework infrastructure and is not intended to be used directly from your code.Determines whether the code inherits a statement.</summary>
+      <returns>true if the code inherits a statement; otherwise, false.</returns>
+    </member>
+    <member name="M:System.Web.Razor.Parser.VBCodeParser.IsAtEmbeddedTransition(System.Boolean,System.Boolean)">
+      <summary>This type/member supports the .NET Framework infrastructure and is not intended to be used directly from your code.Determines whether the code is at embedded transition.</summary>
+      <returns>true if the code is at embedded transition; otherwise, false.</returns>
+      <param name="allowTemplatesAndComments">true to allow templates and comments; otherwise, false.</param>
+      <param name="allowTransitions">true to allow transitions; otherwise, false.</param>
+    </member>
+    <member name="M:System.Web.Razor.Parser.VBCodeParser.IsDirectiveDefined(System.String)">
+      <summary>This type/member supports the .NET Framework infrastructure and is not intended to be used directly from your code.Determines whether the code is directive defined.</summary>
+      <returns>true if the code is directive defined; otherwise, false.</returns>
+      <param name="directive">The directive.</param>
+    </member>
+    <member name="P:System.Web.Razor.Parser.VBCodeParser.Keywords">
+      <summary>This type/member supports the .NET Framework infrastructure and is not intended to be used directly from your code.Gets the keywords associated with the code.</summary>
+      <returns>The keywords associated with the code.</returns>
+    </member>
+    <member name="M:System.Web.Razor.Parser.VBCodeParser.KeywordTerminatedStatement(System.Web.Razor.Tokenizer.Symbols.VBKeyword,System.Web.Razor.Tokenizer.Symbols.VBKeyword,System.Boolean,System.Boolean)">
+      <summary>Indicates a keyword that terminates a statement.</summary>
+      <returns>The function that terminates the statement.</returns>
+      <param name="start">The start.</param>
+      <param name="terminator">The terminator.</param>
+      <param name="supportsExit">true if the termination supports exit; otherwise, false.</param>
+      <param name="supportsContinue">true if the termination supports continue; otherwise, false.</param>
+    </member>
+    <member name="P:System.Web.Razor.Parser.VBCodeParser.Language">
+      <summary>This type/member supports the .NET Framework infrastructure and is not intended to be used directly from your code.Gets the language for the parser.</summary>
+      <returns>The language for the parser.</returns>
+    </member>
+    <member name="M:System.Web.Razor.Parser.VBCodeParser.LayoutDirective">
+      <summary>This type/member supports the .NET Framework infrastructure and is not intended to be used directly from your code.Determines whether the code is a layout directive.</summary>
+      <returns>true if the code is a layout directive; otherwise, false.</returns>
+    </member>
+    <member name="M:System.Web.Razor.Parser.VBCodeParser.MapDirective(System.String,System.Func{System.Boolean})">
+      <summary>Maps a given directive.</summary>
+      <param name="directive">The directive.</param>
+      <param name="action">The action whether to map a given directive.</param>
+    </member>
+    <member name="M:System.Web.Razor.Parser.VBCodeParser.MapKeyword(System.Web.Razor.Tokenizer.Symbols.VBKeyword,System.Func{System.Boolean})">
+      <summary>Maps a given keyword.</summary>
+      <param name="keyword">The keyword.</param>
+      <param name="action">The action whether to map a given keyword.</param>
+    </member>
+    <member name="M:System.Web.Razor.Parser.VBCodeParser.NestedBlock">
+      <summary>This type/member supports the .NET Framework infrastructure and is not intended to be used directly from your code.Indicates a nested block.</summary>
+    </member>
+    <member name="M:System.Web.Razor.Parser.VBCodeParser.Optional(System.Web.Razor.Tokenizer.Symbols.VBKeyword)">
+      <summary>This type/member supports the .NET Framework infrastructure and is not intended to be used directly from your code.Determines whether the keyword from the code is optional.</summary>
+      <returns>true if the keyword from the code is optional; otherwise, false.</returns>
+      <param name="keyword">The keyword.</param>
+    </member>
+    <member name="M:System.Web.Razor.Parser.VBCodeParser.OptionStatement">
+      <summary>This type/member supports the .NET Framework infrastructure and is not intended to be used directly from your code.Determines whether the code is an option statement.</summary>
+      <returns>true if the code is an option statement; otherwise, false.</returns>
+    </member>
+    <member name="P:System.Web.Razor.Parser.VBCodeParser.OtherParser">
+      <summary>This type/member supports the .NET Framework infrastructure and is not intended to be used directly from your code.Gets the other parser.</summary>
+      <returns>The other parser.</returns>
+    </member>
+    <member name="M:System.Web.Razor.Parser.VBCodeParser.OtherParserBlock">
+      <summary>This type/member supports the .NET Framework infrastructure and is not intended to be used directly from your code.Indicates the parser block.</summary>
+    </member>
+    <member name="M:System.Web.Razor.Parser.VBCodeParser.OtherParserBlock(System.String,System.String)">
+      <summary>This type/member supports the .NET Framework infrastructure and is not intended to be used directly from your code.Indicates the parser block.</summary>
+      <param name="startSequence">The start sequence.</param>
+      <param name="endSequence">The end sequence.</param>
+    </member>
+    <member name="M:System.Web.Razor.Parser.VBCodeParser.OutputSpanBeforeRazorComment">
+      <summary>This type/member supports the .NET Framework infrastructure and is not intended to be used directly from your code.Spans the output before Razor comment.</summary>
+    </member>
+    <member name="M:System.Web.Razor.Parser.VBCodeParser.ParseBlock">
+      <summary>This type/member supports the .NET Framework infrastructure and is not intended to be used directly from your code.Blocks the parsing.</summary>
+    </member>
+    <member name="M:System.Web.Razor.Parser.VBCodeParser.ReadVBSpaces">
+      <summary>This type/member supports the .NET Framework infrastructure and is not intended to be used directly from your code.Reads a list of Visual Basic spaces.</summary>
+      <returns>A list of Visual Basic spaces.</returns>
+    </member>
+    <member name="M:System.Web.Razor.Parser.VBCodeParser.Required(System.Web.Razor.Tokenizer.Symbols.VBSymbolType,System.String)">
+      <summary>This type/member supports the .NET Framework infrastructure and is not intended to be used directly from your code.Determines whether the expected symbol is required.</summary>
+      <returns>true if the expected symbol is required; otherwise, false.</returns>
+      <param name="expected">The expected symbol.</param>
+      <param name="errorBase">The error base.</param>
+    </member>
+    <member name="M:System.Web.Razor.Parser.VBCodeParser.ReservedWord">
+      <summary>This type/member supports the .NET Framework infrastructure and is not intended to be used directly from your code.Determines whether the code is a reserved word.</summary>
+      <returns>true if the code is a reserved word; otherwise, false.</returns>
+    </member>
+    <member name="M:System.Web.Razor.Parser.VBCodeParser.SectionDirective">
+      <summary>This type/member supports the .NET Framework infrastructure and is not intended to be used directly from your code.Determines whether the code is a section directive.</summary>
+      <returns>true if the code is a section directive; otherwise, false.</returns>
+    </member>
+    <member name="M:System.Web.Razor.Parser.VBCodeParser.SessionStateDirective">
+      <summary>This type/member supports the .NET Framework infrastructure and is not intended to be used directly from your code.Determines whether the code has a session state directive.</summary>
+      <returns>true if the code has a session state directive; otherwise, false.</returns>
+    </member>
+    <member name="T:System.Web.Razor.Parser.VBLanguageCharacteristics">
+      <summary>This type/member supports the .NET Framework infrastructure and is not intended to be used directly from your code.Represents the characteristics of the Visual Basic language.</summary>
+    </member>
+    <member name="M:System.Web.Razor.Parser.VBLanguageCharacteristics.CreateMarkerSymbol(System.Web.Razor.Text.SourceLocation)">
+      <summary>This type/member supports the .NET Framework infrastructure and is not intended to be used directly from your code.Creates a Visual Basic marker symbol.</summary>
+      <returns>The created Visual Basic marker symbol.</returns>
+      <param name="location">The location to create the symbol.</param>
+    </member>
+    <member name="M:System.Web.Razor.Parser.VBLanguageCharacteristics.CreateSymbol(System.Web.Razor.Text.SourceLocation,System.String,System.Web.Razor.Tokenizer.Symbols.VBSymbolType,System.Collections.Generic.IEnumerable{System.Web.Razor.Parser.SyntaxTree.RazorError})">
+      <summary>This type/member supports the .NET Framework infrastructure and is not intended to be used directly from your code.Creates a Visual Basic symbol.</summary>
+      <returns>The created <see cref="T:System.Web.Razor.Tokenizer.Symbols.VBSymbol" />.</returns>
+      <param name="location">The location to create the symbol.</param>
+      <param name="content">The content.</param>
+      <param name="type">The type of the symbol.</param>
+      <param name="errors">The errors.</param>
+    </member>
+    <member name="M:System.Web.Razor.Parser.VBLanguageCharacteristics.CreateTokenizer(System.Web.Razor.Text.ITextDocument)">
+      <summary>This type/member supports the .NET Framework infrastructure and is not intended to be used directly from your code.Creates a Visual Basic tokenizer.</summary>
+      <returns>The created <see cref="T:System.Web.Razor.Tokenizer.VBTokenizer" />.</returns>
+      <param name="source">The source where to create the tokenizer.</param>
+    </member>
+    <member name="M:System.Web.Razor.Parser.VBLanguageCharacteristics.FlipBracket(System.Web.Razor.Tokenizer.Symbols.VBSymbolType)">
+      <summary>This type/member supports the .NET Framework infrastructure and is not intended to be used directly from your code.Flips the given bracket.</summary>
+      <returns>The type of the Visual Basic symbol.</returns>
+      <param name="bracket">The bracket to flip.</param>
+    </member>
+    <member name="M:System.Web.Razor.Parser.VBLanguageCharacteristics.GetKnownSymbolType(System.Web.Razor.Tokenizer.Symbols.KnownSymbolType)">
+      <summary>This type/member supports the .NET Framework infrastructure and is not intended to be used directly from your code.Retrieves the type of the known symbol.</summary>
+      <returns>The type of the known symbol.</returns>
+      <param name="type">The type to retrieve.</param>
+    </member>
+    <member name="M:System.Web.Razor.Parser.VBLanguageCharacteristics.GetSample(System.Web.Razor.Tokenizer.Symbols.VBSymbolType)">
+      <summary>This type/member supports the .NET Framework infrastructure and is not intended to be used directly from your code.Gets a sample symbol with the given type.</summary>
+      <returns>A sample symbol with the given type.</returns>
+      <param name="type">The type of the symbol.</param>
+    </member>
+    <member name="P:System.Web.Razor.Parser.VBLanguageCharacteristics.Instance">
+      <summary>This type/member supports the .NET Framework infrastructure and is not intended to be used directly from your code.Gets an instance of this <see cref="T:System.Web.Razor.Parser.VBLanguageCharacteristics" />.</summary>
+      <returns>An instance of <see cref="T:System.Web.Razor.Parser.VBLanguageCharacteristics" />.</returns>
+    </member>
+    <member name="T:System.Web.Razor.Parser.SyntaxTree.AcceptedCharacters">
+      <summary>This type/member supports the .NET Framework infrastructure and is not intended to be used directly from your code.</summary>
+    </member>
+    <member name="F:System.Web.Razor.Parser.SyntaxTree.AcceptedCharacters.AllWhiteSpace">
+      <summary>This type/member supports the .NET Framework infrastructure and is not intended to be used directly from your code.</summary>
+    </member>
+    <member name="F:System.Web.Razor.Parser.SyntaxTree.AcceptedCharacters.Any">
+      <summary>This type/member supports the .NET Framework infrastructure and is not intended to be used directly from your code.</summary>
+    </member>
+    <member name="F:System.Web.Razor.Parser.SyntaxTree.AcceptedCharacters.AnyExceptNewline">
+      <summary>This type/member supports the .NET Framework infrastructure and is not intended to be used directly from your code.</summary>
+    </member>
+    <member name="F:System.Web.Razor.Parser.SyntaxTree.AcceptedCharacters.NewLine">
+      <summary>This type/member supports the .NET Framework infrastructure and is not intended to be used directly from your code.</summary>
+    </member>
+    <member name="F:System.Web.Razor.Parser.SyntaxTree.AcceptedCharacters.None">
+      <summary>This type/member supports the .NET Framework infrastructure and is not intended to be used directly from your code.</summary>
+    </member>
+    <member name="F:System.Web.Razor.Parser.SyntaxTree.AcceptedCharacters.NonWhiteSpace">
+      <summary>This type/member supports the .NET Framework infrastructure and is not intended to be used directly from your code.</summary>
+    </member>
+    <member name="F:System.Web.Razor.Parser.SyntaxTree.AcceptedCharacters.WhiteSpace">
+      <summary>This type/member supports the .NET Framework infrastructure and is not intended to be used directly from your code.</summary>
+    </member>
+    <member name="T:System.Web.Razor.Parser.SyntaxTree.AutoCompleteEditHandler">
+      <summary>This type/member supports the .NET Framework infrastructure and is not intended to be used directly from your code.Represents the auto-complete editing handler class.</summary>
+    </member>
+    <member name="M:System.Web.Razor.Parser.SyntaxTree.AutoCompleteEditHandler.#ctor(System.Func{System.String,System.Collections.Generic.IEnumerable{System.Web.Razor.Tokenizer.Symbols.ISymbol}})">
+      <summary>Initializes a new instance of the <see cref="T:System.Web.Razor.Parser.SyntaxTree.AutoCompleteEditHandler" /> class.</summary>
+      <param name="tokenizer">The tokenizer.</param>
+    </member>
+    <member name="M:System.Web.Razor.Parser.SyntaxTree.AutoCompleteEditHandler.#ctor(System.Func{System.String,System.Collections.Generic.IEnumerable{System.Web.Razor.Tokenizer.Symbols.ISymbol}},System.Web.Razor.Parser.SyntaxTree.AcceptedCharacters)">
+      <summary>Initializes a new instance of the <see cref="T:System.Web.Razor.Parser.SyntaxTree.AutoCompleteEditHandler" /> class.</summary>
+      <param name="tokenizer">The tokenizer.</param>
+      <param name="accepted">The accepted characters.</param>
+    </member>
+    <member name="P:System.Web.Razor.Parser.SyntaxTree.AutoCompleteEditHandler.AutoCompleteAtEndOfSpan">
+      <summary>This type/member supports the .NET Framework infrastructure and is not intended to be used directly from your code.Gets or sets a value that indicates whether the auto-complete function is at the end of this span.</summary>
+      <returns>true if the auto-complete function is at the end of this span; otherwise, false.</returns>
+    </member>
+    <member name="P:System.Web.Razor.Parser.SyntaxTree.AutoCompleteEditHandler.AutoCompleteString">
+      <summary>This type/member supports the .NET Framework infrastructure and is not intended to be used directly from your code.Gets or sets a string value to auto-complete.</summary>
+      <returns>A string value to auto-complete.</returns>
+    </member>
+    <member name="M:System.Web.Razor.Parser.SyntaxTree.AutoCompleteEditHandler.CanAcceptChange(System.Web.Razor.Parser.SyntaxTree.Span,System.Web.Razor.Text.TextChange)">
+      <summary>This type/member supports the .NET Framework infrastructure and is not intended to be used directly from your code.Indicates a parse result that can accept changes.</summary>
+      <param name="target">The phase of the target.</param>
+      <param name="normalizedChange">The normalized <see cref="T:System.Web.Razor.Text.TextChange" />.</param>
+    </member>
+    <member name="M:System.Web.Razor.Parser.SyntaxTree.AutoCompleteEditHandler.Equals(System.Object)">
+      <summary>This type/member supports the .NET Framework infrastructure and is not intended to be used directly from your code.Indicates whether this instance and a specified object are equal.</summary>
+      <returns>true if <paramref name="obj" /> and this instance are the same type and represent the same value; otherwise, false.</returns>
+      <param name="obj">The object.</param>
+    </member>
+    <member name="M:System.Web.Razor.Parser.SyntaxTree.AutoCompleteEditHandler.GetHashCode">
+      <summary>This type/member supports the .NET Framework infrastructure and is not intended to be used directly from your code.Returns the hash code for this instance.</summary>
+      <returns>A 32-bit signed integer that is the hash code for this instance.</returns>
+    </member>
+    <member name="M:System.Web.Razor.Parser.SyntaxTree.AutoCompleteEditHandler.ToString">
+      <summary>This type/member supports the .NET Framework infrastructure and is not intended to be used directly from your code.Returns the fully qualified type name of this instance.</summary>
+      <returns>A String containing a fully qualified type name.</returns>
+    </member>
+    <member name="T:System.Web.Razor.Parser.SyntaxTree.Block">
+      <summary>This type/member supports the .NET Framework infrastructure and is not intended to be used directly from your code.Represents the block for creating webpages.</summary>
+    </member>
+    <member name="M:System.Web.Razor.Parser.SyntaxTree.Block.#ctor(System.Web.Razor.Parser.SyntaxTree.BlockBuilder)">
+      <summary>This type/member supports the .NET Framework infrastructure and is not intended to be used directly from your code.Initializes a new instance of the <see cref="T:System.Web.Razor.Parser.SyntaxTree.Block" /> class.</summary>
+      <param name="source">The source for the block builder.</param>
+    </member>
+    <member name="M:System.Web.Razor.Parser.SyntaxTree.Block.Accept(System.Web.Razor.Parser.ParserVisitor)">
+      <summary>This type/member supports the .NET Framework infrastructure and is not intended to be used directly from your code.Accepts the parser visitor of the block.</summary>
+      <param name="visitor">The parser visitor.</param>
+    </member>
+    <member name="P:System.Web.Razor.Parser.SyntaxTree.Block.Children">
+      <summary>This type/member supports the .NET Framework infrastructure and is not intended to be used directly from your code.Gets a collection of SyntaxTreeNode to view the children of the block.</summary>
+      <returns>A collection of SyntaxTreeNode to view the children of the block.</returns>
+    </member>
+    <member name="P:System.Web.Razor.Parser.SyntaxTree.Block.CodeGenerator">
+      <summary>This type/member supports the .NET Framework infrastructure and is not intended to be used directly from your code.Gets the IBlockCodeGenerator to generate codes for the elements.</summary>
+      <returns>The IBlockCodeGenerator to generate codes for the elements.</returns>
+    </member>
+    <member name="M:System.Web.Razor.Parser.SyntaxTree.Block.Equals(System.Object)">
+      <summary>This type/member supports the .NET Framework infrastructure and is not intended to be used directly from your code.Determines whether the specified object is equal to the current block.</summary>
+      <returns>true if the specified object is equal to the current block; otherwise, false.</returns>
+      <param name="obj">The object to compare with the current object.</param>
+    </member>
+    <member name="M:System.Web.Razor.Parser.SyntaxTree.Block.EquivalentTo(System.Web.Razor.Parser.SyntaxTree.SyntaxTreeNode)">
+      <summary>This type/member supports the .NET Framework infrastructure and is not intended to be used directly from your code.Returns a value indicating whether the block is equivalent to the same element.</summary>
+      <returns>true if the block is equivalent to the same element; otherwise, false.</returns>
+      <param name="node">The syntax tree node.</param>
+    </member>
+    <member name="M:System.Web.Razor.Parser.SyntaxTree.Block.FindFirstDescendentSpan">
+      <summary>This type/member supports the .NET Framework infrastructure and is not intended to be used directly from your code.Finds the first descendent span of the block.</summary>
+      <returns>The first descendent span of the block.</returns>
+    </member>
+    <member name="M:System.Web.Razor.Parser.SyntaxTree.Block.FindLastDescendentSpan">
+      <summary>This type/member supports the .NET Framework infrastructure and is not intended to be used directly from your code.Finds the last descendent span of the block.</summary>
+      <returns>The last descendent span of the block.</returns>
+    </member>
+    <member name="M:System.Web.Razor.Parser.SyntaxTree.Block.Flatten">
+      <summary>This type/member supports the .NET Framework infrastructure and is not intended to be used directly from your code.Flattens a collection of a specified type for the block.</summary>
+      <returns>A collection of a specified type for the block to flatten.</returns>
+    </member>
+    <member name="M:System.Web.Razor.Parser.SyntaxTree.Block.GetHashCode">
+      <summary>This type/member supports the .NET Framework infrastructure and is not intended to be used directly from your code.Returns the hash code for this instance.</summary>
+      <returns>The hash code for this instance.</returns>
+    </member>
+    <member name="P:System.Web.Razor.Parser.SyntaxTree.Block.IsBlock">
+      <summary>This type/member supports the .NET Framework infrastructure and is not intended to be used directly from your code.Gets a value that indicates whether the object is a block-level object.</summary>
+      <returns>true if the object is a block-level object; otherwise, false.</returns>
+    </member>
+    <member name="P:System.Web.Razor.Parser.SyntaxTree.Block.Length">
+      <summary>This type/member supports the .NET Framework infrastructure and is not intended to be used directly from your code.Gets the length value of the block.</summary>
+      <returns>The length value of the block.</returns>
+    </member>
+    <member name="M:System.Web.Razor.Parser.SyntaxTree.Block.LocateOwner(System.Web.Razor.Text.TextChange)">
+      <summary>This type/member supports the .NET Framework infrastructure and is not intended to be used directly from your code.Locates the owner of the block.</summary>
+      <returns>The owner of the block to locate.</returns>
+      <param name="change">The text change.</param>
+    </member>
+    <member name="P:System.Web.Razor.Parser.SyntaxTree.Block.Name">
+      <summary>This type/member supports the .NET Framework infrastructure and is not intended to be used directly from your code.Gets the string name of the block.</summary>
+      <returns>The string name of the block.</returns>
+    </member>
+    <member name="P:System.Web.Razor.Parser.SyntaxTree.Block.Start">
+      <summary>This type/member supports the .NET Framework infrastructure and is not intended to be used directly from your code.Gets the start to identify the specific location of the block.</summary>
+      <returns>The start to identify the specific location of the block.</returns>
+    </member>
+    <member name="M:System.Web.Razor.Parser.SyntaxTree.Block.ToString">
+      <summary>This type/member supports the .NET Framework infrastructure and is not intended to be used directly from your code.Returns a string that represents the current object.</summary>
+      <returns>A string that represents the current object.</returns>
+    </member>
+    <member name="P:System.Web.Razor.Parser.SyntaxTree.Block.Type">
+      <summary>This type/member supports the .NET Framework infrastructure and is not intended to be used directly from your code.Gets the type of code block.</summary>
+      <returns>The type of code block.</returns>
+    </member>
+    <member name="T:System.Web.Razor.Parser.SyntaxTree.BlockBuilder">
+      <summary>This type/member supports the .NET Framework infrastructure and is not intended to be used directly from your code.Represents the block builder for the webpages.</summary>
+    </member>
+    <member name="M:System.Web.Razor.Parser.SyntaxTree.BlockBuilder.#ctor">
+      <summary>This type/member supports the .NET Framework infrastructure and is not intended to be used directly from your code.Initializes a new instance of the <see cref="T:System.Web.Razor.Parser.SyntaxTree.BlockBuilder" /> class.</summary>
+    </member>
+    <member name="M:System.Web.Razor.Parser.SyntaxTree.BlockBuilder.#ctor(System.Web.Razor.Parser.SyntaxTree.Block)">
+      <summary>This type/member supports the .NET Framework infrastructure and is not intended to be used directly from your code.Initializes a new instance of the <see cref="T:System.Web.Razor.Parser.SyntaxTree.BlockBuilder" /> class.</summary>
+      <param name="original">The original block builder.</param>
+    </member>
+    <member name="M:System.Web.Razor.Parser.SyntaxTree.BlockBuilder.Build">
+      <summary>This type/member supports the .NET Framework infrastructure and is not intended to be used directly from your code.Builds a block for this instance.</summary>
+      <returns>A block builds for this instance.</returns>
+    </member>
+    <member name="P:System.Web.Razor.Parser.SyntaxTree.BlockBuilder.Children">
+      <summary>This type/member supports the .NET Framework infrastructure and is not intended to be used directly from your code.Gets the collection of child elements of the block builder.</summary>
+      <returns>The collection of child elements of the block builder.</returns>
+    </member>
+    <member name="P:System.Web.Razor.Parser.SyntaxTree.BlockBuilder.CodeGenerator">
+      <summary>This type/member supports the .NET Framework infrastructure and is not intended to be used directly from your code.Gets or sets the code generator for the block builder.</summary>
+      <returns>The code generator for the block builder.</returns>
+    </member>
+    <member name="P:System.Web.Razor.Parser.SyntaxTree.BlockBuilder.Name">
+      <summary>This type/member supports the .NET Framework infrastructure and is not intended to be used directly from your code.Gets or sets the string name for the block builder.</summary>
+      <returns>The string name for the block builder.</returns>
+    </member>
+    <member name="M:System.Web.Razor.Parser.SyntaxTree.BlockBuilder.Reset">
+      <summary>This type/member supports the .NET Framework infrastructure and is not intended to be used directly from your code.Resets the block builder to its original position.</summary>
+    </member>
+    <member name="P:System.Web.Razor.Parser.SyntaxTree.BlockBuilder.Type">
+      <summary>This type/member supports the .NET Framework infrastructure and is not intended to be used directly from your code.Gets or sets a block type that can be assigned null.</summary>
+      <returns>A block type that can be assigned null.</returns>
+    </member>
+    <member name="T:System.Web.Razor.Parser.SyntaxTree.BlockType">
+      <summary>This type/member supports the .NET Framework infrastructure and is not intended to be used directly from your code.</summary>
+    </member>
+    <member name="F:System.Web.Razor.Parser.SyntaxTree.BlockType.Comment">
+      <summary>This type/member supports the .NET Framework infrastructure and is not intended to be used directly from your code.</summary>
+    </member>
+    <member name="F:System.Web.Razor.Parser.SyntaxTree.BlockType.Directive">
+      <summary>This type/member supports the .NET Framework infrastructure and is not intended to be used directly from your code.</summary>
+    </member>
+    <member name="F:System.Web.Razor.Parser.SyntaxTree.BlockType.Expression">
+      <summary>This type/member supports the .NET Framework infrastructure and is not intended to be used directly from your code.</summary>
+    </member>
+    <member name="F:System.Web.Razor.Parser.SyntaxTree.BlockType.Functions">
+      <summary>This type/member supports the .NET Framework infrastructure and is not intended to be used directly from your code.</summary>
+    </member>
+    <member name="F:System.Web.Razor.Parser.SyntaxTree.BlockType.Helper">
+      <summary>This type/member supports the .NET Framework infrastructure and is not intended to be used directly from your code.</summary>
+    </member>
+    <member name="F:System.Web.Razor.Parser.SyntaxTree.BlockType.Markup">
+      <summary>This type/member supports the .NET Framework infrastructure and is not intended to be used directly from your code.</summary>
+    </member>
+    <member name="F:System.Web.Razor.Parser.SyntaxTree.BlockType.Section">
+      <summary>This type/member supports the .NET Framework infrastructure and is not intended to be used directly from your code.</summary>
+    </member>
+    <member name="F:System.Web.Razor.Parser.SyntaxTree.BlockType.Statement">
+      <summary>This type/member supports the .NET Framework infrastructure and is not intended to be used directly from your code.</summary>
+    </member>
+    <member name="F:System.Web.Razor.Parser.SyntaxTree.BlockType.Template">
+      <summary>This type/member supports the .NET Framework infrastructure and is not intended to be used directly from your code.</summary>
+    </member>
+    <member name="T:System.Web.Razor.Parser.SyntaxTree.RazorError">
+      <summary>This type/member supports the .NET Framework infrastructure and is not intended to be used directly from your code. Represents a parsing error in Razor.</summary>
+    </member>
+    <member name="M:System.Web.Razor.Parser.SyntaxTree.RazorError.#ctor(System.String,System.Int32,System.Int32,System.Int32)">
+      <summary>This type/member supports the .NET Framework infrastructure and is not intended to be used directly from your code. Initializes a new instance of the <see cref="T:System.Web.Razor.Parser.SyntaxTree.RazorError" /> class.</summary>
+      <param name="message">The error message.</param>
+      <param name="absoluteIndex">The absolute index of the source location.</param>
+      <param name="lineIndex">The line index of the source location.</param>
+      <param name="columnIndex">The column index of the source location.</param>
+    </member>
+    <member name="M:System.Web.Razor.Parser.SyntaxTree.RazorError.#ctor(System.String,System.Int32,System.Int32,System.Int32,System.Int32)">
+      <summary>This type/member supports the .NET Framework infrastructure and is not intended to be used directly from your code. Initializes a new instance of the <see cref="T:System.Web.Razor.Parser.SyntaxTree.RazorError" /> class.</summary>
+      <param name="message">The error message.</param>
+      <param name="absoluteIndex">The absolute index of the source location.</param>
+      <param name="lineIndex">The line index of the source location.</param>
+      <param name="columnIndex">The column index of the source location.</param>
+      <param name="length">The length for the error.</param>
+    </member>
+    <member name="M:System.Web.Razor.Parser.SyntaxTree.RazorError.#ctor(System.String,System.Web.Razor.Text.SourceLocation)">
+      <summary>This type/member supports the .NET Framework infrastructure and is not intended to be used directly from your code. Initializes a new instance of the <see cref="T:System.Web.Razor.Parser.SyntaxTree.RazorError" /> class.</summary>
+      <param name="message">The error message.</param>
+      <param name="location">The source location of the error.</param>
+    </member>
+    <member name="M:System.Web.Razor.Parser.SyntaxTree.RazorError.#ctor(System.String,System.Web.Razor.Text.SourceLocation,System.Int32)">
+      <summary>This type/member supports the .NET Framework infrastructure and is not intended to be used directly from your code. Initializes a new instance of the <see cref="T:System.Web.Razor.Parser.SyntaxTree.RazorError" /> class.</summary>
+      <param name="message">The error message.</param>
+      <param name="location">The source location of the error.</param>
+      <param name="length">The length for the error.</param>
+    </member>
+    <member name="M:System.Web.Razor.Parser.SyntaxTree.RazorError.Equals(System.Object)">
+      <summary>This type/member supports the .NET Framework infrastructure and is not intended to be used directly from your code. Determines whether the specified object is equal to this instance.</summary>
+      <returns>true if the specified object is equal to this instance; otherwise, false.</returns>
+      <param name="obj">The object to compare to this instance.</param>
+    </member>
+    <member name="M:System.Web.Razor.Parser.SyntaxTree.RazorError.Equals(System.Web.Razor.Parser.SyntaxTree.RazorError)">
+      <summary>This type/member supports the .NET Framework infrastructure and is not intended to be used directly from your code. Determines whether the specified object is equal to this instance.</summary>
+      <returns>true if the specified object is equal to this instance; otherwise, false.</returns>
+      <param name="other">The object to compare to this instance.</param>
+    </member>
+    <member name="M:System.Web.Razor.Parser.SyntaxTree.RazorError.GetHashCode">
+      <summary>This type/member supports the .NET Framework infrastructure and is not intended to be used directly from your code. Returns the hash code for the current instance.</summary>
+      <returns>The hash code for the current instance.</returns>
+    </member>
+    <member name="P:System.Web.Razor.Parser.SyntaxTree.RazorError.Length">
+      <summary>This type/member supports the .NET Framework infrastructure and is not intended to be used directly from your code. Gets or sets the length for the error.</summary>
+      <returns>The length for the error.</returns>
+    </member>
+    <member name="P:System.Web.Razor.Parser.SyntaxTree.RazorError.Location">
+      <summary>This type/member supports the .NET Framework infrastructure and is not intended to be used directly from your code. Gets the source location of the error.</summary>
+      <returns>The source location of the error.</returns>
+    </member>
+    <member name="P:System.Web.Razor.Parser.SyntaxTree.RazorError.Message">
+      <summary>This type/member supports the .NET Framework infrastructure and is not intended to be used directly from your code. Gets or sets the error message.</summary>
+      <returns>The error message.</returns>
+    </member>
+    <member name="M:System.Web.Razor.Parser.SyntaxTree.RazorError.ToString">
+      <summary>This type/member supports the .NET Framework infrastructure and is not intended to be used directly from your code. Returns the string representation of this error instance.</summary>
+      <returns>The string representation of this error instance.</returns>
+    </member>
+    <member name="T:System.Web.Razor.Parser.SyntaxTree.Span">
+      <summary>This type/member supports the .NET Framework infrastructure and is not intended to be used directly from your code. Represents a Razor parse tree node that contains the all the content of a block node.</summary>
+    </member>
+    <member name="M:System.Web.Razor.Parser.SyntaxTree.Span.#ctor(System.Web.Razor.Parser.SyntaxTree.SpanBuilder)">
+      <summary>This type/member supports the .NET Framework infrastructure and is not intended to be used directly from your code. Initializes a new instance of the <see cref="T:System.Web.Razor.Parser.SyntaxTree.Span" /> class.</summary>
+      <param name="builder">The builder to use for this span.</param>
+    </member>
+    <member name="M:System.Web.Razor.Parser.SyntaxTree.Span.Accept(System.Web.Razor.Parser.ParserVisitor)">
+      <summary>This type/member supports the .NET Framework infrastructure and is not intended to be used directly from your code. Accepts visit from the specified visitor.</summary>
+      <param name="visitor">The object that performs the visit.</param>
+    </member>
+    <member name="M:System.Web.Razor.Parser.SyntaxTree.Span.Change(System.Action{System.Web.Razor.Parser.SyntaxTree.SpanBuilder})">
+      <summary>This type/member supports the .NET Framework infrastructure and is not intended to be used directly from your code. Changes the span builder for this span.</summary>
+      <param name="changes">A delegate that will be executed along with this change.</param>
+    </member>
+    <member name="M:System.Web.Razor.Parser.SyntaxTree.Span.ChangeStart(System.Web.Razor.Text.SourceLocation)">
+      <summary>This type/member supports the .NET Framework infrastructure and is not intended to be used directly from your code. Sets the start character location of this span.</summary>
+      <param name="newStart">The new start location to set for this span.</param>
+    </member>
+    <member name="P:System.Web.Razor.Parser.SyntaxTree.Span.CodeGenerator">
+      <summary>This type/member supports the .NET Framework infrastructure and is not intended to be used directly from your code. Gets or sets the code generator for the span.</summary>
+      <returns>The code generator for the span.</returns>
+    </member>
+    <member name="P:System.Web.Razor.Parser.SyntaxTree.Span.Content">
+      <summary>This type/member supports the .NET Framework infrastructure and is not intended to be used directly from your code. Gets or sets the content of the span.</summary>
+      <returns>The content of the span.</returns>
+    </member>
+    <member name="P:System.Web.Razor.Parser.SyntaxTree.Span.EditHandler">
+      <summary>This type/member supports the .NET Framework infrastructure and is not intended to be used directly from your code. Gets or sets the handler for span edits.</summary>
+      <returns>The handler for span edits.</returns>
+    </member>
+    <member name="M:System.Web.Razor.Parser.SyntaxTree.Span.Equals(System.Object)">
+      <summary>This type/member supports the .NET Framework infrastructure and is not intended to be used directly from your code. Determines whether the specified object is equal to this span.</summary>
+      <returns>true if the specified object is equal to this span; otherwise, false.</returns>
+      <param name="obj">The object to compare to this span.</param>
+    </member>
+    <member name="M:System.Web.Razor.Parser.SyntaxTree.Span.EquivalentTo(System.Web.Razor.Parser.SyntaxTree.SyntaxTreeNode)">
+      <summary>This type/member supports the .NET Framework infrastructure and is not intended to be used directly from your code. Determines whether the specified node is equivalent to this span.</summary>
+      <returns>true if the specified node is equal to this span; otherwise, false.</returns>
+      <param name="node">The node to compare with this span.</param>
+    </member>
+    <member name="M:System.Web.Razor.Parser.SyntaxTree.Span.GetHashCode">
+      <summary>This type/member supports the .NET Framework infrastructure and is not intended to be used directly from your code. Returns the hash code for this current span.</summary>
+      <returns>The hash code for this current span.</returns>
+    </member>
+    <member name="P:System.Web.Razor.Parser.SyntaxTree.Span.IsBlock">
+      <summary>This type/member supports the .NET Framework infrastructure and is not intended to be used directly from your code. Gets or sets a value that indicates whether this node is a block node.</summary>
+      <returns>false.</returns>
+    </member>
+    <member name="P:System.Web.Razor.Parser.SyntaxTree.Span.Kind">
+      <summary>This type/member supports the .NET Framework infrastructure and is not intended to be used directly from your code. Gets or sets the kind for this span.</summary>
+      <returns>One of the values of the <see cref="T:System.Web.Razor.Parser.SyntaxTree.SpanKind" /> enumeration.</returns>
+    </member>
+    <member name="P:System.Web.Razor.Parser.SyntaxTree.Span.Length">
+      <summary>This type/member supports the .NET Framework infrastructure and is not intended to be used directly from your code. Gets or sets the length of the span content.</summary>
+      <returns>The length of the span content.</returns>
+    </member>
+    <member name="P:System.Web.Razor.Parser.SyntaxTree.Span.Next">
+      <summary>This type/member supports the .NET Framework infrastructure and is not intended to be used directly from your code. Gets or sets the next span in the tree node.</summary>
+      <returns>The next span in the tree node.</returns>
+    </member>
+    <member name="P:System.Web.Razor.Parser.SyntaxTree.Span.Previous">
+      <summary>This type/member supports the .NET Framework infrastructure and is not intended to be used directly from your code. Gets or sets the previous span in the tree node.</summary>
+      <returns>The previous span in the tree node.</returns>
+    </member>
+    <member name="M:System.Web.Razor.Parser.SyntaxTree.Span.ReplaceWith(System.Web.Razor.Parser.SyntaxTree.SpanBuilder)">
+      <summary>This type/member supports the .NET Framework infrastructure and is not intended to be used directly from your code. Replaces the span builder for this span with the specified span builder.</summary>
+      <param name="builder">The new builder to use for this span.</param>
+    </member>
+    <member name="P:System.Web.Razor.Parser.SyntaxTree.Span.Start">
+      <summary>This type/member supports the .NET Framework infrastructure and is not intended to be used directly from your code. Gets or sets the start character location of the span.</summary>
+      <returns>The start character location of the span.</returns>
+    </member>
+    <member name="P:System.Web.Razor.Parser.SyntaxTree.Span.Symbols">
+      <summary>This type/member supports the .NET Framework infrastructure and is not intended to be used directly from your code. Gets or sets the symbols used to generate the code for the span.</summary>
+      <returns>The symbols used to generate the code for the span.</returns>
+    </member>
+    <member name="M:System.Web.Razor.Parser.SyntaxTree.Span.ToString">
+      <summary>This type/member supports the .NET Framework infrastructure and is not intended to be used directly from your code. Returns the string representation of this current span.</summary>
+      <returns>The string representation of this current span.</returns>
+    </member>
+    <member name="T:System.Web.Razor.Parser.SyntaxTree.SpanBuilder">
+      <summary>This type/member supports the .NET Framework infrastructure and is not intended to be used directly from your code.Represents the span builder for the syntax tree.</summary>
+    </member>
+    <member name="M:System.Web.Razor.Parser.SyntaxTree.SpanBuilder.#ctor">
+      <summary>This type/member supports the .NET Framework infrastructure and is not intended to be used directly from your code.Initializes a new instance of the <see cref="T:System.Web.Razor.Parser.SyntaxTree.SpanBuilder" /> class.</summary>
+    </member>
+    <member name="M:System.Web.Razor.Parser.SyntaxTree.SpanBuilder.#ctor(System.Web.Razor.Parser.SyntaxTree.Span)">
+      <summary>This type/member supports the .NET Framework infrastructure and is not intended to be used directly from your code.Initializes a new instance of the <see cref="T:System.Web.Razor.Parser.SyntaxTree.SpanBuilder" /> class.</summary>
+      <param name="original">The original span.</param>
+    </member>
+    <member name="M:System.Web.Razor.Parser.SyntaxTree.SpanBuilder.Accept(System.Web.Razor.Tokenizer.Symbols.ISymbol)">
+      <summary>This type/member supports the .NET Framework infrastructure and is not intended to be used directly from your code.Accepts the given symbol for the span builder.</summary>
+      <param name="symbol">The symbol builder.</param>
+    </member>
+    <member name="M:System.Web.Razor.Parser.SyntaxTree.SpanBuilder.Build">
+      <summary>This type/member supports the .NET Framework infrastructure and is not intended to be used directly from your code.Builds a span builder for this instance.</summary>
+      <returns>A span builder for this instance.</returns>
+    </member>
+    <member name="M:System.Web.Razor.Parser.SyntaxTree.SpanBuilder.ClearSymbols">
+      <summary>This type/member supports the .NET Framework infrastructure and is not intended to be used directly from your code.Clears the symbols of the span builder.</summary>
+    </member>
+    <member name="P:System.Web.Razor.Parser.SyntaxTree.SpanBuilder.CodeGenerator">
+      <summary>This type/member supports the .NET Framework infrastructure and is not intended to be used directly from your code.Gets or sets the span code generator.</summary>
+      <returns>The span code generator.</returns>
+    </member>
+    <member name="P:System.Web.Razor.Parser.SyntaxTree.SpanBuilder.EditHandler">
+      <summary>This type/member supports the .NET Framework infrastructure and is not intended to be used directly from your code.Gets or sets the span edit handler of the builder.</summary>
+      <returns>The span edit handler of the builder.</returns>
+    </member>
+    <member name="P:System.Web.Razor.Parser.SyntaxTree.SpanBuilder.Kind">
+      <summary>This type/member supports the .NET Framework infrastructure and is not intended to be used directly from your code.Gets or sets the span kind of the span builder.</summary>
+      <returns>The span kind of the span builder.</returns>
+    </member>
+    <member name="M:System.Web.Razor.Parser.SyntaxTree.SpanBuilder.Reset">
+      <summary>This type/member supports the .NET Framework infrastructure and is not intended to be used directly from your code.Resets the span builder.</summary>
+    </member>
+    <member name="P:System.Web.Razor.Parser.SyntaxTree.SpanBuilder.Start">
+      <summary>This type/member supports the .NET Framework infrastructure and is not intended to be used directly from your code.Gets or sets the source location of the span builder.</summary>
+      <returns>The source location of the span builder.</returns>
+    </member>
+    <member name="P:System.Web.Razor.Parser.SyntaxTree.SpanBuilder.Symbols">
+      <summary>This type/member supports the .NET Framework infrastructure and is not intended to be used directly from your code.Gets the symbols for a generic read-only collection.</summary>
+      <returns>The symbols for a generic read-only collection.</returns>
+    </member>
+    <member name="T:System.Web.Razor.Parser.SyntaxTree.SpanKind">
+      <summary>This type/member supports the .NET Framework infrastructure and is not intended to be used directly from your code.</summary>
+    </member>
+    <member name="F:System.Web.Razor.Parser.SyntaxTree.SpanKind.Code">
+      <summary>This type/member supports the .NET Framework infrastructure and is not intended to be used directly from your code.</summary>
+    </member>
+    <member name="F:System.Web.Razor.Parser.SyntaxTree.SpanKind.Comment">
+      <summary>This type/member supports the .NET Framework infrastructure and is not intended to be used directly from your code.</summary>
+    </member>
+    <member name="F:System.Web.Razor.Parser.SyntaxTree.SpanKind.Markup">
+      <summary>This type/member supports the .NET Framework infrastructure and is not intended to be used directly from your code.</summary>
+    </member>
+    <member name="F:System.Web.Razor.Parser.SyntaxTree.SpanKind.MetaCode">
+      <summary>This type/member supports the .NET Framework infrastructure and is not intended to be used directly from your code.</summary>
+    </member>
+    <member name="F:System.Web.Razor.Parser.SyntaxTree.SpanKind.Transition">
+      <summary>This type/member supports the .NET Framework infrastructure and is not intended to be used directly from your code.</summary>
+    </member>
+    <member name="T:System.Web.Razor.Parser.SyntaxTree.SyntaxTreeNode">
+      <summary>This type/member supports the .NET Framework infrastructure and is not intended to be used directly from your code.Represents the node for the syntax tree.</summary>
+    </member>
+    <member name="M:System.Web.Razor.Parser.SyntaxTree.SyntaxTreeNode.#ctor">
+      <summary>This type/member supports the .NET Framework infrastructure and is not intended to be used directly from your code.Initializes a new instance of the <see cref="T:System.Web.Razor.Parser.SyntaxTree.SyntaxTreeNode" /> class.</summary>
+    </member>
+    <member name="M:System.Web.Razor.Parser.SyntaxTree.SyntaxTreeNode.Accept(System.Web.Razor.Parser.ParserVisitor)">
+      <summary>This type/member supports the .NET Framework infrastructure and is not intended to be used directly from your code.Accepts the visitor of the tree node.</summary>
+      <param name="visitor">The parser visitor.</param>
+    </member>
+    <member name="M:System.Web.Razor.Parser.SyntaxTree.SyntaxTreeNode.EquivalentTo(System.Web.Razor.Parser.SyntaxTree.SyntaxTreeNode)">
+      <summary>This type/member supports the .NET Framework infrastructure and is not intended to be used directly from your code.Indicates whether the syntax tree node is equivalent to given node.</summary>
+      <returns>true the syntax tree node is equivalent to given node; false.</returns>
+      <param name="node">The given node.</param>
+    </member>
+    <member name="P:System.Web.Razor.Parser.SyntaxTree.SyntaxTreeNode.IsBlock">
+      <summary>This type/member supports the .NET Framework infrastructure and is not intended to be used directly from your code.Gets a value that indicates whether the syntax tree node is a block-level object.</summary>
+      <returns>true if the syntax tree node is a block-level object; otherwise, false.</returns>
+    </member>
+    <member name="P:System.Web.Razor.Parser.SyntaxTree.SyntaxTreeNode.Length">
+      <summary>This type/member supports the .NET Framework infrastructure and is not intended to be used directly from your code.Gets the length of the syntax tree node.</summary>
+      <returns>The length of the syntax tree node.</returns>
+    </member>
+    <member name="P:System.Web.Razor.Parser.SyntaxTree.SyntaxTreeNode.Parent">
+      <summary>This type/member supports the .NET Framework infrastructure and is not intended to be used directly from your code.Gets the parent tree node of the current tree node.</summary>
+      <returns>The parent tree node of the current tree node.</returns>
+    </member>
+    <member name="P:System.Web.Razor.Parser.SyntaxTree.SyntaxTreeNode.Start">
+      <summary>This type/member supports the .NET Framework infrastructure and is not intended to be used directly from your code.Gets the specific source location for the syntax tree node.</summary>
+      <returns>The specific source location for the syntax tree node.</returns>
+    </member>
+    <member name="T:System.Web.Razor.Text.BufferingTextReader">
+      <summary>This type/member supports the .NET Framework infrastructure and is not intended to be used directly from your code. Provides a lookahead buffer for the text reader.</summary>
+    </member>
+    <member name="M:System.Web.Razor.Text.BufferingTextReader.#ctor(System.IO.TextReader)">
+      <summary>This type/member supports the .NET Framework infrastructure and is not intended to be used directly from your code.Initializes a new instance of the <see cref="T:System.Web.Razor.Text.BufferingTextReader" /> class.</summary>
+      <param name="source">The text reader for the buffer.</param>
+    </member>
+    <member name="M:System.Web.Razor.Text.BufferingTextReader.BeginLookahead">
+      <summary>This type/member supports the .NET Framework infrastructure and is not intended to be used directly from your code. Begins the lookahead buffering operation for this <see cref="T:System.Web.Razor.Text.BufferingTextReader" />.</summary>
+      <returns>A disposable action that ends the lookahead buffering.</returns>
+    </member>
+    <member name="M:System.Web.Razor.Text.BufferingTextReader.CancelBacktrack">
+      <summary>This type/member supports the .NET Framework infrastructure and is not intended to be used directly from your code. Discards the backtrack context associated the lookahead buffering operation.</summary>
+    </member>
+    <member name="P:System.Web.Razor.Text.BufferingTextReader.CurrentCharacter">
+      <summary>This type/member supports the .NET Framework infrastructure and is not intended to be used directly from your code. Gets the current character in the buffer.</summary>
+      <returns>The current character in the buffer.</returns>
+    </member>
+    <member name="P:System.Web.Razor.Text.BufferingTextReader.CurrentLocation">
+      <summary>This type/member supports the .NET Framework infrastructure and is not intended to be used directly from your code. Gets or sets the current location of the character in the buffer.</summary>
+      <returns>The current location of the character in the buffer.</returns>
+    </member>
+    <member name="M:System.Web.Razor.Text.BufferingTextReader.Dispose(System.Boolean)">
+      <summary>This type/member supports the .NET Framework infrastructure and is not intended to be used directly from your code. Releases the unmanaged resources used by the current instance of this class, and optionally releases the managed resources.</summary>
+      <param name="disposing">true to release both managed and unmanaged resources; false to release only unmanaged resources.</param>
+    </member>
+    <member name="M:System.Web.Razor.Text.BufferingTextReader.ExpandBuffer">
+      <summary>This type/member supports the .NET Framework infrastructure and is not intended to be used directly from your code. Reads the next character from the text reader and appends it to the lookahead buffer.</summary>
+      <returns>true if a character was read from the text reader; otherwise, false.</returns>
+    </member>
+    <member name="M:System.Web.Razor.Text.BufferingTextReader.NextCharacter">
+      <summary>This type/member supports the .NET Framework infrastructure and is not intended to be used directly from your code. Advances the buffer position to the next character.</summary>
+    </member>
+    <member name="M:System.Web.Razor.Text.BufferingTextReader.Peek">
+      <summary>This type/member supports the .NET Framework infrastructure and is not intended to be used directly from your code. Returns the current character in the buffer.</summary>
+      <returns>The current character in the buffer.</returns>
+    </member>
+    <member name="M:System.Web.Razor.Text.BufferingTextReader.Read">
+      <summary>This type/member supports the .NET Framework infrastructure and is not intended to be used directly from your code. Returns the current character from the buffer and advances the buffer position to the next character.</summary>
+      <returns>The current character from the buffer.</returns>
+    </member>
+    <member name="T:System.Web.Razor.Text.ITextBuffer">
+      <summary>This type/member supports the .NET Framework infrastructure and is not intended to be used directly from your code.</summary>
+    </member>
+    <member name="P:System.Web.Razor.Text.ITextBuffer.Length">
+      <summary>This type/member supports the .NET Framework infrastructure and is not intended to be used directly from your code.</summary>
+    </member>
+    <member name="M:System.Web.Razor.Text.ITextBuffer.Peek">
+      <summary>This type/member supports the .NET Framework infrastructure and is not intended to be used directly from your code.</summary>
+    </member>
+    <member name="P:System.Web.Razor.Text.ITextBuffer.Position">
+      <summary>This type/member supports the .NET Framework infrastructure and is not intended to be used directly from your code.</summary>
+    </member>
+    <member name="M:System.Web.Razor.Text.ITextBuffer.Read">
+      <summary>This type/member supports the .NET Framework infrastructure and is not intended to be used directly from your code.</summary>
+    </member>
+    <member name="T:System.Web.Razor.Text.ITextDocument">
+      <summary>This type/member supports the .NET Framework infrastructure and is not intended to be used directly from your code.</summary>
+    </member>
+    <member name="P:System.Web.Razor.Text.ITextDocument.Location">
+      <summary>This type/member supports the .NET Framework infrastructure and is not intended to be used directly from your code.</summary>
+    </member>
+    <member name="T:System.Web.Razor.Text.LocationTagged`1">
+      <summary>This type/member supports the .NET Framework infrastructure and is not intended to be used directly from your code.Represents a location tagged.</summary>
+      <typeparam name="T">The type of the location tagged.</typeparam>
+    </member>
+    <member name="M:System.Web.Razor.Text.LocationTagged`1.#ctor(`0,System.Int32,System.Int32,System.Int32)">
+      <summary>This type/member supports the .NET Framework infrastructure and is not intended to be used directly from your code.Initializes a new instance of the <see cref="T:System.Web.Razor.Text.LocationTagged`1" /> class.</summary>
+      <param name="value">The value of the source.</param>
+      <param name="offset">The offset.</param>
+      <param name="line">The line.</param>
+      <param name="col">The column location of the source.</param>
+    </member>
+    <member name="M:System.Web.Razor.Text.LocationTagged`1.#ctor(`0,System.Web.Razor.Text.SourceLocation)">
+      <summary>This type/member supports the .NET Framework infrastructure and is not intended to be used directly from your code.Initializes a new instance of the <see cref="T:System.Web.Razor.Text.LocationTagged`1" /> class.</summary>
+      <param name="value">The value of the source.</param>
+      <param name="location">The location of the source.</param>
+    </member>
+    <member name="M:System.Web.Razor.Text.LocationTagged`1.Equals(System.Object)">
+      <summary>This type/member supports the .NET Framework infrastructure and is not intended to be used directly from your code.Determines whether the specified object is equal to the current object.</summary>
+      <returns>true if the specified object is equal to the current object; otherwise, false.</returns>
+      <param name="obj">The object to compare to.</param>
+    </member>
+    <member name="M:System.Web.Razor.Text.LocationTagged`1.GetHashCode">
+      <summary>This type/member supports the .NET Framework infrastructure and is not intended to be used directly from your code.Returns the hash code for the current instance.</summary>
+      <returns>The hash code for the current instance.</returns>
+    </member>
+    <member name="P:System.Web.Razor.Text.LocationTagged`1.Location">
+      <summary>This type/member supports the .NET Framework infrastructure and is not intended to be used directly from your code.Gets or sets the location of the source.</summary>
+      <returns>The location of the source.</returns>
+    </member>
+    <member name="M:System.Web.Razor.Text.LocationTagged`1.op_Equality(System.Web.Razor.Text.LocationTagged{`0},System.Web.Razor.Text.LocationTagged{`0})">
+      <summary>This type/member supports the .NET Framework infrastructure and is not intended to be used directly from your code.Determines whether the two <see cref="T:System.Web.Razor.Text.LocationTagged{`0}" /> object are equal.</summary>
+      <returns>true if the two object are equal; otherwise, false.</returns>
+      <param name="left">The first object to compare.</param>
+      <param name="right">The second object to compare.</param>
+    </member>
+    <member name="M:System.Web.Razor.Text.LocationTagged`1.op_Implicit(System.Web.Razor.Text.LocationTagged{`0})~`0">
+      <summary>This type/member supports the .NET Framework infrastructure and is not intended to be used directly from your code.Converts the specified value to a <see cref="T:System.Web.Razor.Text.LocationTagged`1" /> object.</summary>
+      <returns>true if successfully converted; otherwise, false.</returns>
+      <param name="value">The value to convert.</param>
+    </member>
+    <member name="M:System.Web.Razor.Text.LocationTagged`1.op_Inequality(System.Web.Razor.Text.LocationTagged{`0},System.Web.Razor.Text.LocationTagged{`0})">
+      <summary>This type/member supports the .NET Framework infrastructure and is not intended to be used directly from your code.Determines whether the two <see cref="T:System.Web.Razor.Text.LocationTagged{`0}" /> object are not equal.</summary>
+      <returns>true if the two object are not equal; otherwise, false.</returns>
+      <param name="left">The first object to compare.</param>
+      <param name="right">The second objet to compare.</param>
+    </member>
+    <member name="M:System.Web.Razor.Text.LocationTagged`1.ToString">
+      <summary>This type/member supports the .NET Framework infrastructure and is not intended to be used directly from your code.Returns a string representation of the current instance.</summary>
+      <returns>The string that represents the current instance.</returns>
+    </member>
+    <member name="M:System.Web.Razor.Text.LocationTagged`1.ToString(System.String,System.IFormatProvider)">
+      <summary>This type/member supports the .NET Framework infrastructure and is not intended to be used directly from your code.Returns a string representation of the current instance.</summary>
+      <returns>A string that represents the current instance.</returns>
+      <param name="format">The format.</param>
+      <param name="formatProvider">The format provider.</param>
+    </member>
+    <member name="P:System.Web.Razor.Text.LocationTagged`1.Value">
+      <summary>This type/member supports the .NET Framework infrastructure and is not intended to be used directly from your code.Gets or sets the value of the source.</summary>
+      <returns>The value of the source.</returns>
+    </member>
+    <member name="T:System.Web.Razor.Text.LookaheadTextReader">
+      <summary>This type/member supports the .NET Framework infrastructure and is not intended to be used directly from your code.</summary>
+    </member>
+    <member name="M:System.Web.Razor.Text.LookaheadTextReader.#ctor">
+      <summary>This type/member supports the .NET Framework infrastructure and is not intended to be used directly from your code.</summary>
+    </member>
+    <member name="M:System.Web.Razor.Text.LookaheadTextReader.BeginLookahead">
+      <summary>This type/member supports the .NET Framework infrastructure and is not intended to be used directly from your code.</summary>
+    </member>
+    <member name="M:System.Web.Razor.Text.LookaheadTextReader.CancelBacktrack">
+      <summary>This type/member supports the .NET Framework infrastructure and is not intended to be used directly from your code.</summary>
+    </member>
+    <member name="P:System.Web.Razor.Text.LookaheadTextReader.CurrentLocation">
+      <summary>This type/member supports the .NET Framework infrastructure and is not intended to be used directly from your code.</summary>
+    </member>
+    <member name="T:System.Web.Razor.Text.LookaheadToken">
+      <summary>This type/member supports the .NET Framework infrastructure and is not intended to be used directly from your code.Represents the token to look for the razor.</summary>
+    </member>
+    <member name="M:System.Web.Razor.Text.LookaheadToken.#ctor(System.Action)">
+      <summary>This type/member supports the .NET Framework infrastructure and is not intended to be used directly from your code.Initializes a new instance of the <see cref="T:System.Web.Razor.Text.LookaheadToken" /> class.</summary>
+      <param name="cancelAction">The action to cancel.</param>
+    </member>
+    <member name="M:System.Web.Razor.Text.LookaheadToken.Accept">
+      <summary>This type/member supports the .NET Framework infrastructure and is not intended to be used directly from your code.Accepts the token.</summary>
+    </member>
+    <member name="M:System.Web.Razor.Text.LookaheadToken.Dispose">
+      <summary>This type/member supports the .NET Framework infrastructure and is not intended to be used directly from your code.Releases the resources used by the current instance of the <see cref="T:System.Web.Razor.Text.LookaheadToken" /> class.</summary>
+    </member>
+    <member name="M:System.Web.Razor.Text.LookaheadToken.Dispose(System.Boolean)">
+      <summary>This type/member supports the .NET Framework infrastructure and is not intended to be used directly from your code.Releases the unmanaged resources used by the <see cref="T:System.Web.Razor.Text.LookaheadToken" /> and optionally releases the managed resources.</summary>
+      <param name="disposing">true to release both managed and unmanaged resources; false to release only unmanaged resources.</param>
+    </member>
+    <member name="T:System.Web.Razor.Text.SeekableTextReader">
+      <summary>This type/member supports the .NET Framework infrastructure and is not intended to be used directly from your code.Represents a reader </summary>
+    </member>
+    <member name="M:System.Web.Razor.Text.SeekableTextReader.#ctor(System.IO.TextReader)">
+      <summary>This type/member supports the .NET Framework infrastructure and is not intended to be used directly from your code.Initializes a new instance of the <see cref="T:System.Web.Razor.Text.SeekableTextReader" /> class.</summary>
+      <param name="source">The source reader.</param>
+    </member>
+    <member name="M:System.Web.Razor.Text.SeekableTextReader.#ctor(System.String)">
+      <summary>This type/member supports the .NET Framework infrastructure and is not intended to be used directly from your code.Initializes a new instance of the <see cref="T:System.Web.Razor.Text.SeekableTextReader" /> class.</summary>
+      <param name="content">The string content.</param>
+    </member>
+    <member name="M:System.Web.Razor.Text.SeekableTextReader.#ctor(System.Web.Razor.Text.ITextBuffer)">
+      <summary>This type/member supports the .NET Framework infrastructure and is not intended to be used directly from your code.Initializes a new instance of the <see cref="T:System.Web.Razor.Text.SeekableTextReader" /> class.</summary>
+      <param name="buffer">The text buffering.</param>
+    </member>
+    <member name="P:System.Web.Razor.Text.SeekableTextReader.Length">
+      <summary>This type/member supports the .NET Framework infrastructure and is not intended to be used directly from your code.Gets the length of the text to read.</summary>
+      <returns>The length of the text to read.</returns>
+    </member>
+    <member name="P:System.Web.Razor.Text.SeekableTextReader.Location">
+      <summary>This type/member supports the .NET Framework infrastructure and is not intended to be used directly from your code.Gets the source of location for the text reader.</summary>
+      <returns>The source of location for the text reader.</returns>
+    </member>
+    <member name="M:System.Web.Razor.Text.SeekableTextReader.Peek">
+      <summary>This type/member supports the .NET Framework infrastructure and is not intended to be used directly from your code.Reads the next character without changing the state of the reader or the character source.</summary>
+      <returns>An integer representing the next character to be read.</returns>
+    </member>
+    <member name="P:System.Web.Razor.Text.SeekableTextReader.Position">
+      <summary>This type/member supports the .NET Framework infrastructure and is not intended to be used directly from your code.Gets or sets the position to seek the text reader.</summary>
+      <returns>The position to seek the text reader.</returns>
+    </member>
+    <member name="M:System.Web.Razor.Text.SeekableTextReader.Read">
+      <summary>This type/member supports the .NET Framework infrastructure and is not intended to be used directly from your code.Reads the next character from the text reader and advances the character position by one character.</summary>
+      <returns>The next character from the text reader, or -1 if no more characters are available.</returns>
+    </member>
+    <member name="T:System.Web.Razor.Text.SourceLocation">
+      <summary>This type/member supports the .NET Framework infrastructure and is not intended to be used directly from your code.Represents a source location.</summary>
+    </member>
+    <member name="M:System.Web.Razor.Text.SourceLocation.#ctor(System.Int32,System.Int32,System.Int32)">
+      <summary>This type/member supports the .NET Framework infrastructure and is not intended to be used directly from your code.Initializes a new instance of the <see cref="T:System.Web.Razor.Text.SourceLocation" /> class.</summary>
+      <param name="absoluteIndex">The absolute index.</param>
+      <param name="lineIndex">The line index.</param>
+      <param name="characterIndex">The character index.</param>
+    </member>
+    <member name="P:System.Web.Razor.Text.SourceLocation.AbsoluteIndex">
+      <summary>This type/member supports the .NET Framework infrastructure and is not intended to be used directly from your code.Gets the absolute index for the source location.</summary>
+      <returns>The absolute index for the source location.</returns>
+    </member>
+    <member name="M:System.Web.Razor.Text.SourceLocation.Add(System.Web.Razor.Text.SourceLocation,System.Web.Razor.Text.SourceLocation)">
+      <summary>This type/member supports the .NET Framework infrastructure and is not intended to be used directly from your code.Adds the two <see cref="T:System.Web.Razor.Text.SourceLocation" /> object.</summary>
+      <returns>The sum of the two <see cref="T:System.Web.Razor.Text.SourceLocation" /> object.</returns>
+      <param name="left">The first object to add.</param>
+      <param name="right">The second object to add.</param>
+    </member>
+    <member name="M:System.Web.Razor.Text.SourceLocation.Advance(System.Web.Razor.Text.SourceLocation,System.String)">
+      <summary>This type/member supports the .NET Framework infrastructure and is not intended to be used directly from your code.Advances the specified object to the given location.</summary>
+      <returns>The source location.</returns>
+      <param name="left">The location where to advance the object.</param>
+      <param name="text">The text that advances to the given location.</param>
+    </member>
+    <member name="P:System.Web.Razor.Text.SourceLocation.CharacterIndex">
+      <summary>This type/member supports the .NET Framework infrastructure and is not intended to be used directly from your code.Gets the character index for the source location.</summary>
+      <returns>The character index for the source location.</returns>
+    </member>
+    <member name="M:System.Web.Razor.Text.SourceLocation.CompareTo(System.Web.Razor.Text.SourceLocation)">
+      <summary>This type/member supports the .NET Framework infrastructure and is not intended to be used directly from your code.Compares current object to the other object.</summary>
+      <returns>The value of the objects compared.</returns>
+      <param name="other">The object to compare.</param>
+    </member>
+    <member name="M:System.Web.Razor.Text.SourceLocation.Equals(System.Object)">
+      <summary>This type/member supports the .NET Framework infrastructure and is not intended to be used directly from your code.Determines whether the specified object is equal to the current object.</summary>
+      <returns>true if the specified object is equal to the current object; otherwise, false.</returns>
+      <param name="obj">The object to compare to.</param>
+    </member>
+    <member name="M:System.Web.Razor.Text.SourceLocation.Equals(System.Web.Razor.Text.SourceLocation)">
+      <summary>This type/member supports the .NET Framework infrastructure and is not intended to be used directly from your code.Determines whether the current <see cref="T:System.Web.Razor.Text.SourceLocation" /> object is equal to the other <see cref="T:System.Web.Razor.Text.SourceLocation" /> object.</summary>
+      <returns>true if the current object is equal to the other object; otherwise, false.</returns>
+      <param name="other">The object to compare to.</param>
+    </member>
+    <member name="M:System.Web.Razor.Text.SourceLocation.GetHashCode">
+      <summary>This type/member supports the .NET Framework infrastructure and is not intended to be used directly from your code.Returns the hash code for this instance.</summary>
+      <returns>The hash code for this instance.</returns>
+    </member>
+    <member name="P:System.Web.Razor.Text.SourceLocation.LineIndex">
+      <summary>This type/member supports the .NET Framework infrastructure and is not intended to be used directly from your code.Gets the line index for the source location.</summary>
+      <returns>The line index for the source location.</returns>
+    </member>
+    <member name="M:System.Web.Razor.Text.SourceLocation.op_Addition(System.Web.Razor.Text.SourceLocation,System.Web.Razor.Text.SourceLocation)">
+      <summary>This type/member supports the .NET Framework infrastructure and is not intended to be used directly from your code.Adds the two <see cref="T:System.Web.Razor.Text.SourceLocation" /> object.</summary>
+      <returns>The <see cref="T:System.Web.Razor.Text.SourceLocation" /> that is the sum of the two object.</returns>
+      <param name="left">The object to add.</param>
+      <param name="right">The object to add.</param>
+    </member>
+    <member name="M:System.Web.Razor.Text.SourceLocation.op_Equality(System.Web.Razor.Text.SourceLocation,System.Web.Razor.Text.SourceLocation)">
+      <summary>This type/member supports the .NET Framework infrastructure and is not intended to be used directly from your code.Determines whether the two objects are equal.</summary>
+      <returns>true if the two objects are equal; otherwise, false.</returns>
+      <param name="left">The first object to compare.</param>
+      <param name="right">The second object to compare.</param>
+    </member>
+    <member name="M:System.Web.Razor.Text.SourceLocation.op_GreaterThan(System.Web.Razor.Text.SourceLocation,System.Web.Razor.Text.SourceLocation)">
+      <summary>This type/member supports the .NET Framework infrastructure and is not intended to be used directly from your code.Determines whether the first object is greater than the second object.</summary>
+      <returns>true if the first object is greater than the second object; otherwise, false.</returns>
+      <param name="left">The first object.</param>
+      <param name="right">The second object.</param>
+    </member>
+    <member name="M:System.Web.Razor.Text.SourceLocation.op_Inequality(System.Web.Razor.Text.SourceLocation,System.Web.Razor.Text.SourceLocation)">
+      <summary>This type/member supports the .NET Framework infrastructure and is not intended to be used directly from your code.Determines whether the two <see cref="T:System.Web.Razor.Text.SourceLocation" />object are not equal.</summary>
+      <returns>true if the two objects are not equal; otherwise, false.</returns>
+      <param name="left">The object to compare.</param>
+      <param name="right">The object to compare.</param>
+    </member>
+    <member name="M:System.Web.Razor.Text.SourceLocation.op_LessThan(System.Web.Razor.Text.SourceLocation,System.Web.Razor.Text.SourceLocation)">
+      <summary>This type/member supports the .NET Framework infrastructure and is not intended to be used directly from your code.Determines whether the first object is less than the second object.</summary>
+      <returns>true if the first object is greater than the second object; otherwise, false.</returns>
+      <param name="left">The first object.</param>
+      <param name="right">The second object.</param>
+    </member>
+    <member name="M:System.Web.Razor.Text.SourceLocation.op_Subtraction(System.Web.Razor.Text.SourceLocation,System.Web.Razor.Text.SourceLocation)">
+      <summary>This type/member supports the .NET Framework infrastructure and is not intended to be used directly from your code.</summary>
+      <returns>Returns <see cref="T:System.Web.Razor.Text.SourceLocation" />.</returns>
+    </member>
+    <member name="M:System.Web.Razor.Text.SourceLocation.Subtract(System.Web.Razor.Text.SourceLocation,System.Web.Razor.Text.SourceLocation)">
+      <summary>This type/member supports the .NET Framework infrastructure and is not intended to be used directly from your code.Subtracts the first object to the second object.</summary>
+      <returns>The difference of the two objects.</returns>
+      <param name="left">The first object.</param>
+      <param name="right">The second object.</param>
+    </member>
+    <member name="M:System.Web.Razor.Text.SourceLocation.ToString">
+      <summary>This type/member supports the .NET Framework infrastructure and is not intended to be used directly from your code.Returns a string representation of the source location.</summary>
+      <returns>A string representation of the source location.</returns>
+    </member>
+    <member name="F:System.Web.Razor.Text.SourceLocation.Undefined">
+      <summary>This type/member supports the .NET Framework infrastructure and is not intended to be used directly from your code.</summary>
+    </member>
+    <member name="F:System.Web.Razor.Text.SourceLocation.Zero">
+      <summary>This type/member supports the .NET Framework infrastructure and is not intended to be used directly from your code.</summary>
+    </member>
+    <member name="T:System.Web.Razor.Text.SourceLocationTracker">
+      <summary>This type/member supports the .NET Framework infrastructure and is not intended to be used directly from your code.Provides a source location tracker.</summary>
+    </member>
+    <member name="M:System.Web.Razor.Text.SourceLocationTracker.#ctor">
+      <summary>This type/member supports the .NET Framework infrastructure and is not intended to be used directly from your code.Initializes a new instance of the <see cref="T:System.Web.Razor.Text.SourceLocationTracker" /> class.</summary>
+    </member>
+    <member name="M:System.Web.Razor.Text.SourceLocationTracker.#ctor(System.Web.Razor.Text.SourceLocation)">
+      <summary>This type/member supports the .NET Framework infrastructure and is not intended to be used directly from your code.Initializes a new instance of the <see cref="T:System.Web.Razor.Text.SourceLocationTracker" /> class.</summary>
+      <param name="currentLocation">The current location of the source.</param>
+    </member>
+    <member name="M:System.Web.Razor.Text.SourceLocationTracker.CalculateNewLocation(System.Web.Razor.Text.SourceLocation,System.String)">
+      <summary>This type/member supports the .NET Framework infrastructure and is not intended to be used directly from your code.Calculates the new location of the source.</summary>
+      <returns>The new source location.</returns>
+      <param name="lastPosition">The last position.</param>
+      <param name="newContent">The new content.</param>
+    </member>
+    <member name="P:System.Web.Razor.Text.SourceLocationTracker.CurrentLocation">
+      <summary>This type/member supports the .NET Framework infrastructure and is not intended to be used directly from your code.Gets or sets the current location of the source.</summary>
+      <returns>The current location of the source.</returns>
+    </member>
+    <member name="M:System.Web.Razor.Text.SourceLocationTracker.UpdateLocation(System.Char,System.Char)">
+      <summary>This type/member supports the .NET Framework infrastructure and is not intended to be used directly from your code.Updates the source location.</summary>
+      <param name="characterRead">The character to read.</param>
+      <param name="nextCharacter">The character to update.</param>
+    </member>
+    <member name="M:System.Web.Razor.Text.SourceLocationTracker.UpdateLocation(System.String)">
+      <summary>This type/member supports the .NET Framework infrastructure and is not intended to be used directly from your code.Updates the location of the source.</summary>
+      <returns>The <see cref="T:System.Web.Razor.Text.SourceLocationTracker" /> object.</returns>
+      <param name="content">The content of the source.</param>
+    </member>
+    <member name="T:System.Web.Razor.Text.TextBufferReader">
+      <summary>This type/member supports the .NET Framework infrastructure and is not intended to be used directly from your code.Provides a reader for text buffer.</summary>
+    </member>
+    <member name="M:System.Web.Razor.Text.TextBufferReader.#ctor(System.Web.Razor.Text.ITextBuffer)">
+      <summary>This type/member supports the .NET Framework infrastructure and is not intended to be used directly from your code.Initializes a new instance of the <see cref="T:System.Web.Razor.Text.TextBufferReader" /> class.</summary>
+      <param name="buffer">The text buffer to read.</param>
+    </member>
+    <member name="M:System.Web.Razor.Text.TextBufferReader.BeginLookahead">
+      <summary>This type/member supports the .NET Framework infrastructure and is not intended to be used directly from your code.Begins reading the current text buffer.</summary>
+      <returns>An <see cref="T:System.IDisposable" /> instance that stops the text buffer.</returns>
+    </member>
+    <member name="M:System.Web.Razor.Text.TextBufferReader.CancelBacktrack">
+      <summary>This type/member supports the .NET Framework infrastructure and is not intended to be used directly from your code.Cancels backtrack.</summary>
+    </member>
+    <member name="P:System.Web.Razor.Text.TextBufferReader.CurrentLocation">
+      <summary>This type/member supports the .NET Framework infrastructure and is not intended to be used directly from your code.Gets the current location of the text buffer.</summary>
+      <returns>The current location of the text buffer.</returns>
+    </member>
+    <member name="M:System.Web.Razor.Text.TextBufferReader.Dispose(System.Boolean)">
+      <summary>This type/member supports the .NET Framework infrastructure and is not intended to be used directly from your code.Releases the unmanaged resources used by the <see cref="T:System.Web.Razor.Text.TextBufferReader" /> class and optionally releases the managed resources.</summary>
+      <param name="disposing">true to release both managed and unmanaged resources; false to release only unmanaged resources.</param>
+    </member>
+    <member name="M:System.Web.Razor.Text.TextBufferReader.Peek">
+      <summary>This type/member supports the .NET Framework infrastructure and is not intended to be used directly from your code.Returns the next text buffer to read.</summary>
+      <returns>The next text buffer to read.</returns>
+    </member>
+    <member name="M:System.Web.Razor.Text.TextBufferReader.Read">
+      <summary>This type/member supports the .NET Framework infrastructure and is not intended to be used directly from your code.Reads the current text buffer.</summary>
+      <returns>The current text buffer.</returns>
+    </member>
+    <member name="T:System.Web.Razor.Text.TextChange">
+      <summary>This type/member supports the .NET Framework infrastructure and is not intended to be used directly from your code.Describes a text change operation.</summary>
+    </member>
+    <member name="M:System.Web.Razor.Text.TextChange.#ctor(System.Int32,System.Int32,System.Web.Razor.Text.ITextBuffer,System.Int32,System.Int32,System.Web.Razor.Text.ITextBuffer)">
+      <summary>This type/member supports the .NET Framework infrastructure and is not intended to be used directly from your code.Initializes a new instance of the <see cref="T:System.Web.Razor.Text.TextChange" /> class.</summary>
+      <param name="oldPosition">The position of the text change in the snapshot immediately before the change.</param>
+      <param name="oldLength">The length of the old text.</param>
+      <param name="oldBuffer">An old text buffer.</param>
+      <param name="newPosition">The position of the text change in the snapshot immediately after the change.</param>
+      <param name="newLength">The length of the new text.</param>
+      <param name="newBuffer">A new text buffer.</param>
+    </member>
+    <member name="M:System.Web.Razor.Text.TextChange.ApplyChange(System.String,System.Int32)">
+      <summary>This type/member supports the .NET Framework infrastructure and is not intended to be used directly from your code.Applies the specified text change.</summary>
+      <returns>A string that contains the value of the text.</returns>
+      <param name="content">The content of the text.</param>
+      <param name="changeOffset">The change offset.</param>
+    </member>
+    <member name="M:System.Web.Razor.Text.TextChange.ApplyChange(System.Web.Razor.Parser.SyntaxTree.Span)">
+      <summary>This type/member supports the .NET Framework infrastructure and is not intended to be used directly from your code.Applies the specified text change.</summary>
+      <returns>A string that contains the value of the text.</returns>
+      <param name="span">The span of the text change.</param>
+    </member>
+    <member name="M:System.Web.Razor.Text.TextChange.Equals(System.Object)">
+      <summary>This type/member supports the .NET Framework infrastructure and is not intended to be used directly from your code.Determines whether the specified object is equal to the current object.</summary>
+      <returns>true if the specified object is equal to the current object; otherwise, false.</returns>
+      <param name="obj">The object to compare to.</param>
+    </member>
+    <member name="M:System.Web.Razor.Text.TextChange.GetHashCode">
+      <summary>This type/member supports the .NET Framework infrastructure and is not intended to be used directly from your code.Gets the hash code for this text change.</summary>
+      <returns>The hash code for this text change.</returns>
+    </member>
+    <member name="P:System.Web.Razor.Text.TextChange.IsDelete">
+      <summary>This type/member supports the .NET Framework infrastructure and is not intended to be used directly from your code.Gets a value that indicates whether this text change is a delete.</summary>
+      <returns>true if this text change is a delete; otherwise, false.</returns>
+    </member>
+    <member name="P:System.Web.Razor.Text.TextChange.IsInsert">
+      <summary>This type/member supports the .NET Framework infrastructure and is not intended to be used directly from your code.Gets a value that indicates whether this text change is an insert.</summary>
+      <returns>true if this text change is an insert; otherwise, false.</returns>
+    </member>
+    <member name="P:System.Web.Razor.Text.TextChange.IsReplace">
+      <summary>This type/member supports the .NET Framework infrastructure and is not intended to be used directly from your code.Gets a value that indicates whether this text change is a replace.</summary>
+      <returns>true if this text change is a replace; otherwise, false.</returns>
+    </member>
+    <member name="P:System.Web.Razor.Text.TextChange.NewBuffer">
+      <summary>This type/member supports the .NET Framework infrastructure and is not intended to be used directly from your code.Gets or sets a new text buffer.</summary>
+      <returns>A new text buffer.</returns>
+    </member>
+    <member name="P:System.Web.Razor.Text.TextChange.NewLength">
+      <summary>This type/member supports the .NET Framework infrastructure and is not intended to be used directly from your code.Gets or sets the length of the new text.</summary>
+      <returns>The length of the new text.</returns>
+    </member>
+    <member name="P:System.Web.Razor.Text.TextChange.NewPosition">
+      <summary>This type/member supports the .NET Framework infrastructure and is not intended to be used directly from your code.Gets or sets the position of the text change in the snapshot immediately after the change.</summary>
+      <returns>The position of the text change in the snapshot immediately after the change.</returns>
+    </member>
+    <member name="P:System.Web.Razor.Text.TextChange.NewText">
+      <summary>This type/member supports the .NET Framework infrastructure and is not intended to be used directly from your code.Gets or sets the text that replaced the old text.</summary>
+      <returns>The text that replaced the old text.</returns>
+    </member>
+    <member name="M:System.Web.Razor.Text.TextChange.Normalize">
+      <summary>This type/member supports the .NET Framework infrastructure and is not intended to be used directly from your code.Returns a normalized value of this text change.</summary>
+      <returns>A normalized value of this text change.</returns>
+    </member>
+    <member name="P:System.Web.Razor.Text.TextChange.OldBuffer">
+      <summary>This type/member supports the .NET Framework infrastructure and is not intended to be used directly from your code.Gets or sets an old text buffer.</summary>
+      <returns>An old text buffer.</returns>
+    </member>
+    <member name="P:System.Web.Razor.Text.TextChange.OldLength">
+      <summary>This type/member supports the .NET Framework infrastructure and is not intended to be used directly from your code.Gets or sets the length of the old text.</summary>
+      <returns>The length of the old text.</returns>
+    </member>
+    <member name="P:System.Web.Razor.Text.TextChange.OldPosition">
+      <summary>This type/member supports the .NET Framework infrastructure and is not intended to be used directly from your code.Gets or sets the position of the text change in the snapshot immediately before the change.</summary>
+      <returns>The position of the text change in the snapshot immediately before the change.</returns>
+    </member>
+    <member name="P:System.Web.Razor.Text.TextChange.OldText">
+      <summary>This type/member supports the .NET Framework infrastructure and is not intended to be used directly from your code.Gets or sets the text that was replaced.</summary>
+      <returns>The text that was replaced.</returns>
+    </member>
+    <member name="M:System.Web.Razor.Text.TextChange.op_Equality(System.Web.Razor.Text.TextChange,System.Web.Razor.Text.TextChange)">
+      <summary>This type/member supports the .NET Framework infrastructure and is not intended to be used directly from your code.Determines whether the two text change are equal.</summary>
+      <returns>true if the two text change are equal; otherwise, false.</returns>
+      <param name="left">The left text change.</param>
+      <param name="right">The right text change.</param>
+    </member>
+    <member name="M:System.Web.Razor.Text.TextChange.op_Inequality(System.Web.Razor.Text.TextChange,System.Web.Razor.Text.TextChange)">
+      <summary>This type/member supports the .NET Framework infrastructure and is not intended to be used directly from your code.Determines whether the two text change are not equal.</summary>
+      <returns>true if the two text change are not equal; otherwise, false.</returns>
+      <param name="left">The left text change.</param>
+      <param name="right">The right text change.</param>
+    </member>
+    <member name="M:System.Web.Razor.Text.TextChange.ToString">
+      <summary>This type/member supports the .NET Framework infrastructure and is not intended to be used directly from your code.Returns a string representation of the text change.</summary>
+      <returns>A string representation of the text change.</returns>
+    </member>
+    <member name="T:System.Web.Razor.Text.TextChangeType">
+      <summary>This type/member supports the .NET Framework infrastructure and is not intended to be used directly from your code.</summary>
+    </member>
+    <member name="F:System.Web.Razor.Text.TextChangeType.Insert">
+      <summary>This type/member supports the .NET Framework infrastructure and is not intended to be used directly from your code.</summary>
+    </member>
+    <member name="F:System.Web.Razor.Text.TextChangeType.Remove">
+      <summary>This type/member supports the .NET Framework infrastructure and is not intended to be used directly from your code.</summary>
+    </member>
+    <member name="T:System.Web.Razor.Text.TextDocumentReader">
+      <summary>This type/member supports the .NET Framework infrastructure and is not intended to be used directly from your code.Provides reader for text document.</summary>
+    </member>
+    <member name="M:System.Web.Razor.Text.TextDocumentReader.#ctor(System.Web.Razor.Text.ITextDocument)">
+      <summary>This type/member supports the .NET Framework infrastructure and is not intended to be used directly from your code.Initializes a new instance of the <see cref="T:System.Web.Razor.Text.TextDocumentReader" /> class.</summary>
+      <param name="source">The source to read.</param>
+    </member>
+    <member name="P:System.Web.Razor.Text.TextDocumentReader.Length">
+      <summary>This type/member supports the .NET Framework infrastructure and is not intended to be used directly from your code.Gets the length of the document.</summary>
+      <returns>The length of the document.</returns>
+    </member>
+    <member name="P:System.Web.Razor.Text.TextDocumentReader.Location">
+      <summary>This type/member supports the .NET Framework infrastructure and is not intended to be used directly from your code.Gets the location of the document.</summary>
+      <returns>The location of the document.</returns>
+    </member>
+    <member name="M:System.Web.Razor.Text.TextDocumentReader.Peek">
+      <summary>This type/member supports the .NET Framework infrastructure and is not intended to be used directly from your code.Returns the next text document to read.</summary>
+      <returns>The next text document to read.</returns>
+    </member>
+    <member name="P:System.Web.Razor.Text.TextDocumentReader.Position">
+      <summary>This type/member supports the .NET Framework infrastructure and is not intended to be used directly from your code.Gets or sets the position of the text document.</summary>
+      <returns>The position of the text document.</returns>
+    </member>
+    <member name="M:System.Web.Razor.Text.TextDocumentReader.Read">
+      <summary>This type/member supports the .NET Framework infrastructure and is not intended to be used directly from your code.Reads a specified text document.</summary>
+      <returns>The text document.</returns>
+    </member>
+    <member name="T:System.Web.Razor.Tokenizer.CSharpHelpers">
+      <summary>This type/member supports the .NET Framework infrastructure and is not intended to be used directly from your code.Provides helper functions for the CSharp tokenizer.</summary>
+    </member>
+    <member name="M:System.Web.Razor.Tokenizer.CSharpHelpers.IsIdentifierPart(System.Char)">
+      <summary>This type/member supports the .NET Framework infrastructure and is not intended to be used directly from your code.Determines whether the specified character can be used for identifier.</summary>
+      <returns>true if the specified character can be used for identifier; otherwise, false.</returns>
+      <param name="character">The character to check.</param>
+    </member>
+    <member name="M:System.Web.Razor.Tokenizer.CSharpHelpers.IsIdentifierStart(System.Char)">
+      <summary>This type/member supports the .NET Framework infrastructure and is not intended to be used directly from your code.Determines whether the specified character can be used as an identifier start character.</summary>
+      <returns>true if the specified character can be used as an identifier start character; otherwise, false.</returns>
+      <param name="character">The character to check.</param>
+    </member>
+    <member name="M:System.Web.Razor.Tokenizer.CSharpHelpers.IsRealLiteralSuffix(System.Char)">
+      <summary>This type/member supports the .NET Framework infrastructure and is not intended to be used directly from your code.Determines whether the specified character is a literal suffix for real numbers.</summary>
+      <returns>true if the specified character is a literal suffix for real numbers; otherwise, false.</returns>
+      <param name="character">The character to check.</param>
+    </member>
+    <member name="T:System.Web.Razor.Tokenizer.CSharpTokenizer">
+      <summary>This type/member supports the .NET Framework infrastructure and is not intended to be used directly from your code.Represents a CSharp tokenizer.</summary>
+    </member>
+    <member name="M:System.Web.Razor.Tokenizer.CSharpTokenizer.#ctor(System.Web.Razor.Text.ITextDocument)">
+      <summary>This type/member supports the .NET Framework infrastructure and is not intended to be used directly from your code.Initializes a new instance of the <see cref="T:System.Web.Razor.Tokenizer.CSharpTokenizer" /> class.</summary>
+      <param name="source">The source.</param>
+    </member>
+    <member name="M:System.Web.Razor.Tokenizer.CSharpTokenizer.CreateSymbol(System.Web.Razor.Text.SourceLocation,System.String,System.Web.Razor.Tokenizer.Symbols.CSharpSymbolType,System.Collections.Generic.IEnumerable{System.Web.Razor.Parser.SyntaxTree.RazorError})">
+      <summary>This type/member supports the .NET Framework infrastructure and is not intended to be used directly from your code.Creates a CSharp tokenizer symbol.</summary>
+      <returns>A CSharp tokenizer symbol.</returns>
+      <param name="start">The beginning of the source location.</param>
+      <param name="content">The contents.</param>
+      <param name="type">The CSharp symbol type.</param>
+      <param name="errors">A collection of razor errors.</param>
+    </member>
+    <member name="P:System.Web.Razor.Tokenizer.CSharpTokenizer.RazorCommentStarType">
+      <summary>This type/member supports the .NET Framework infrastructure and is not intended to be used directly from your code.Gets the star type of the <see cref="T:System.Web.Razor.Tokenizer.Symbols.CSharpSymbolType" />.</summary>
+      <returns>The star type of the <see cref="T:System.Web.Razor.Tokenizer.Symbols.CSharpSymbolType" />.</returns>
+    </member>
+    <member name="P:System.Web.Razor.Tokenizer.CSharpTokenizer.RazorCommentTransitionType">
+      <summary>This type/member supports the .NET Framework infrastructure and is not intended to be used directly from your code.Gets the razor comment transition type for the <see cref="T:System.Web.Razor.Tokenizer.Symbols.CSharpSymbolType" />.</summary>
+      <returns>The razor comment transition type for the <see cref="T:System.Web.Razor.Tokenizer.Symbols.CSharpSymbolType" />.</returns>
+    </member>
+    <member name="P:System.Web.Razor.Tokenizer.CSharpTokenizer.RazorCommentType">
+      <summary>This type/member supports the .NET Framework infrastructure and is not intended to be used directly from your code.Gets the razor comment type for the <see cref="T:System.Web.Razor.Tokenizer.Symbols.CSharpSymbolType" />.</summary>
+      <returns>The razor comment type for the <see cref="T:System.Web.Razor.Tokenizer.Symbols.CSharpSymbolType" />.</returns>
+    </member>
+    <member name="P:System.Web.Razor.Tokenizer.CSharpTokenizer.StartState">
+      <summary>This type/member supports the .NET Framework infrastructure and is not intended to be used directly from your code.Gets the state of the machine.</summary>
+      <returns>The state of the machine.</returns>
+    </member>
+    <member name="T:System.Web.Razor.Tokenizer.HtmlTokenizer">
+      <summary>This type/member supports the .NET Framework infrastructure and is not intended to be used directly from your code.Represents the html tokenizer of the razor.</summary>
+    </member>
+    <member name="M:System.Web.Razor.Tokenizer.HtmlTokenizer.#ctor(System.Web.Razor.Text.ITextDocument)">
+      <summary>This type/member supports the .NET Framework infrastructure and is not intended to be used directly from your code.Initializes a new instance of the <see cref="T:System.Web.Razor.Tokenizer.HtmlTokenizer" /> class.</summary>
+      <param name="source">The source for the text document.</param>
+    </member>
+    <member name="M:System.Web.Razor.Tokenizer.HtmlTokenizer.CreateSymbol(System.Web.Razor.Text.SourceLocation,System.String,System.Web.Razor.Tokenizer.Symbols.HtmlSymbolType,System.Collections.Generic.IEnumerable{System.Web.Razor.Parser.SyntaxTree.RazorError})">
+      <summary>This type/member supports the .NET Framework infrastructure and is not intended to be used directly from your code.Creates a symbol for the specified parameters of the html tokenizer.</summary>
+      <returns>A symbol to create for the specified parameters of the html tokenizer.</returns>
+      <param name="start">The source location.</param>
+      <param name="content">The content string.</param>
+      <param name="type">The type of html symbol.</param>
+      <param name="errors">The razor errors.</param>
+    </member>
+    <member name="P:System.Web.Razor.Tokenizer.HtmlTokenizer.RazorCommentStarType">
+      <summary>This type/member supports the .NET Framework infrastructure and is not intended to be used directly from your code.Gets the html symbols for the razor comment star type.</summary>
+      <returns>The html symbols for the razor comment star type.</returns>
+    </member>
+    <member name="P:System.Web.Razor.Tokenizer.HtmlTokenizer.RazorCommentTransitionType">
+      <summary>This type/member supports the .NET Framework infrastructure and is not intended to be used directly from your code.Gets the html symbols for the razor comment transition type.</summary>
+      <returns>The html symbols for the razor comment transition type.</returns>
+    </member>
+    <member name="P:System.Web.Razor.Tokenizer.HtmlTokenizer.RazorCommentType">
+      <summary>This type/member supports the .NET Framework infrastructure and is not intended to be used directly from your code.Gets the html symbols for the razor comment type.</summary>
+      <returns>The html symbols for the razor comment type.</returns>
+    </member>
+    <member name="P:System.Web.Razor.Tokenizer.HtmlTokenizer.StartState">
+      <summary>This type/member supports the .NET Framework infrastructure and is not intended to be used directly from your code.Gets the start of the state machine for the html.</summary>
+      <returns>The start of the state machine for the html.</returns>
+    </member>
+    <member name="T:System.Web.Razor.Tokenizer.ITokenizer">
+      <summary>This type/member supports the .NET Framework infrastructure and is not intended to be used directly from your code.</summary>
+    </member>
+    <member name="M:System.Web.Razor.Tokenizer.ITokenizer.NextSymbol">
+      <summary>This type/member supports the .NET Framework infrastructure and is not intended to be used directly from your code.</summary>
+    </member>
+    <member name="T:System.Web.Razor.Tokenizer.Tokenizer`2">
+      <summary>This type/member supports the .NET Framework infrastructure and is not intended to be used directly from your code.</summary>
+      <typeparam name="TSymbol">The type for the language symbol.</typeparam>
+      <typeparam name="TSymbolType">The enumeration type for the language symbol.</typeparam>
+    </member>
+    <member name="M:System.Web.Razor.Tokenizer.Tokenizer`2.#ctor(System.Web.Razor.Text.ITextDocument)">
+      <summary>This type/member supports the .NET Framework infrastructure and is not intended to be used directly from your code.Initializes a new instance of the <see cref="T:System.Web.Razor.Tokenizer.Tokenizer`2" /> class.</summary>
+      <param name="source">The source.</param>
+    </member>
+    <member name="M:System.Web.Razor.Tokenizer.Tokenizer`2.AfterRazorCommentTransition">
+      <summary>This type/member supports the .NET Framework infrastructure and is not intended to be used directly from your code.Returns a result after the razor comment transition.</summary>
+      <returns>The result.</returns>
+    </member>
+    <member name="M:System.Web.Razor.Tokenizer.Tokenizer`2.At(System.String,System.Boolean)">
+      <summary>This type/member supports the .NET Framework infrastructure and is not intended to be used directly from your code.Determines whether the lookahead buffer contains the expected string.</summary>
+      <returns>true if the lookahead buffer contains the expected string; otherwise, false.</returns>
+      <param name="expected">The string to check.</param>
+      <param name="caseSensitive">true to indicate comparison is case sensitive; otherwise, false.</param>
+    </member>
+    <member name="P:System.Web.Razor.Tokenizer.Tokenizer`2.Buffer">
+      <summary>This type/member supports the .NET Framework infrastructure and is not intended to be used directly from your code.Gets or sets the buffer for the tokenizer.</summary>
+      <returns>The buffer for the tokenizer.</returns>
+    </member>
+    <member name="M:System.Web.Razor.Tokenizer.Tokenizer`2.CharOrWhiteSpace(System.Char)">
+      <summary>Returns a function delegate, that accepts a character parameter and returns a value that indicates whether the character parameter is equal to specified character or white space.</summary>
+      <returns>A function delegate.</returns>
+      <param name="character">The character used to compare.</param>
+    </member>
+    <member name="M:System.Web.Razor.Tokenizer.Tokenizer`2.CreateSymbol(System.Web.Razor.Text.SourceLocation,System.String,`1,System.Collections.Generic.IEnumerable{System.Web.Razor.Parser.SyntaxTree.RazorError})">
+      <summary>This type/member supports the .NET Framework infrastructure and is not intended to be used directly from your code.Creates a language symbol type for the tokenizer with the specified content.</summary>
+      <returns>A language symbol type for the tokenizer.</returns>
+      <param name="start">The start of the source location.</param>
+      <param name="content">The content value.</param>
+      <param name="type">The symbol type.</param>
+      <param name="errors">The razor error.</param>
+    </member>
+    <member name="P:System.Web.Razor.Tokenizer.Tokenizer`2.CurrentCharacter">
+      <summary>This type/member supports the .NET Framework infrastructure and is not intended to be used directly from your code.Gets the current character in the tokenizer.</summary>
+      <returns>The current character.</returns>
+    </member>
+    <member name="P:System.Web.Razor.Tokenizer.Tokenizer`2.CurrentErrors">
+      <summary>This type/member supports the .NET Framework infrastructure and is not intended to be used directly from your code.Gets a list of the current razor errors.</summary>
+      <returns>A list of the current errors.</returns>
+    </member>
+    <member name="P:System.Web.Razor.Tokenizer.Tokenizer`2.CurrentLocation">
+      <summary>This type/member supports the .NET Framework infrastructure and is not intended to be used directly from your code.Gets the current source location.</summary>
+      <returns>The current source location.</returns>
+    </member>
+    <member name="P:System.Web.Razor.Tokenizer.Tokenizer`2.CurrentStart">
+      <summary>This type/member supports the .NET Framework infrastructure and is not intended to be used directly from your code.Gets the current start of the source location.</summary>
+      <returns>The current start of the source location.</returns>
+    </member>
+    <member name="P:System.Web.Razor.Tokenizer.Tokenizer`2.EndOfFile">
+      <summary>This type/member supports the .NET Framework infrastructure and is not intended to be used directly from your code.Gets a value whether the tokenizer current location is at the end of the file.</summary>
+      <returns>true if the tokenizer current location is at the end of the file; otherwise, false.</returns>
+    </member>
+    <member name="M:System.Web.Razor.Tokenizer.Tokenizer`2.EndSymbol(System.Web.Razor.Text.SourceLocation,`1)">
+      <summary>This type/member supports the .NET Framework infrastructure and is not intended to be used directly from your code.Returns the language end symbol type used by the tokenizer.</summary>
+      <returns>The language end symbol type.</returns>
+      <param name="start">The start of the source location.</param>
+      <param name="type">The enumeration type for the language symbol.</param>
+    </member>
+    <member name="M:System.Web.Razor.Tokenizer.Tokenizer`2.EndSymbol(`1)">
+      <summary>This type/member supports the .NET Framework infrastructure and is not intended to be used directly from your code.Returns the language end symbol type used by the tokenizer.</summary>
+      <returns>The language end symbol type.</returns>
+      <param name="type">The enumeration type for the language symbol.</param>
+    </member>
+    <member name="P:System.Web.Razor.Tokenizer.Tokenizer`2.HaveContent">
+      <summary>This type/member supports the .NET Framework infrastructure and is not intended to be used directly from your code.Gets a value whether the tokenizer have content.</summary>
+      <returns>true if the tokenizer have content; otherwise, false.</returns>
+    </member>
+    <member name="M:System.Web.Razor.Tokenizer.Tokenizer`2.MoveNext">
+      <summary>This type/member supports the .NET Framework infrastructure and is not intended to be used directly from your code.Reads to the next character from the code reader.</summary>
+    </member>
+    <member name="M:System.Web.Razor.Tokenizer.Tokenizer`2.NextSymbol">
+      <summary>This type/member supports the .NET Framework infrastructure and is not intended to be used directly from your code.Shows the next symbol to be used.</summary>
+      <returns>The next symbol to be used.</returns>
+    </member>
+    <member name="M:System.Web.Razor.Tokenizer.Tokenizer`2.Peek">
+      <summary>This type/member supports the .NET Framework infrastructure and is not intended to be used directly from your code.Reads the next symbol in the code.</summary>
+      <returns>The next symbol to read.</returns>
+    </member>
+    <member name="M:System.Web.Razor.Tokenizer.Tokenizer`2.RazorCommentBody">
+      <summary>This type/member supports the .NET Framework infrastructure and is not intended to be used directly from your code.Parses the Razor comment body.</summary>
+      <returns>The object that represent the state of the result.</returns>
+    </member>
+    <member name="P:System.Web.Razor.Tokenizer.Tokenizer`2.RazorCommentStarType">
+      <summary>This type/member supports the .NET Framework infrastructure and is not intended to be used directly from your code.Gets the star type for the razor comment.</summary>
+      <returns>The star type for the razor comment.</returns>
+    </member>
+    <member name="P:System.Web.Razor.Tokenizer.Tokenizer`2.RazorCommentTransitionType">
+      <summary>This type/member supports the .NET Framework infrastructure and is not intended to be used directly from your code.Gets the transition type for the razor comment.</summary>
+      <returns>The transition type for the razor comment.</returns>
+    </member>
+    <member name="P:System.Web.Razor.Tokenizer.Tokenizer`2.RazorCommentType">
+      <summary>This type/member supports the .NET Framework infrastructure and is not intended to be used directly from your code.Gets the type of razor comment.</summary>
+      <returns>The type of razor comment.</returns>
+    </member>
+    <member name="M:System.Web.Razor.Tokenizer.Tokenizer`2.Reset">
+      <summary>This type/member supports the .NET Framework infrastructure and is not intended to be used directly from your code.Sets the tokenizer status to its initial state.</summary>
+    </member>
+    <member name="M:System.Web.Razor.Tokenizer.Tokenizer`2.ResumeSymbol(`0)">
+      <summary>This type/member supports the .NET Framework infrastructure and is not intended to be used directly from your code.Resumes using the previous language symbol type.</summary>
+      <param name="previous">The previous language symbol type.</param>
+    </member>
+    <member name="M:System.Web.Razor.Tokenizer.Tokenizer`2.Single(`1)">
+      <summary>This type/member supports the .NET Framework infrastructure and is not intended to be used directly from your code.Uses a single type of symbol.</summary>
+      <returns>A single type of symbol.</returns>
+      <param name="type">The type of symbol.</param>
+    </member>
+    <member name="P:System.Web.Razor.Tokenizer.Tokenizer`2.Source">
+      <summary>This type/member supports the .NET Framework infrastructure and is not intended to be used directly from your code.Gets the source of the text document.</summary>
+      <returns>The source of the source document.</returns>
+    </member>
+    <member name="M:System.Web.Razor.Tokenizer.Tokenizer`2.StartSymbol">
+      <summary>This type/member supports the .NET Framework infrastructure and is not intended to be used directly from your code.Returns the start symbol used in this class.</summary>
+    </member>
+    <member name="M:System.Web.Razor.Tokenizer.Tokenizer`2.System#Web#Razor#Tokenizer#ITokenizer#NextSymbol">
+      <summary>This type/member supports the .NET Framework infrastructure and is not intended to be used directly from your code.Returns the next language symbol type.</summary>
+      <returns>The next language symbol type.</returns>
+    </member>
+    <member name="M:System.Web.Razor.Tokenizer.Tokenizer`2.TakeAll(System.String,System.Boolean)">
+      <summary>This type/member supports the .NET Framework infrastructure and is not intended to be used directly from your code.Takes the string if found in the lookahead buffer into the tokenizer buffer.</summary>
+      <returns>true if the lookahead buffer contains the expected string; otherwise, false.</returns>
+      <param name="expected">The string to match.</param>
+      <param name="caseSensitive">true to indicate comparison is case sensitive; otherwise, false.</param>
+    </member>
+    <member name="M:System.Web.Razor.Tokenizer.Tokenizer`2.TakeCurrent">
+      <summary>This type/member supports the .NET Framework infrastructure and is not intended to be used directly from your code.Accepts the current character into the buffer.</summary>
+    </member>
+    <member name="M:System.Web.Razor.Tokenizer.Tokenizer`2.TakeString(System.String,System.Boolean)">
+      <summary>This type/member supports the .NET Framework infrastructure and is not intended to be used directly from your code.Accepts the given input string into the buffer.</summary>
+      <returns>true if the whole input string was accepted; false, if only a substring was accepted.</returns>
+      <param name="input">The input string.</param>
+      <param name="caseSensitive">true to indicate comparison is case sensitive; otherwise, false.</param>
+    </member>
+    <member name="M:System.Web.Razor.Tokenizer.Tokenizer`2.TakeUntil(System.Func{System.Char,System.Boolean})">
+      <summary>Parses the source document until the condition specified by predicate is met or end file is reached.</summary>
+      <returns>true if the predicate condition is met; false if end of file is reached.</returns>
+      <param name="predicate">The predicate that specifies the processing condition.</param>
+    </member>
+    <member name="T:System.Web.Razor.Tokenizer.TokenizerView`3">
+      <summary>This type/member supports the .NET Framework infrastructure and is not intended to be used directly from your code.Represents the specified parameters for the tokenizer view.</summary>
+      <typeparam name="TTokenizer">The type tokenizer.</typeparam>
+      <typeparam name="TSymbol">The type symbol.</typeparam>
+      <typeparam name="TSymbolType">The token symbol type.</typeparam>
+    </member>
+    <member name="M:System.Web.Razor.Tokenizer.TokenizerView`3.#ctor(`0)">
+      <summary>This type/member supports the .NET Framework infrastructure and is not intended to be used directly from your code.Initializes a new instance of the <see cref="T:System.Web.Razor.Tokenizer.TokenizerView`3" /> class.</summary>
+      <param name="tokenizer">The tokenizer view.</param>
+    </member>
+    <member name="P:System.Web.Razor.Tokenizer.TokenizerView`3.Current">
+      <summary>This type/member supports the .NET Framework infrastructure and is not intended to be used directly from your code.Gets the current view of the TSymbol.</summary>
+      <returns>The current view of the TSymbol.</returns>
+    </member>
+    <member name="P:System.Web.Razor.Tokenizer.TokenizerView`3.EndOfFile">
+      <summary>This type/member supports the .NET Framework infrastructure and is not intended to be used directly from your code.Gets a value that indicates whether the view can reach the end of a file.</summary>
+      <returns>true if the view can reach the end of a file; otherwise, false.</returns>
+    </member>
+    <member name="M:System.Web.Razor.Tokenizer.TokenizerView`3.Next">
+      <summary>This type/member supports the .NET Framework infrastructure and is not intended to be used directly from your code.Determines whether the tokenizer moves to the next view.</summary>
+      <returns>true if the tokenizer moves to the next view; false.</returns>
+    </member>
+    <member name="M:System.Web.Razor.Tokenizer.TokenizerView`3.PutBack(`1)">
+      <summary>This type/member supports the .NET Framework infrastructure and is not intended to be used directly from your code.Puts a specified symbol into the tokenizer view.</summary>
+      <param name="symbol">The symbol.</param>
+    </member>
+    <member name="P:System.Web.Razor.Tokenizer.TokenizerView`3.Source">
+      <summary>This type/member supports the .NET Framework infrastructure and is not intended to be used directly from your code.Gets the source of the text document for the tokenizer view.</summary>
+      <returns>The source of the text document for the tokenizer view.</returns>
+    </member>
+    <member name="P:System.Web.Razor.Tokenizer.TokenizerView`3.Tokenizer">
+      <summary>This type/member supports the .NET Framework infrastructure and is not intended to be used directly from your code.Gets the tokenizer to view the symbols for the razor.</summary>
+      <returns>The tokenizer to view the symbols for the razor.</returns>
+    </member>
+    <member name="T:System.Web.Razor.Tokenizer.VBHelpers">
+      <summary>This type/member supports the .NET Framework infrastructure and is not intended to be used directly from your code.Represents a set of characters as helpers in VB.</summary>
+    </member>
+    <member name="M:System.Web.Razor.Tokenizer.VBHelpers.IsDoubleQuote(System.Char)">
+      <summary>This type/member supports the .NET Framework infrastructure and is not intended to be used directly from your code.Indicates a value whether a specified character is enclosed in double quotation marks (").</summary>
+      <returns>true if the character is enclosed in double quotation marks ("); otherwise, false.</returns>
+      <param name="character">The character.</param>
+    </member>
+    <member name="M:System.Web.Razor.Tokenizer.VBHelpers.IsOctalDigit(System.Char)">
+      <summary>This type/member supports the .NET Framework infrastructure and is not intended to be used directly from your code.Indicates a value whether a character is in octal digit.</summary>
+      <returns>true if a character is in octal digit; otherwise, false.</returns>
+      <param name="character">The character.</param>
+    </member>
+    <member name="M:System.Web.Razor.Tokenizer.VBHelpers.IsSingleQuote(System.Char)">
+      <summary>This type/member supports the .NET Framework infrastructure and is not intended to be used directly from your code.Indicates a value whether a specified character is enclosed in a single quotation mark (').</summary>
+      <returns>true if the character is enclosed in a single quotation mark ('); otherwise, false.</returns>
+      <param name="character">The character.</param>
+    </member>
+    <member name="T:System.Web.Razor.Tokenizer.VBTokenizer">
+      <summary>This type/member supports the .NET Framework infrastructure and is not intended to be used directly from your code.Allows an application to break a VB symbol into tokens.</summary>
+    </member>
+    <member name="M:System.Web.Razor.Tokenizer.VBTokenizer.#ctor(System.Web.Razor.Text.ITextDocument)">
+      <summary>This type/member supports the .NET Framework infrastructure and is not intended to be used directly from your code.Initializes a new instance of the <see cref="T:System.Web.Razor.Tokenizer.VBTokenizer" /> class.</summary>
+      <param name="source">The source of text.</param>
+    </member>
+    <member name="M:System.Web.Razor.Tokenizer.VBTokenizer.CreateSymbol(System.Web.Razor.Text.SourceLocation,System.String,System.Web.Razor.Tokenizer.Symbols.VBSymbolType,System.Collections.Generic.IEnumerable{System.Web.Razor.Parser.SyntaxTree.RazorError})">
+      <summary>This type/member supports the .NET Framework infrastructure and is not intended to be used directly from your code.Creates a domain of symbols.</summary>
+      <returns>A domain of symbols.</returns>
+      <param name="start">The source location.</param>
+      <param name="content">The content value.</param>
+      <param name="type">The <see cref="T:System.Web.Razor.Tokenizer.Symbols.VBSymbolType" />.</param>
+      <param name="errors">The razor error.</param>
+    </member>
+    <member name="P:System.Web.Razor.Tokenizer.VBTokenizer.RazorCommentStarType">
+      <summary>This type/member supports the .NET Framework infrastructure and is not intended to be used directly from your code.Gets the VB symbol type.</summary>
+      <returns>The VB symbol type.</returns>
+    </member>
+    <member name="P:System.Web.Razor.Tokenizer.VBTokenizer.RazorCommentTransitionType">
+      <summary>This type/member supports the .NET Framework infrastructure and is not intended to be used directly from your code.Gets the transition style of the VB symbol.</summary>
+      <returns>The transition style of the VB symbol.</returns>
+    </member>
+    <member name="P:System.Web.Razor.Tokenizer.VBTokenizer.RazorCommentType">
+      <summary>This type/member supports the .NET Framework infrastructure and is not intended to be used directly from your code.Gets the razor type comment of the <see cref="T:System.Web.Razor.Tokenizer.Symbols.VBSymbolType" />.</summary>
+      <returns>The razor type comment of the <see cref="T:System.Web.Razor.Tokenizer.Symbols.VBSymbolType" />.</returns>
+    </member>
+    <member name="P:System.Web.Razor.Tokenizer.VBTokenizer.StartState">
+      <summary>This type/member supports the .NET Framework infrastructure and is not intended to be used directly from your code.Gets the start state of the machine.</summary>
+      <returns>The start state of the machine.</returns>
+    </member>
+    <member name="T:System.Web.Razor.Tokenizer.Symbols.CSharpKeyword">
+      <summary>This type/member supports the .NET Framework infrastructure and is not intended to be used directly from your code.</summary>
+    </member>
+    <member name="F:System.Web.Razor.Tokenizer.Symbols.CSharpKeyword.Abstract">
+      <summary>This type/member supports the .NET Framework infrastructure and is not intended to be used directly from your code.</summary>
+    </member>
+    <member name="F:System.Web.Razor.Tokenizer.Symbols.CSharpKeyword.As">
+      <summary>This type/member supports the .NET Framework infrastructure and is not intended to be used directly from your code.</summary>
+    </member>
+    <member name="F:System.Web.Razor.Tokenizer.Symbols.CSharpKeyword.Base">
+      <summary>This type/member supports the .NET Framework infrastructure and is not intended to be used directly from your code.</summary>
+    </member>
+    <member name="F:System.Web.Razor.Tokenizer.Symbols.CSharpKeyword.Bool">
+      <summary>This type/member supports the .NET Framework infrastructure and is not intended to be used directly from your code.</summary>
+    </member>
+    <member name="F:System.Web.Razor.Tokenizer.Symbols.CSharpKeyword.Break">
+      <summary>This type/member supports the .NET Framework infrastructure and is not intended to be used directly from your code.</summary>
+    </member>
+    <member name="F:System.Web.Razor.Tokenizer.Symbols.CSharpKeyword.Byte">
+      <summary>This type/member supports the .NET Framework infrastructure and is not intended to be used directly from your code.</summary>
+    </member>
+    <member name="F:System.Web.Razor.Tokenizer.Symbols.CSharpKeyword.Case">
+      <summary>This type/member supports the .NET Framework infrastructure and is not intended to be used directly from your code.</summary>
+    </member>
+    <member name="F:System.Web.Razor.Tokenizer.Symbols.CSharpKeyword.Catch">
+      <summary>This type/member supports the .NET Framework infrastructure and is not intended to be used directly from your code.</summary>
+    </member>
+    <member name="F:System.Web.Razor.Tokenizer.Symbols.CSharpKeyword.Char">
+      <summary>This type/member supports the .NET Framework infrastructure and is not intended to be used directly from your code.</summary>
+    </member>
+    <member name="F:System.Web.Razor.Tokenizer.Symbols.CSharpKeyword.Checked">
+      <summary>This type/member supports the .NET Framework infrastructure and is not intended to be used directly from your code.</summary>
+    </member>
+    <member name="F:System.Web.Razor.Tokenizer.Symbols.CSharpKeyword.Class">
+      <summary>This type/member supports the .NET Framework infrastructure and is not intended to be used directly from your code.</summary>
+    </member>
+    <member name="F:System.Web.Razor.Tokenizer.Symbols.CSharpKeyword.Const">
+      <summary>This type/member supports the .NET Framework infrastructure and is not intended to be used directly from your code.</summary>
+    </member>
+    <member name="F:System.Web.Razor.Tokenizer.Symbols.CSharpKeyword.Continue">
+      <summary>This type/member supports the .NET Framework infrastructure and is not intended to be used directly from your code.</summary>
+    </member>
+    <member name="F:System.Web.Razor.Tokenizer.Symbols.CSharpKeyword.Decimal">
+      <summary>This type/member supports the .NET Framework infrastructure and is not intended to be used directly from your code.</summary>
+    </member>
+    <member name="F:System.Web.Razor.Tokenizer.Symbols.CSharpKeyword.Default">
+      <summary>This type/member supports the .NET Framework infrastructure and is not intended to be used directly from your code.</summary>
+    </member>
+    <member name="F:System.Web.Razor.Tokenizer.Symbols.CSharpKeyword.Delegate">
+      <summary>This type/member supports the .NET Framework infrastructure and is not intended to be used directly from your code.</summary>
+    </member>
+    <member name="F:System.Web.Razor.Tokenizer.Symbols.CSharpKeyword.Do">
+      <summary>This type/member supports the .NET Framework infrastructure and is not intended to be used directly from your code.</summary>
+    </member>
+    <member name="F:System.Web.Razor.Tokenizer.Symbols.CSharpKeyword.Double">
+      <summary>This type/member supports the .NET Framework infrastructure and is not intended to be used directly from your code.</summary>
+    </member>
+    <member name="F:System.Web.Razor.Tokenizer.Symbols.CSharpKeyword.Else">
+      <summary>This type/member supports the .NET Framework infrastructure and is not intended to be used directly from your code.</summary>
+    </member>
+    <member name="F:System.Web.Razor.Tokenizer.Symbols.CSharpKeyword.Enum">
+      <summary>This type/member supports the .NET Framework infrastructure and is not intended to be used directly from your code.</summary>
+    </member>
+    <member name="F:System.Web.Razor.Tokenizer.Symbols.CSharpKeyword.Event">
+      <summary>This type/member supports the .NET Framework infrastructure and is not intended to be used directly from your code.</summary>
+    </member>
+    <member name="F:System.Web.Razor.Tokenizer.Symbols.CSharpKeyword.Explicit">
+      <summary>This type/member supports the .NET Framework infrastructure and is not intended to be used directly from your code.</summary>
+    </member>
+    <member name="F:System.Web.Razor.Tokenizer.Symbols.CSharpKeyword.Extern">
+      <summary>This type/member supports the .NET Framework infrastructure and is not intended to be used directly from your code.</summary>
+    </member>
+    <member name="F:System.Web.Razor.Tokenizer.Symbols.CSharpKeyword.False">
+      <summary>This type/member supports the .NET Framework infrastructure and is not intended to be used directly from your code.</summary>
+    </member>
+    <member name="F:System.Web.Razor.Tokenizer.Symbols.CSharpKeyword.Finally">
+      <summary>This type/member supports the .NET Framework infrastructure and is not intended to be used directly from your code.</summary>
+    </member>
+    <member name="F:System.Web.Razor.Tokenizer.Symbols.CSharpKeyword.Fixed">
+      <summary>This type/member supports the .NET Framework infrastructure and is not intended to be used directly from your code.</summary>
+    </member>
+    <member name="F:System.Web.Razor.Tokenizer.Symbols.CSharpKeyword.Float">
+      <summary>This type/member supports the .NET Framework infrastructure and is not intended to be used directly from your code.</summary>
+    </member>
+    <member name="F:System.Web.Razor.Tokenizer.Symbols.CSharpKeyword.For">
+      <summary>This type/member supports the .NET Framework infrastructure and is not intended to be used directly from your code.</summary>
+    </member>
+    <member name="F:System.Web.Razor.Tokenizer.Symbols.CSharpKeyword.Foreach">
+      <summary>This type/member supports the .NET Framework infrastructure and is not intended to be used directly from your code.</summary>
+    </member>
+    <member name="F:System.Web.Razor.Tokenizer.Symbols.CSharpKeyword.Goto">
+      <summary>This type/member supports the .NET Framework infrastructure and is not intended to be used directly from your code.</summary>
+    </member>
+    <member name="F:System.Web.Razor.Tokenizer.Symbols.CSharpKeyword.If">
+      <summary>This type/member supports the .NET Framework infrastructure and is not intended to be used directly from your code.</summary>
+    </member>
+    <member name="F:System.Web.Razor.Tokenizer.Symbols.CSharpKeyword.Implicit">
+      <summary>This type/member supports the .NET Framework infrastructure and is not intended to be used directly from your code.</summary>
+    </member>
+    <member name="F:System.Web.Razor.Tokenizer.Symbols.CSharpKeyword.In">
+      <summary>This type/member supports the .NET Framework infrastructure and is not intended to be used directly from your code.</summary>
+    </member>
+    <member name="F:System.Web.Razor.Tokenizer.Symbols.CSharpKeyword.Int">
+      <summary>This type/member supports the .NET Framework infrastructure and is not intended to be used directly from your code.</summary>
+    </member>
+    <member name="F:System.Web.Razor.Tokenizer.Symbols.CSharpKeyword.Interface">
+      <summary>This type/member supports the .NET Framework infrastructure and is not intended to be used directly from your code.</summary>
+    </member>
+    <member name="F:System.Web.Razor.Tokenizer.Symbols.CSharpKeyword.Internal">
+      <summary>This type/member supports the .NET Framework infrastructure and is not intended to be used directly from your code.</summary>
+    </member>
+    <member name="F:System.Web.Razor.Tokenizer.Symbols.CSharpKeyword.Is">
+      <summary>This type/member supports the .NET Framework infrastructure and is not intended to be used directly from your code.</summary>
+    </member>
+    <member name="F:System.Web.Razor.Tokenizer.Symbols.CSharpKeyword.Lock">
+      <summary>This type/member supports the .NET Framework infrastructure and is not intended to be used directly from your code.</summary>
+    </member>
+    <member name="F:System.Web.Razor.Tokenizer.Symbols.CSharpKeyword.Long">
+      <summary>This type/member supports the .NET Framework infrastructure and is not intended to be used directly from your code.</summary>
+    </member>
+    <member name="F:System.Web.Razor.Tokenizer.Symbols.CSharpKeyword.Namespace">
+      <summary>This type/member supports the .NET Framework infrastructure and is not intended to be used directly from your code.</summary>
+    </member>
+    <member name="F:System.Web.Razor.Tokenizer.Symbols.CSharpKeyword.New">
+      <summary>This type/member supports the .NET Framework infrastructure and is not intended to be used directly from your code.</summary>
+    </member>
+    <member name="F:System.Web.Razor.Tokenizer.Symbols.CSharpKeyword.Null">
+      <summary>This type/member supports the .NET Framework infrastructure and is not intended to be used directly from your code.</summary>
+    </member>
+    <member name="F:System.Web.Razor.Tokenizer.Symbols.CSharpKeyword.Object">
+      <summary>This type/member supports the .NET Framework infrastructure and is not intended to be used directly from your code.</summary>
+    </member>
+    <member name="F:System.Web.Razor.Tokenizer.Symbols.CSharpKeyword.Operator">
+      <summary>This type/member supports the .NET Framework infrastructure and is not intended to be used directly from your code.</summary>
+    </member>
+    <member name="F:System.Web.Razor.Tokenizer.Symbols.CSharpKeyword.Out">
+      <summary>This type/member supports the .NET Framework infrastructure and is not intended to be used directly from your code.</summary>
+    </member>
+    <member name="F:System.Web.Razor.Tokenizer.Symbols.CSharpKeyword.Override">
+      <summary>This type/member supports the .NET Framework infrastructure and is not intended to be used directly from your code.</summary>
+    </member>
+    <member name="F:System.Web.Razor.Tokenizer.Symbols.CSharpKeyword.Params">
+      <summary>This type/member supports the .NET Framework infrastructure and is not intended to be used directly from your code.</summary>
+    </member>
+    <member name="F:System.Web.Razor.Tokenizer.Symbols.CSharpKeyword.Private">
+      <summary>This type/member supports the .NET Framework infrastructure and is not intended to be used directly from your code.</summary>
+    </member>
+    <member name="F:System.Web.Razor.Tokenizer.Symbols.CSharpKeyword.Protected">
+      <summary>This type/member supports the .NET Framework infrastructure and is not intended to be used directly from your code.</summary>
+    </member>
+    <member name="F:System.Web.Razor.Tokenizer.Symbols.CSharpKeyword.Public">
+      <summary>This type/member supports the .NET Framework infrastructure and is not intended to be used directly from your code.</summary>
+    </member>
+    <member name="F:System.Web.Razor.Tokenizer.Symbols.CSharpKeyword.Readonly">
+      <summary>This type/member supports the .NET Framework infrastructure and is not intended to be used directly from your code.</summary>
+    </member>
+    <member name="F:System.Web.Razor.Tokenizer.Symbols.CSharpKeyword.Ref">
+      <summary>This type/member supports the .NET Framework infrastructure and is not intended to be used directly from your code.</summary>
+    </member>
+    <member name="F:System.Web.Razor.Tokenizer.Symbols.CSharpKeyword.Return">
+      <summary>This type/member supports the .NET Framework infrastructure and is not intended to be used directly from your code.</summary>
+    </member>
+    <member name="F:System.Web.Razor.Tokenizer.Symbols.CSharpKeyword.Sbyte">
+      <summary>This type/member supports the .NET Framework infrastructure and is not intended to be used directly from your code.</summary>
+    </member>
+    <member name="F:System.Web.Razor.Tokenizer.Symbols.CSharpKeyword.Sealed">
+      <summary>This type/member supports the .NET Framework infrastructure and is not intended to be used directly from your code.</summary>
+    </member>
+    <member name="F:System.Web.Razor.Tokenizer.Symbols.CSharpKeyword.Short">
+      <summary>This type/member supports the .NET Framework infrastructure and is not intended to be used directly from your code.</summary>
+    </member>
+    <member name="F:System.Web.Razor.Tokenizer.Symbols.CSharpKeyword.Sizeof">
+      <summary>This type/member supports the .NET Framework infrastructure and is not intended to be used directly from your code.</summary>
+    </member>
+    <member name="F:System.Web.Razor.Tokenizer.Symbols.CSharpKeyword.Stackalloc">
+      <summary>This type/member supports the .NET Framework infrastructure and is not intended to be used directly from your code.</summary>
+    </member>
+    <member name="F:System.Web.Razor.Tokenizer.Symbols.CSharpKeyword.Static">
+      <summary>This type/member supports the .NET Framework infrastructure and is not intended to be used directly from your code.</summary>
+    </member>
+    <member name="F:System.Web.Razor.Tokenizer.Symbols.CSharpKeyword.String">
+      <summary>This type/member supports the .NET Framework infrastructure and is not intended to be used directly from your code.</summary>
+    </member>
+    <member name="F:System.Web.Razor.Tokenizer.Symbols.CSharpKeyword.Struct">
+      <summary>This type/member supports the .NET Framework infrastructure and is not intended to be used directly from your code.</summary>
+    </member>
+    <member name="F:System.Web.Razor.Tokenizer.Symbols.CSharpKeyword.Switch">
+      <summary>This type/member supports the .NET Framework infrastructure and is not intended to be used directly from your code.</summary>
+    </member>
+    <member name="F:System.Web.Razor.Tokenizer.Symbols.CSharpKeyword.This">
+      <summary>This type/member supports the .NET Framework infrastructure and is not intended to be used directly from your code.</summary>
+    </member>
+    <member name="F:System.Web.Razor.Tokenizer.Symbols.CSharpKeyword.Throw">
+      <summary>This type/member supports the .NET Framework infrastructure and is not intended to be used directly from your code.</summary>
+    </member>
+    <member name="F:System.Web.Razor.Tokenizer.Symbols.CSharpKeyword.True">
+      <summary>This type/member supports the .NET Framework infrastructure and is not intended to be used directly from your code.</summary>
+    </member>
+    <member name="F:System.Web.Razor.Tokenizer.Symbols.CSharpKeyword.Try">
+      <summary>This type/member supports the .NET Framework infrastructure and is not intended to be used directly from your code.</summary>
+    </member>
+    <member name="F:System.Web.Razor.Tokenizer.Symbols.CSharpKeyword.Typeof">
+      <summary>This type/member supports the .NET Framework infrastructure and is not intended to be used directly from your code.</summary>
+    </member>
+    <member name="F:System.Web.Razor.Tokenizer.Symbols.CSharpKeyword.Uint">
+      <summary>This type/member supports the .NET Framework infrastructure and is not intended to be used directly from your code.</summary>
+    </member>
+    <member name="F:System.Web.Razor.Tokenizer.Symbols.CSharpKeyword.Ulong">
+      <summary>This type/member supports the .NET Framework infrastructure and is not intended to be used directly from your code.</summary>
+    </member>
+    <member name="F:System.Web.Razor.Tokenizer.Symbols.CSharpKeyword.Unchecked">
+      <summary>This type/member supports the .NET Framework infrastructure and is not intended to be used directly from your code.</summary>
+    </member>
+    <member name="F:System.Web.Razor.Tokenizer.Symbols.CSharpKeyword.Unsafe">
+      <summary>This type/member supports the .NET Framework infrastructure and is not intended to be used directly from your code.</summary>
+    </member>
+    <member name="F:System.Web.Razor.Tokenizer.Symbols.CSharpKeyword.Ushort">
+      <summary>This type/member supports the .NET Framework infrastructure and is not intended to be used directly from your code.</summary>
+    </member>
+    <member name="F:System.Web.Razor.Tokenizer.Symbols.CSharpKeyword.Using">
+      <summary>This type/member supports the .NET Framework infrastructure and is not intended to be used directly from your code.</summary>
+    </member>
+    <member name="F:System.Web.Razor.Tokenizer.Symbols.CSharpKeyword.Virtual">
+      <summary>This type/member supports the .NET Framework infrastructure and is not intended to be used directly from your code.</summary>
+    </member>
+    <member name="F:System.Web.Razor.Tokenizer.Symbols.CSharpKeyword.Void">
+      <summary>This type/member supports the .NET Framework infrastructure and is not intended to be used directly from your code.</summary>
+    </member>
+    <member name="F:System.Web.Razor.Tokenizer.Symbols.CSharpKeyword.Volatile">
+      <summary>This type/member supports the .NET Framework infrastructure and is not intended to be used directly from your code.</summary>
+    </member>
+    <member name="F:System.Web.Razor.Tokenizer.Symbols.CSharpKeyword.While">
+      <summary>This type/member supports the .NET Framework infrastructure and is not intended to be used directly from your code.</summary>
+    </member>
+    <member name="T:System.Web.Razor.Tokenizer.Symbols.CSharpSymbol">
+      <summary>This type/member supports the .NET Framework infrastructure and is not intended to be used directly from your code.Represents a C sharp symbol for the razor tokenizer.</summary>
+    </member>
+    <member name="M:System.Web.Razor.Tokenizer.Symbols.CSharpSymbol.#ctor(System.Int32,System.Int32,System.Int32,System.String,System.Web.Razor.Tokenizer.Symbols.CSharpSymbolType)">
+      <summary>This type/member supports the .NET Framework infrastructure and is not intended to be used directly from your code.Initializes a new instance of the <see cref="T:System.Web.Razor.Tokenizer.Symbols.CSharpSymbol" /> class.</summary>
+      <param name="offset">The symbol’s offset.</param>
+      <param name="line">The line.</param>
+      <param name="column">The column</param>
+      <param name="content">The content of the symbol.</param>
+      <param name="type">The type of the symbol.</param>
+    </member>
+    <member name="M:System.Web.Razor.Tokenizer.Symbols.CSharpSymbol.#ctor(System.Int32,System.Int32,System.Int32,System.String,System.Web.Razor.Tokenizer.Symbols.CSharpSymbolType,System.Collections.Generic.IEnumerable{System.Web.Razor.Parser.SyntaxTree.RazorError})">
+      <summary>This type/member supports the .NET Framework infrastructure and is not intended to be used directly from your code.Initializes a new instance of the <see cref="T:System.Web.Razor.Tokenizer.Symbols.CSharpSymbol" /> class.</summary>
+      <param name="offset">The symbol’s offset.</param>
+      <param name="line">The line.</param>
+      <param name="column">The column</param>
+      <param name="content">The content of the symbol.</param>
+      <param name="type">The type of the symbol.</param>
+      <param name="errors">A list of errors.</param>
+    </member>
+    <member name="M:System.Web.Razor.Tokenizer.Symbols.CSharpSymbol.#ctor(System.Web.Razor.Text.SourceLocation,System.String,System.Web.Razor.Tokenizer.Symbols.CSharpSymbolType)">
+      <summary>This type/member supports the .NET Framework infrastructure and is not intended to be used directly from your code.Initializes a new instance of the <see cref="T:System.Web.Razor.Tokenizer.Symbols.CSharpSymbol" /> class.</summary>
+      <param name="start">The location to start the symbol.</param>
+      <param name="content">The content of the symbol.</param>
+      <param name="type">The type of the symbol.</param>
+    </member>
+    <member name="M:System.Web.Razor.Tokenizer.Symbols.CSharpSymbol.#ctor(System.Web.Razor.Text.SourceLocation,System.String,System.Web.Razor.Tokenizer.Symbols.CSharpSymbolType,System.Collections.Generic.IEnumerable{System.Web.Razor.Parser.SyntaxTree.RazorError})">
+      <summary>This type/member supports the .NET Framework infrastructure and is not intended to be used directly from your code.Initializes a new instance of the <see cref="T:System.Web.Razor.Tokenizer.Symbols.CSharpSymbol" /> class.</summary>
+      <param name="start">The location to start the symbol.</param>
+      <param name="content">The content of the symbol.</param>
+      <param name="type">The type of the symbol.</param>
+      <param name="errors">A list of errors.</param>
+    </member>
+    <member name="M:System.Web.Razor.Tokenizer.Symbols.CSharpSymbol.Equals(System.Object)">
+      <summary>This type/member supports the .NET Framework infrastructure and is not intended to be used directly from your code.Determines whether the specified object is equal to the current object.</summary>
+      <returns>true if the specified object is equal to the current object; otherwise, false.</returns>
+      <param name="obj">The object to compare to.</param>
+    </member>
+    <member name="P:System.Web.Razor.Tokenizer.Symbols.CSharpSymbol.EscapedIdentifier">
+      <summary>This type/member supports the .NET Framework infrastructure and is not intended to be used directly from your code.Gets or sets a value that indicates whether the symbol has an escaped identifier.</summary>
+      <returns>true if the symbol has an escaped identifier; otherwise, false.</returns>
+    </member>
+    <member name="M:System.Web.Razor.Tokenizer.Symbols.CSharpSymbol.GetHashCode">
+      <summary>This type/member supports the .NET Framework infrastructure and is not intended to be used directly from your code.Returns the hash code for this current instance.</summary>
+      <returns>The hash code for this current instance.</returns>
+    </member>
+    <member name="P:System.Web.Razor.Tokenizer.Symbols.CSharpSymbol.Keyword">
+      <summary>This type/member supports the .NET Framework infrastructure and is not intended to be used directly from your code.Gets or sets the language keyword.</summary>
+      <returns>The language keyword.</returns>
+    </member>
+    <member name="T:System.Web.Razor.Tokenizer.Symbols.CSharpSymbolType">
+      <summary>This type/member supports the .NET Framework infrastructure and is not intended to be used directly from your code.</summary>
+    </member>
+    <member name="F:System.Web.Razor.Tokenizer.Symbols.CSharpSymbolType.And">
+      <summary>This type/member supports the .NET Framework infrastructure and is not intended to be used directly from your code.</summary>
+    </member>
+    <member name="F:System.Web.Razor.Tokenizer.Symbols.CSharpSymbolType.AndAssign">
+      <summary>This type/member supports the .NET Framework infrastructure and is not intended to be used directly from your code.</summary>
+    </member>
+    <member name="F:System.Web.Razor.Tokenizer.Symbols.CSharpSymbolType.Arrow">
+      <summary>This type/member supports the .NET Framework infrastructure and is not intended to be used directly from your code.</summary>
+    </member>
+    <member name="F:System.Web.Razor.Tokenizer.Symbols.CSharpSymbolType.Assign">
+      <summary>This type/member supports the .NET Framework infrastructure and is not intended to be used directly from your code.</summary>
+    </member>
+    <member name="F:System.Web.Razor.Tokenizer.Symbols.CSharpSymbolType.CharacterLiteral">
+      <summary>This type/member supports the .NET Framework infrastructure and is not intended to be used directly from your code.</summary>
+    </member>
+    <member name="F:System.Web.Razor.Tokenizer.Symbols.CSharpSymbolType.Colon">
+      <summary>This type/member supports the .NET Framework infrastructure and is not intended to be used directly from your code.</summary>
+    </member>
+    <member name="F:System.Web.Razor.Tokenizer.Symbols.CSharpSymbolType.Comma">
+      <summary>This type/member supports the .NET Framework infrastructure and is not intended to be used directly from your code.</summary>
+    </member>
+    <member name="F:System.Web.Razor.Tokenizer.Symbols.CSharpSymbolType.Comment">
+      <summary>This type/member supports the .NET Framework infrastructure and is not intended to be used directly from your code.</summary>
+    </member>
+    <member name="F:System.Web.Razor.Tokenizer.Symbols.CSharpSymbolType.Decrement">
+      <summary>This type/member supports the .NET Framework infrastructure and is not intended to be used directly from your code.</summary>
+    </member>
+    <member name="F:System.Web.Razor.Tokenizer.Symbols.CSharpSymbolType.DivideAssign">
+      <summary>This type/member supports the .NET Framework infrastructure and is not intended to be used directly from your code.</summary>
+    </member>
+    <member name="F:System.Web.Razor.Tokenizer.Symbols.CSharpSymbolType.Dot">
+      <summary>This type/member supports the .NET Framework infrastructure and is not intended to be used directly from your code.</summary>
+    </member>
+    <member name="F:System.Web.Razor.Tokenizer.Symbols.CSharpSymbolType.DoubleAnd">
+      <summary>This type/member supports the .NET Framework infrastructure and is not intended to be used directly from your code.</summary>
+    </member>
+    <member name="F:System.Web.Razor.Tokenizer.Symbols.CSharpSymbolType.DoubleColon">
+      <summary>This type/member supports the .NET Framework infrastructure and is not intended to be used directly from your code.</summary>
+    </member>
+    <member name="F:System.Web.Razor.Tokenizer.Symbols.CSharpSymbolType.DoubleOr">
+      <summary>This type/member supports the .NET Framework infrastructure and is not intended to be used directly from your code.</summary>
+    </member>
+    <member name="F:System.Web.Razor.Tokenizer.Symbols.CSharpSymbolType.Equals">
+      <summary>This type/member supports the .NET Framework infrastructure and is not intended to be used directly from your code.</summary>
+    </member>
+    <member name="F:System.Web.Razor.Tokenizer.Symbols.CSharpSymbolType.GreaterThan">
+      <summary>This type/member supports the .NET Framework infrastructure and is not intended to be used directly from your code.</summary>
+    </member>
+    <member name="F:System.Web.Razor.Tokenizer.Symbols.CSharpSymbolType.GreaterThanEqual">
+      <summary>This type/member supports the .NET Framework infrastructure and is not intended to be used directly from your code.</summary>
+    </member>
+    <member name="F:System.Web.Razor.Tokenizer.Symbols.CSharpSymbolType.Hash">
+      <summary>This type/member supports the .NET Framework infrastructure and is not intended to be used directly from your code.</summary>
+    </member>
+    <member name="F:System.Web.Razor.Tokenizer.Symbols.CSharpSymbolType.Identifier">
+      <summary>This type/member supports the .NET Framework infrastructure and is not intended to be used directly from your code.</summary>
+    </member>
+    <member name="F:System.Web.Razor.Tokenizer.Symbols.CSharpSymbolType.Increment">
+      <summary>This type/member supports the .NET Framework infrastructure and is not intended to be used directly from your code.</summary>
+    </member>
+    <member name="F:System.Web.Razor.Tokenizer.Symbols.CSharpSymbolType.IntegerLiteral">
+      <summary>This type/member supports the .NET Framework infrastructure and is not intended to be used directly from your code.</summary>
+    </member>
+    <member name="F:System.Web.Razor.Tokenizer.Symbols.CSharpSymbolType.Keyword">
+      <summary>This type/member supports the .NET Framework infrastructure and is not intended to be used directly from your code.</summary>
+    </member>
+    <member name="F:System.Web.Razor.Tokenizer.Symbols.CSharpSymbolType.LeftBrace">
+      <summary>This type/member supports the .NET Framework infrastructure and is not intended to be used directly from your code.</summary>
+    </member>
+    <member name="F:System.Web.Razor.Tokenizer.Symbols.CSharpSymbolType.LeftBracket">
+      <summary>This type/member supports the .NET Framework infrastructure and is not intended to be used directly from your code.</summary>
+    </member>
+    <member name="F:System.Web.Razor.Tokenizer.Symbols.CSharpSymbolType.LeftParenthesis">
+      <summary>This type/member supports the .NET Framework infrastructure and is not intended to be used directly from your code.</summary>
+    </member>
+    <member name="F:System.Web.Razor.Tokenizer.Symbols.CSharpSymbolType.LeftShift">
+      <summary>This type/member supports the .NET Framework infrastructure and is not intended to be used directly from your code.</summary>
+    </member>
+    <member name="F:System.Web.Razor.Tokenizer.Symbols.CSharpSymbolType.LeftShiftAssign">
+      <summary>This type/member supports the .NET Framework infrastructure and is not intended to be used directly from your code.</summary>
+    </member>
+    <member name="F:System.Web.Razor.Tokenizer.Symbols.CSharpSymbolType.LessThan">
+      <summary>This type/member supports the .NET Framework infrastructure and is not intended to be used directly from your code.</summary>
+    </member>
+    <member name="F:System.Web.Razor.Tokenizer.Symbols.CSharpSymbolType.LessThanEqual">
+      <summary>This type/member supports the .NET Framework infrastructure and is not intended to be used directly from your code.</summary>
+    </member>
+    <member name="F:System.Web.Razor.Tokenizer.Symbols.CSharpSymbolType.Minus">
+      <summary>This type/member supports the .NET Framework infrastructure and is not intended to be used directly from your code.</summary>
+    </member>
+    <member name="F:System.Web.Razor.Tokenizer.Symbols.CSharpSymbolType.MinusAssign">
+      <summary>This type/member supports the .NET Framework infrastructure and is not intended to be used directly from your code.</summary>
+    </member>
+    <member name="F:System.Web.Razor.Tokenizer.Symbols.CSharpSymbolType.Modulo">
+      <summary>This type/member supports the .NET Framework infrastructure and is not intended to be used directly from your code.</summary>
+    </member>
+    <member name="F:System.Web.Razor.Tokenizer.Symbols.CSharpSymbolType.ModuloAssign">
+      <summary>This type/member supports the .NET Framework infrastructure and is not intended to be used directly from your code.</summary>
+    </member>
+    <member name="F:System.Web.Razor.Tokenizer.Symbols.CSharpSymbolType.MultiplyAssign">
+      <summary>This type/member supports the .NET Framework infrastructure and is not intended to be used directly from your code.</summary>
+    </member>
+    <member name="F:System.Web.Razor.Tokenizer.Symbols.CSharpSymbolType.NewLine">
+      <summary>This type/member supports the .NET Framework infrastructure and is not intended to be used directly from your code.</summary>
+    </member>
+    <member name="F:System.Web.Razor.Tokenizer.Symbols.CSharpSymbolType.Not">
+      <summary>This type/member supports the .NET Framework infrastructure and is not intended to be used directly from your code.</summary>
+    </member>
+    <member name="F:System.Web.Razor.Tokenizer.Symbols.CSharpSymbolType.NotEqual">
+      <summary>This type/member supports the .NET Framework infrastructure and is not intended to be used directly from your code.</summary>
+    </member>
+    <member name="F:System.Web.Razor.Tokenizer.Symbols.CSharpSymbolType.NullCoalesce">
+      <summary>This type/member supports the .NET Framework infrastructure and is not intended to be used directly from your code.</summary>
+    </member>
+    <member name="F:System.Web.Razor.Tokenizer.Symbols.CSharpSymbolType.Or">
+      <summary>This type/member supports the .NET Framework infrastructure and is not intended to be used directly from your code.</summary>
+    </member>
+    <member name="F:System.Web.Razor.Tokenizer.Symbols.CSharpSymbolType.OrAssign">
+      <summary>This type/member supports the .NET Framework infrastructure and is not intended to be used directly from your code.</summary>
+    </member>
+    <member name="F:System.Web.Razor.Tokenizer.Symbols.CSharpSymbolType.Plus">
+      <summary>This type/member supports the .NET Framework infrastructure and is not intended to be used directly from your code.</summary>
+    </member>
+    <member name="F:System.Web.Razor.Tokenizer.Symbols.CSharpSymbolType.PlusAssign">
+      <summary>This type/member supports the .NET Framework infrastructure and is not intended to be used directly from your code.</summary>
+    </member>
+    <member name="F:System.Web.Razor.Tokenizer.Symbols.CSharpSymbolType.QuestionMark">
+      <summary>This type/member supports the .NET Framework infrastructure and is not intended to be used directly from your code.</summary>
+    </member>
+    <member name="F:System.Web.Razor.Tokenizer.Symbols.CSharpSymbolType.RazorComment">
+      <summary>This type/member supports the .NET Framework infrastructure and is not intended to be used directly from your code.</summary>
+    </member>
+    <member name="F:System.Web.Razor.Tokenizer.Symbols.CSharpSymbolType.RazorCommentStar">
+      <summary>This type/member supports the .NET Framework infrastructure and is not intended to be used directly from your code.</summary>
+    </member>
+    <member name="F:System.Web.Razor.Tokenizer.Symbols.CSharpSymbolType.RazorCommentTransition">
+      <summary>This type/member supports the .NET Framework infrastructure and is not intended to be used directly from your code.</summary>
+    </member>
+    <member name="F:System.Web.Razor.Tokenizer.Symbols.CSharpSymbolType.RealLiteral">
+      <summary>This type/member supports the .NET Framework infrastructure and is not intended to be used directly from your code.</summary>
+    </member>
+    <member name="F:System.Web.Razor.Tokenizer.Symbols.CSharpSymbolType.RightBrace">
+      <summary>This type/member supports the .NET Framework infrastructure and is not intended to be used directly from your code.</summary>
+    </member>
+    <member name="F:System.Web.Razor.Tokenizer.Symbols.CSharpSymbolType.RightBracket">
+      <summary>This type/member supports the .NET Framework infrastructure and is not intended to be used directly from your code.</summary>
+    </member>
+    <member name="F:System.Web.Razor.Tokenizer.Symbols.CSharpSymbolType.RightParenthesis">
+      <summary>This type/member supports the .NET Framework infrastructure and is not intended to be used directly from your code.</summary>
+    </member>
+    <member name="F:System.Web.Razor.Tokenizer.Symbols.CSharpSymbolType.RightShift">
+      <summary>This type/member supports the .NET Framework infrastructure and is not intended to be used directly from your code.</summary>
+    </member>
+    <member name="F:System.Web.Razor.Tokenizer.Symbols.CSharpSymbolType.RightShiftAssign">
+      <summary>This type/member supports the .NET Framework infrastructure and is not intended to be used directly from your code.</summary>
+    </member>
+    <member name="F:System.Web.Razor.Tokenizer.Symbols.CSharpSymbolType.Semicolon">
+      <summary>This type/member supports the .NET Framework infrastructure and is not intended to be used directly from your code.</summary>
+    </member>
+    <member name="F:System.Web.Razor.Tokenizer.Symbols.CSharpSymbolType.Slash">
+      <summary>This type/member supports the .NET Framework infrastructure and is not intended to be used directly from your code.</summary>
+    </member>
+    <member name="F:System.Web.Razor.Tokenizer.Symbols.CSharpSymbolType.Star">
+      <summary>This type/member supports the .NET Framework infrastructure and is not intended to be used directly from your code.</summary>
+    </member>
+    <member name="F:System.Web.Razor.Tokenizer.Symbols.CSharpSymbolType.StringLiteral">
+      <summary>This type/member supports the .NET Framework infrastructure and is not intended to be used directly from your code.</summary>
+    </member>
+    <member name="F:System.Web.Razor.Tokenizer.Symbols.CSharpSymbolType.Tilde">
+      <summary>This type/member supports the .NET Framework infrastructure and is not intended to be used directly from your code.</summary>
+    </member>
+    <member name="F:System.Web.Razor.Tokenizer.Symbols.CSharpSymbolType.Transition">
+      <summary>This type/member supports the .NET Framework infrastructure and is not intended to be used directly from your code.</summary>
+    </member>
+    <member name="F:System.Web.Razor.Tokenizer.Symbols.CSharpSymbolType.Unknown">
+      <summary>This type/member supports the .NET Framework infrastructure and is not intended to be used directly from your code.</summary>
+    </member>
+    <member name="F:System.Web.Razor.Tokenizer.Symbols.CSharpSymbolType.WhiteSpace">
+      <summary>This type/member supports the .NET Framework infrastructure and is not intended to be used directly from your code.</summary>
+    </member>
+    <member name="F:System.Web.Razor.Tokenizer.Symbols.CSharpSymbolType.Xor">
+      <summary>This type/member supports the .NET Framework infrastructure and is not intended to be used directly from your code.</summary>
+    </member>
+    <member name="F:System.Web.Razor.Tokenizer.Symbols.CSharpSymbolType.XorAssign">
+      <summary>This type/member supports the .NET Framework infrastructure and is not intended to be used directly from your code.</summary>
+    </member>
+    <member name="T:System.Web.Razor.Tokenizer.Symbols.HtmlSymbol">
+      <summary>This type/member supports the .NET Framework infrastructure and is not intended to be used directly from your code.Represents the Html symbols.</summary>
+    </member>
+    <member name="M:System.Web.Razor.Tokenizer.Symbols.HtmlSymbol.#ctor(System.Int32,System.Int32,System.Int32,System.String,System.Web.Razor.Tokenizer.Symbols.HtmlSymbolType)">
+      <summary>This type/member supports the .NET Framework infrastructure and is not intended to be used directly from your code.Initializes a new instance of the <see cref="T:System.Web.Razor.Tokenizer.Symbols.HtmlSymbol" /> class.</summary>
+      <param name="offset">The location of the symbol.</param>
+      <param name="line">The exact line the symbol is found.</param>
+      <param name="column">The column number the symbol is found.</param>
+      <param name="content">The content value.</param>
+      <param name="type">The <see cref="T:System.Web.Razor.Tokenizer.Symbols.HtmlSymbolType" />.</param>
+    </member>
+    <member name="M:System.Web.Razor.Tokenizer.Symbols.HtmlSymbol.#ctor(System.Int32,System.Int32,System.Int32,System.String,System.Web.Razor.Tokenizer.Symbols.HtmlSymbolType,System.Collections.Generic.IEnumerable{System.Web.Razor.Parser.SyntaxTree.RazorError})">
+      <summary>This type/member supports the .NET Framework infrastructure and is not intended to be used directly from your code.Initializes a new instance of the <see cref="T:System.Web.Razor.Tokenizer.Symbols.HtmlSymbol" /> class.</summary>
+      <param name="offset">The location of the symbol.</param>
+      <param name="line">The exact line the symbol is found.</param>
+      <param name="column">The column number the symbol is found.</param>
+      <param name="content">The content value.</param>
+      <param name="type">The <see cref="T:System.Web.Razor.Tokenizer.Symbols.HtmlSymbolType" />.</param>
+      <param name="errors">The razor error.</param>
+    </member>
+    <member name="M:System.Web.Razor.Tokenizer.Symbols.HtmlSymbol.#ctor(System.Web.Razor.Text.SourceLocation,System.String,System.Web.Razor.Tokenizer.Symbols.HtmlSymbolType)">
+      <summary>This type/member supports the .NET Framework infrastructure and is not intended to be used directly from your code.Initializes a new instance of the <see cref="T:System.Web.Razor.Tokenizer.Symbols.HtmlSymbol" /> class.</summary>
+      <param name="start">The start of the source location.</param>
+      <param name="content">The content value.</param>
+      <param name="type">The <see cref="T:System.Web.Razor.Tokenizer.Symbols.HtmlSymbolType" />.</param>
+    </member>
+    <member name="M:System.Web.Razor.Tokenizer.Symbols.HtmlSymbol.#ctor(System.Web.Razor.Text.SourceLocation,System.String,System.Web.Razor.Tokenizer.Symbols.HtmlSymbolType,System.Collections.Generic.IEnumerable{System.Web.Razor.Parser.SyntaxTree.RazorError})">
+      <summary>This type/member supports the .NET Framework infrastructure and is not intended to be used directly from your code.Initializes a new instance of the <see cref="T:System.Web.Razor.Tokenizer.Symbols.HtmlSymbol" /> class.</summary>
+      <param name="start">The start of the source location.</param>
+      <param name="content">The content value.</param>
+      <param name="type">The <see cref="T:System.Web.Razor.Tokenizer.Symbols.HtmlSymbolType" />.</param>
+      <param name="errors">The razor error.</param>
+    </member>
+    <member name="T:System.Web.Razor.Tokenizer.Symbols.HtmlSymbolType">
+      <summary>This type/member supports the .NET Framework infrastructure and is not intended to be used directly from your code.</summary>
+    </member>
+    <member name="F:System.Web.Razor.Tokenizer.Symbols.HtmlSymbolType.Bang">
+      <summary>This type/member supports the .NET Framework infrastructure and is not intended to be used directly from your code.</summary>
+    </member>
+    <member name="F:System.Web.Razor.Tokenizer.Symbols.HtmlSymbolType.CloseAngle">
+      <summary>This type/member supports the .NET Framework infrastructure and is not intended to be used directly from your code.</summary>
+    </member>
+    <member name="F:System.Web.Razor.Tokenizer.Symbols.HtmlSymbolType.Colon">
+      <summary>This type/member supports the .NET Framework infrastructure and is not intended to be used directly from your code.</summary>
+    </member>
+    <member name="F:System.Web.Razor.Tokenizer.Symbols.HtmlSymbolType.DoubleHyphen">
+      <summary>This type/member supports the .NET Framework infrastructure and is not intended to be used directly from your code.</summary>
+    </member>
+    <member name="F:System.Web.Razor.Tokenizer.Symbols.HtmlSymbolType.DoubleQuote">
+      <summary>This type/member supports the .NET Framework infrastructure and is not intended to be used directly from your code.</summary>
+    </member>
+    <member name="F:System.Web.Razor.Tokenizer.Symbols.HtmlSymbolType.Equals">
+      <summary>This type/member supports the .NET Framework infrastructure and is not intended to be used directly from your code.</summary>
+    </member>
+    <member name="F:System.Web.Razor.Tokenizer.Symbols.HtmlSymbolType.LeftBracket">
+      <summary>This type/member supports the .NET Framework infrastructure and is not intended to be used directly from your code.</summary>
+    </member>
+    <member name="F:System.Web.Razor.Tokenizer.Symbols.HtmlSymbolType.NewLine">
+      <summary>This type/member supports the .NET Framework infrastructure and is not intended to be used directly from your code.</summary>
+    </member>
+    <member name="F:System.Web.Razor.Tokenizer.Symbols.HtmlSymbolType.OpenAngle">
+      <summary>This type/member supports the .NET Framework infrastructure and is not intended to be used directly from your code.</summary>
+    </member>
+    <member name="F:System.Web.Razor.Tokenizer.Symbols.HtmlSymbolType.QuestionMark">
+      <summary>This type/member supports the .NET Framework infrastructure and is not intended to be used directly from your code.</summary>
+    </member>
+    <member name="F:System.Web.Razor.Tokenizer.Symbols.HtmlSymbolType.RazorComment">
+      <summary>This type/member supports the .NET Framework infrastructure and is not intended to be used directly from your code.</summary>
+    </member>
+    <member name="F:System.Web.Razor.Tokenizer.Symbols.HtmlSymbolType.RazorCommentStar">
+      <summary>This type/member supports the .NET Framework infrastructure and is not intended to be used directly from your code.</summary>
+    </member>
+    <member name="F:System.Web.Razor.Tokenizer.Symbols.HtmlSymbolType.RazorCommentTransition">
+      <summary>This type/member supports the .NET Framework infrastructure and is not intended to be used directly from your code.</summary>
+    </member>
+    <member name="F:System.Web.Razor.Tokenizer.Symbols.HtmlSymbolType.RightBracket">
+      <summary>This type/member supports the .NET Framework infrastructure and is not intended to be used directly from your code.</summary>
+    </member>
+    <member name="F:System.Web.Razor.Tokenizer.Symbols.HtmlSymbolType.SingleQuote">
+      <summary>This type/member supports the .NET Framework infrastructure and is not intended to be used directly from your code.</summary>
+    </member>
+    <member name="F:System.Web.Razor.Tokenizer.Symbols.HtmlSymbolType.Solidus">
+      <summary>This type/member supports the .NET Framework infrastructure and is not intended to be used directly from your code.</summary>
+    </member>
+    <member name="F:System.Web.Razor.Tokenizer.Symbols.HtmlSymbolType.Text">
+      <summary>This type/member supports the .NET Framework infrastructure and is not intended to be used directly from your code.</summary>
+    </member>
+    <member name="F:System.Web.Razor.Tokenizer.Symbols.HtmlSymbolType.Transition">
+      <summary>This type/member supports the .NET Framework infrastructure and is not intended to be used directly from your code.</summary>
+    </member>
+    <member name="F:System.Web.Razor.Tokenizer.Symbols.HtmlSymbolType.Unknown">
+      <summary>This type/member supports the .NET Framework infrastructure and is not intended to be used directly from your code.</summary>
+    </member>
+    <member name="F:System.Web.Razor.Tokenizer.Symbols.HtmlSymbolType.WhiteSpace">
+      <summary>This type/member supports the .NET Framework infrastructure and is not intended to be used directly from your code.</summary>
+    </member>
+    <member name="T:System.Web.Razor.Tokenizer.Symbols.ISymbol">
+      <summary>This type/member supports the .NET Framework infrastructure and is not intended to be used directly from your code.Represents an interface for the web razor symbol.</summary>
+    </member>
+    <member name="M:System.Web.Razor.Tokenizer.Symbols.ISymbol.ChangeStart(System.Web.Razor.Text.SourceLocation)">
+      <summary>This type/member supports the .NET Framework infrastructure and is not intended to be used directly from your code.Changes the location of the symbol.</summary>
+      <param name="newStart">The new location of the symbol.</param>
+    </member>
+    <member name="P:System.Web.Razor.Tokenizer.Symbols.ISymbol.Content">
+      <summary>This type/member supports the .NET Framework infrastructure and is not intended to be used directly from your code.Gets the content of the symbol.</summary>
+      <returns>The content of the symbol.</returns>
+    </member>
+    <member name="M:System.Web.Razor.Tokenizer.Symbols.ISymbol.OffsetStart(System.Web.Razor.Text.SourceLocation)">
+      <summary>This type/member supports the .NET Framework infrastructure and is not intended to be used directly from your code.Indicates the starting offset of the symbol.</summary>
+      <param name="documentStart">The location where to start the document.</param>
+    </member>
+    <member name="P:System.Web.Razor.Tokenizer.Symbols.ISymbol.Start">
+      <summary>This type/member supports the .NET Framework infrastructure and is not intended to be used directly from your code.Gets the location of the symbol.</summary>
+      <returns>The location of the symbol.</returns>
+    </member>
+    <member name="T:System.Web.Razor.Tokenizer.Symbols.KnownSymbolType">
+      <summary>This type/member supports the .NET Framework infrastructure and is not intended to be used directly from your code.</summary>
+    </member>
+    <member name="F:System.Web.Razor.Tokenizer.Symbols.KnownSymbolType.CommentBody">
+      <summary>This type/member supports the .NET Framework infrastructure and is not intended to be used directly from your code.</summary>
+    </member>
+    <member name="F:System.Web.Razor.Tokenizer.Symbols.KnownSymbolType.CommentStar">
+      <summary>This type/member supports the .NET Framework infrastructure and is not intended to be used directly from your code.</summary>
+    </member>
+    <member name="F:System.Web.Razor.Tokenizer.Symbols.KnownSymbolType.CommentStart">
+      <summary>This type/member supports the .NET Framework infrastructure and is not intended to be used directly from your code.</summary>
+    </member>
+    <member name="F:System.Web.Razor.Tokenizer.Symbols.KnownSymbolType.Identifier">
+      <summary>This type/member supports the .NET Framework infrastructure and is not intended to be used directly from your code.</summary>
+    </member>
+    <member name="F:System.Web.Razor.Tokenizer.Symbols.KnownSymbolType.Keyword">
+      <summary>This type/member supports the .NET Framework infrastructure and is not intended to be used directly from your code.</summary>
+    </member>
+    <member name="F:System.Web.Razor.Tokenizer.Symbols.KnownSymbolType.NewLine">
+      <summary>This type/member supports the .NET Framework infrastructure and is not intended to be used directly from your code.</summary>
+    </member>
+    <member name="F:System.Web.Razor.Tokenizer.Symbols.KnownSymbolType.Transition">
+      <summary>This type/member supports the .NET Framework infrastructure and is not intended to be used directly from your code.</summary>
+    </member>
+    <member name="F:System.Web.Razor.Tokenizer.Symbols.KnownSymbolType.Unknown">
+      <summary>This type/member supports the .NET Framework infrastructure and is not intended to be used directly from your code.</summary>
+    </member>
+    <member name="F:System.Web.Razor.Tokenizer.Symbols.KnownSymbolType.WhiteSpace">
+      <summary>This type/member supports the .NET Framework infrastructure and is not intended to be used directly from your code.</summary>
+    </member>
+    <member name="T:System.Web.Razor.Tokenizer.Symbols.SymbolBase`1">
+      <summary>This type/member supports the .NET Framework infrastructure and is not intended to be used directly from your code.Represents a new instance of symbols.</summary>
+      <typeparam name="TType">The generic type.</typeparam>
+    </member>
+    <member name="M:System.Web.Razor.Tokenizer.Symbols.SymbolBase`1.#ctor(System.Web.Razor.Text.SourceLocation,System.String,`0,System.Collections.Generic.IEnumerable{System.Web.Razor.Parser.SyntaxTree.RazorError})">
+      <summary>This type/member supports the .NET Framework infrastructure and is not intended to be used directly from your code.Initializes a new instance of the <see cref="T:System.Web.Razor.Tokenizer.Symbols.SymbolBase`1" /> class.</summary>
+      <param name="start">The source location.</param>
+      <param name="content">The content value.</param>
+      <param name="type">The type.</param>
+      <param name="errors">The razor error.</param>
+    </member>
+    <member name="M:System.Web.Razor.Tokenizer.Symbols.SymbolBase`1.ChangeStart(System.Web.Razor.Text.SourceLocation)">
+      <summary>This type/member supports the .NET Framework infrastructure and is not intended to be used directly from your code.Changes the start of the machine.</summary>
+      <param name="newStart">The new start.</param>
+    </member>
+    <member name="P:System.Web.Razor.Tokenizer.Symbols.SymbolBase`1.Content">
+      <summary>This type/member supports the .NET Framework infrastructure and is not intended to be used directly from your code.Gets the content of a <see cref="T:System.Web.Razor.Tokenizer.Symbols.SymbolBase`1" />.</summary>
+      <returns>The content of a <see cref="T:System.Web.Razor.Tokenizer.Symbols.SymbolBase`1" />.</returns>
+    </member>
+    <member name="M:System.Web.Razor.Tokenizer.Symbols.SymbolBase`1.Equals(System.Object)">
+      <summary>This type/member supports the .NET Framework infrastructure and is not intended to be used directly from your code.Determines whether the specified Object is equal to the current Object.</summary>
+      <returns>true if the specified Object is equal to the current Object; otherwise, false.</returns>
+      <param name="obj">The object.</param>
+    </member>
+    <member name="P:System.Web.Razor.Tokenizer.Symbols.SymbolBase`1.Errors">
+      <summary>This type/member supports the .NET Framework infrastructure and is not intended to be used directly from your code.Gets the razor error.</summary>
+      <returns>The razor error.</returns>
+    </member>
+    <member name="M:System.Web.Razor.Tokenizer.Symbols.SymbolBase`1.GetHashCode">
+      <summary>This type/member supports the .NET Framework infrastructure and is not intended to be used directly from your code.Retrieves a hash code based on the current <see cref="T:System.Web.Razor.Tokenizer.Symbols.SymbolBase`1" /> object.</summary>
+      <returns>A hash of the current <see cref="T:System.Web.Razor.Tokenizer.Symbols.SymbolBase`1" /> object.</returns>
+    </member>
+    <member name="M:System.Web.Razor.Tokenizer.Symbols.SymbolBase`1.OffsetStart(System.Web.Razor.Text.SourceLocation)">
+      <summary>This type/member supports the .NET Framework infrastructure and is not intended to be used directly from your code.Starts the time’s offset for the source location.</summary>
+      <param name="documentStart">The document start.</param>
+    </member>
+    <member name="P:System.Web.Razor.Tokenizer.Symbols.SymbolBase`1.Start">
+      <summary>This type/member supports the .NET Framework infrastructure and is not intended to be used directly from your code.Gets the starting point of the source location.</summary>
+      <returns>The starting point of the source location.</returns>
+    </member>
+    <member name="M:System.Web.Razor.Tokenizer.Symbols.SymbolBase`1.ToString">
+      <summary>This type/member supports the .NET Framework infrastructure and is not intended to be used directly from your code.Generates a string representation of the current <see cref="T:System.Web.Razor.Tokenizer.Symbols.SymbolBase`1" /> object.</summary>
+      <returns>A string representation of the current <see cref="T:System.Web.Razor.Tokenizer.Symbols.SymbolBase`1" /> object.</returns>
+    </member>
+    <member name="P:System.Web.Razor.Tokenizer.Symbols.SymbolBase`1.Type">
+      <summary>This type/member supports the .NET Framework infrastructure and is not intended to be used directly from your code.Gets a Type that inherits from the base Type.</summary>
+      <returns>A Type that inherits from the base Type.</returns>
+    </member>
+    <member name="T:System.Web.Razor.Tokenizer.Symbols.SymbolExtensions">
+      <summary>This type/member supports the .NET Framework infrastructure and is not intended to be used directly from your code.Represents the symbol extensions for the web tokenizer.</summary>
+    </member>
+    <member name="M:System.Web.Razor.Tokenizer.Symbols.SymbolExtensions.GetContent(System.Collections.Generic.IEnumerable{System.Web.Razor.Tokenizer.Symbols.ISymbol},System.Web.Razor.Text.SourceLocation)">
+      <summary>This type/member supports the .NET Framework infrastructure and is not intended to be used directly from your code.Gets the content of this <see cref="T:System.Web.Razor.Tokenizer.Symbols.SymbolExtensions" /> class.</summary>
+      <returns>The content of this <see cref="T:System.Web.Razor.Tokenizer.Symbols.SymbolExtensions" /> class.</returns>
+      <param name="symbols">The symbols to provide.</param>
+      <param name="spanStart">The starting index of the span.</param>
+    </member>
+    <member name="M:System.Web.Razor.Tokenizer.Symbols.SymbolExtensions.GetContent(System.Web.Razor.Parser.SyntaxTree.SpanBuilder)">
+      <summary>This type/member supports the .NET Framework infrastructure and is not intended to be used directly from your code.Gets the content of this <see cref="T:System.Web.Razor.Tokenizer.Symbols.SymbolExtensions" /> class.</summary>
+      <returns>The content of this <see cref="T:System.Web.Razor.Tokenizer.Symbols.SymbolExtensions" /> class.</returns>
+      <param name="span">The intersection with the given span.</param>
+    </member>
+    <member name="M:System.Web.Razor.Tokenizer.Symbols.SymbolExtensions.GetContent(System.Web.Razor.Parser.SyntaxTree.SpanBuilder,System.Func{System.Collections.Generic.IEnumerable{System.Web.Razor.Tokenizer.Symbols.ISymbol},System.Collections.Generic.IEnumerable{System.Web.Razor.Tokenizer.Symbols.ISymbol}})">
+      <summary>Gets the content of this <see cref="T:System.Web.Razor.Tokenizer.Symbols.SymbolExtensions" /> class.</summary>
+      <returns>The content of this <see cref="T:System.Web.Razor.Tokenizer.Symbols.SymbolExtensions" /> class.</returns>
+      <param name="span">The intersection with the given span.</param>
+      <param name="filter">A list of chosen symbols.</param>
+    </member>
+    <member name="M:System.Web.Razor.Tokenizer.Symbols.SymbolExtensions.GetContent(System.Web.Razor.Tokenizer.Symbols.ISymbol)">
+      <summary>This type/member supports the .NET Framework infrastructure and is not intended to be used directly from your code.Gets the content of this <see cref="T:System.Web.Razor.Tokenizer.Symbols.SymbolExtensions" /> class.</summary>
+      <returns>The content of this <see cref="T:System.Web.Razor.Tokenizer.Symbols.SymbolExtensions" /> class.</returns>
+      <param name="symbol">The provided symbols.</param>
+    </member>
+    <member name="T:System.Web.Razor.Tokenizer.Symbols.VBKeyword">
+      <summary>Enumerates the list of Visual Basic keywords.</summary>
+    </member>
+    <member name="F:System.Web.Razor.Tokenizer.Symbols.VBKeyword.AddHandler">
+      <summary>This type/member supports the .NET Framework infrastructure and is not intended to be used directly from your code.</summary>
+    </member>
+    <member name="F:System.Web.Razor.Tokenizer.Symbols.VBKeyword.AddressOf">
+      <summary>This type/member supports the .NET Framework infrastructure and is not intended to be used directly from your code.</summary>
+    </member>
+    <member name="F:System.Web.Razor.Tokenizer.Symbols.VBKeyword.Alias">
+      <summary>This type/member supports the .NET Framework infrastructure and is not intended to be used directly from your code.</summary>
+    </member>
+    <member name="F:System.Web.Razor.Tokenizer.Symbols.VBKeyword.And">
+      <summary>This type/member supports the .NET Framework infrastructure and is not intended to be used directly from your code.</summary>
+    </member>
+    <member name="F:System.Web.Razor.Tokenizer.Symbols.VBKeyword.AndAlso">
+      <summary>This type/member supports the .NET Framework infrastructure and is not intended to be used directly from your code.</summary>
+    </member>
+    <member name="F:System.Web.Razor.Tokenizer.Symbols.VBKeyword.As">
+      <summary>This type/member supports the .NET Framework infrastructure and is not intended to be used directly from your code.</summary>
+    </member>
+    <member name="F:System.Web.Razor.Tokenizer.Symbols.VBKeyword.Boolean">
+      <summary>This type/member supports the .NET Framework infrastructure and is not intended to be used directly from your code.</summary>
+    </member>
+    <member name="F:System.Web.Razor.Tokenizer.Symbols.VBKeyword.ByRef">
+      <summary>This type/member supports the .NET Framework infrastructure and is not intended to be used directly from your code.</summary>
+    </member>
+    <member name="F:System.Web.Razor.Tokenizer.Symbols.VBKeyword.Byte">
+      <summary>This type/member supports the .NET Framework infrastructure and is not intended to be used directly from your code.</summary>
+    </member>
+    <member name="F:System.Web.Razor.Tokenizer.Symbols.VBKeyword.ByVal">
+      <summary>This type/member supports the .NET Framework infrastructure and is not intended to be used directly from your code.</summary>
+    </member>
+    <member name="F:System.Web.Razor.Tokenizer.Symbols.VBKeyword.Call">
+      <summary>This type/member supports the .NET Framework infrastructure and is not intended to be used directly from your code.</summary>
+    </member>
+    <member name="F:System.Web.Razor.Tokenizer.Symbols.VBKeyword.Case">
+      <summary>This type/member supports the .NET Framework infrastructure and is not intended to be used directly from your code.</summary>
+    </member>
+    <member name="F:System.Web.Razor.Tokenizer.Symbols.VBKeyword.Catch">
+      <summary>This type/member supports the .NET Framework infrastructure and is not intended to be used directly from your code.</summary>
+    </member>
+    <member name="F:System.Web.Razor.Tokenizer.Symbols.VBKeyword.CBool">
+      <summary>This type/member supports the .NET Framework infrastructure and is not intended to be used directly from your code.</summary>
+    </member>
+    <member name="F:System.Web.Razor.Tokenizer.Symbols.VBKeyword.CByte">
+      <summary>This type/member supports the .NET Framework infrastructure and is not intended to be used directly from your code.</summary>
+    </member>
+    <member name="F:System.Web.Razor.Tokenizer.Symbols.VBKeyword.CChar">
+      <summary>This type/member supports the .NET Framework infrastructure and is not intended to be used directly from your code.</summary>
+    </member>
+    <member name="F:System.Web.Razor.Tokenizer.Symbols.VBKeyword.CDate">
+      <summary>This type/member supports the .NET Framework infrastructure and is not intended to be used directly from your code.</summary>
+    </member>
+    <member name="F:System.Web.Razor.Tokenizer.Symbols.VBKeyword.CDbl">
+      <summary>This type/member supports the .NET Framework infrastructure and is not intended to be used directly from your code.</summary>
+    </member>
+    <member name="F:System.Web.Razor.Tokenizer.Symbols.VBKeyword.CDec">
+      <summary>This type/member supports the .NET Framework infrastructure and is not intended to be used directly from your code.</summary>
+    </member>
+    <member name="F:System.Web.Razor.Tokenizer.Symbols.VBKeyword.Char">
+      <summary>This type/member supports the .NET Framework infrastructure and is not intended to be used directly from your code.</summary>
+    </member>
+    <member name="F:System.Web.Razor.Tokenizer.Symbols.VBKeyword.CInt">
+      <summary>This type/member supports the .NET Framework infrastructure and is not intended to be used directly from your code.</summary>
+    </member>
+    <member name="F:System.Web.Razor.Tokenizer.Symbols.VBKeyword.Class">
+      <summary>This type/member supports the .NET Framework infrastructure and is not intended to be used directly from your code.</summary>
+    </member>
+    <member name="F:System.Web.Razor.Tokenizer.Symbols.VBKeyword.CLng">
+      <summary>This type/member supports the .NET Framework infrastructure and is not intended to be used directly from your code.</summary>
+    </member>
+    <member name="F:System.Web.Razor.Tokenizer.Symbols.VBKeyword.CObj">
+      <summary>This type/member supports the .NET Framework infrastructure and is not intended to be used directly from your code.</summary>
+    </member>
+    <member name="F:System.Web.Razor.Tokenizer.Symbols.VBKeyword.Const">
+      <summary>This type/member supports the .NET Framework infrastructure and is not intended to be used directly from your code.</summary>
+    </member>
+    <member name="F:System.Web.Razor.Tokenizer.Symbols.VBKeyword.Continue">
+      <summary>This type/member supports the .NET Framework infrastructure and is not intended to be used directly from your code.</summary>
+    </member>
+    <member name="F:System.Web.Razor.Tokenizer.Symbols.VBKeyword.CSByte">
+      <summary>This type/member supports the .NET Framework infrastructure and is not intended to be used directly from your code.</summary>
+    </member>
+    <member name="F:System.Web.Razor.Tokenizer.Symbols.VBKeyword.CShort">
+      <summary>This type/member supports the .NET Framework infrastructure and is not intended to be used directly from your code.</summary>
+    </member>
+    <member name="F:System.Web.Razor.Tokenizer.Symbols.VBKeyword.CSng">
+      <summary>This type/member supports the .NET Framework infrastructure and is not intended to be used directly from your code.</summary>
+    </member>
+    <member name="F:System.Web.Razor.Tokenizer.Symbols.VBKeyword.CStr">
+      <summary>This type/member supports the .NET Framework infrastructure and is not intended to be used directly from your code.</summary>
+    </member>
+    <member name="F:System.Web.Razor.Tokenizer.Symbols.VBKeyword.CType">
+      <summary>This type/member supports the .NET Framework infrastructure and is not intended to be used directly from your code.</summary>
+    </member>
+    <member name="F:System.Web.Razor.Tokenizer.Symbols.VBKeyword.CUInt">
+      <summary>This type/member supports the .NET Framework infrastructure and is not intended to be used directly from your code.</summary>
+    </member>
+    <member name="F:System.Web.Razor.Tokenizer.Symbols.VBKeyword.CULng">
+      <summary>This type/member supports the .NET Framework infrastructure and is not intended to be used directly from your code.</summary>
+    </member>
+    <member name="F:System.Web.Razor.Tokenizer.Symbols.VBKeyword.CUShort">
+      <summary>This type/member supports the .NET Framework infrastructure and is not intended to be used directly from your code.</summary>
+    </member>
+    <member name="F:System.Web.Razor.Tokenizer.Symbols.VBKeyword.Date">
+      <summary>This type/member supports the .NET Framework infrastructure and is not intended to be used directly from your code.</summary>
+    </member>
+    <member name="F:System.Web.Razor.Tokenizer.Symbols.VBKeyword.Decimal">
+      <summary>This type/member supports the .NET Framework infrastructure and is not intended to be used directly from your code.</summary>
+    </member>
+    <member name="F:System.Web.Razor.Tokenizer.Symbols.VBKeyword.Declare">
+      <summary>This type/member supports the .NET Framework infrastructure and is not intended to be used directly from your code.</summary>
+    </member>
+    <member name="F:System.Web.Razor.Tokenizer.Symbols.VBKeyword.Default">
+      <summary>This type/member supports the .NET Framework infrastructure and is not intended to be used directly from your code.</summary>
+    </member>
+    <member name="F:System.Web.Razor.Tokenizer.Symbols.VBKeyword.Delegate">
+      <summary>This type/member supports the .NET Framework infrastructure and is not intended to be used directly from your code.</summary>
+    </member>
+    <member name="F:System.Web.Razor.Tokenizer.Symbols.VBKeyword.Dim">
+      <summary>This type/member supports the .NET Framework infrastructure and is not intended to be used directly from your code.</summary>
+    </member>
+    <member name="F:System.Web.Razor.Tokenizer.Symbols.VBKeyword.DirectCast">
+      <summary>This type/member supports the .NET Framework infrastructure and is not intended to be used directly from your code.</summary>
+    </member>
+    <member name="F:System.Web.Razor.Tokenizer.Symbols.VBKeyword.Do">
+      <summary>This type/member supports the .NET Framework infrastructure and is not intended to be used directly from your code.</summary>
+    </member>
+    <member name="F:System.Web.Razor.Tokenizer.Symbols.VBKeyword.Double">
+      <summary>This type/member supports the .NET Framework infrastructure and is not intended to be used directly from your code.</summary>
+    </member>
+    <member name="F:System.Web.Razor.Tokenizer.Symbols.VBKeyword.Each">
+      <summary>This type/member supports the .NET Framework infrastructure and is not intended to be used directly from your code.</summary>
+    </member>
+    <member name="F:System.Web.Razor.Tokenizer.Symbols.VBKeyword.Else">
+      <summary>This type/member supports the .NET Framework infrastructure and is not intended to be used directly from your code.</summary>
+    </member>
+    <member name="F:System.Web.Razor.Tokenizer.Symbols.VBKeyword.ElseIf">
+      <summary>This type/member supports the .NET Framework infrastructure and is not intended to be used directly from your code.</summary>
+    </member>
+    <member name="F:System.Web.Razor.Tokenizer.Symbols.VBKeyword.End">
+      <summary>This type/member supports the .NET Framework infrastructure and is not intended to be used directly from your code.</summary>
+    </member>
+    <member name="F:System.Web.Razor.Tokenizer.Symbols.VBKeyword.EndIf">
+      <summary>This type/member supports the .NET Framework infrastructure and is not intended to be used directly from your code.</summary>
+    </member>
+    <member name="F:System.Web.Razor.Tokenizer.Symbols.VBKeyword.Enum">
+      <summary>This type/member supports the .NET Framework infrastructure and is not intended to be used directly from your code.</summary>
+    </member>
+    <member name="F:System.Web.Razor.Tokenizer.Symbols.VBKeyword.Erase">
+      <summary>This type/member supports the .NET Framework infrastructure and is not intended to be used directly from your code.</summary>
+    </member>
+    <member name="F:System.Web.Razor.Tokenizer.Symbols.VBKeyword.Error">
+      <summary>This type/member supports the .NET Framework infrastructure and is not intended to be used directly from your code.</summary>
+    </member>
+    <member name="F:System.Web.Razor.Tokenizer.Symbols.VBKeyword.Event">
+      <summary>This type/member supports the .NET Framework infrastructure and is not intended to be used directly from your code.</summary>
+    </member>
+    <member name="F:System.Web.Razor.Tokenizer.Symbols.VBKeyword.Exit">
+      <summary>This type/member supports the .NET Framework infrastructure and is not intended to be used directly from your code.</summary>
+    </member>
+    <member name="F:System.Web.Razor.Tokenizer.Symbols.VBKeyword.False">
+      <summary>This type/member supports the .NET Framework infrastructure and is not intended to be used directly from your code.</summary>
+    </member>
+    <member name="F:System.Web.Razor.Tokenizer.Symbols.VBKeyword.Finally">
+      <summary>This type/member supports the .NET Framework infrastructure and is not intended to be used directly from your code.</summary>
+    </member>
+    <member name="F:System.Web.Razor.Tokenizer.Symbols.VBKeyword.For">
+      <summary>This type/member supports the .NET Framework infrastructure and is not intended to be used directly from your code.</summary>
+    </member>
+    <member name="F:System.Web.Razor.Tokenizer.Symbols.VBKeyword.Friend">
+      <summary>This type/member supports the .NET Framework infrastructure and is not intended to be used directly from your code.</summary>
+    </member>
+    <member name="F:System.Web.Razor.Tokenizer.Symbols.VBKeyword.Function">
+      <summary>This type/member supports the .NET Framework infrastructure and is not intended to be used directly from your code.</summary>
+    </member>
+    <member name="F:System.Web.Razor.Tokenizer.Symbols.VBKeyword.Get">
+      <summary>This type/member supports the .NET Framework infrastructure and is not intended to be used directly from your code.</summary>
+    </member>
+    <member name="F:System.Web.Razor.Tokenizer.Symbols.VBKeyword.GetType">
+      <summary>This type/member supports the .NET Framework infrastructure and is not intended to be used directly from your code.</summary>
+    </member>
+    <member name="F:System.Web.Razor.Tokenizer.Symbols.VBKeyword.GetXmlNamespace">
+      <summary>This type/member supports the .NET Framework infrastructure and is not intended to be used directly from your code.</summary>
+    </member>
+    <member name="F:System.Web.Razor.Tokenizer.Symbols.VBKeyword.Global">
+      <summary>This type/member supports the .NET Framework infrastructure and is not intended to be used directly from your code.</summary>
+    </member>
+    <member name="F:System.Web.Razor.Tokenizer.Symbols.VBKeyword.GoSub">
+      <summary>This type/member supports the .NET Framework infrastructure and is not intended to be used directly from your code.</summary>
+    </member>
+    <member name="F:System.Web.Razor.Tokenizer.Symbols.VBKeyword.GoTo">
+      <summary>This type/member supports the .NET Framework infrastructure and is not intended to be used directly from your code.</summary>
+    </member>
+    <member name="F:System.Web.Razor.Tokenizer.Symbols.VBKeyword.Handles">
+      <summary>This type/member supports the .NET Framework infrastructure and is not intended to be used directly from your code.</summary>
+    </member>
+    <member name="F:System.Web.Razor.Tokenizer.Symbols.VBKeyword.If">
+      <summary>This type/member supports the .NET Framework infrastructure and is not intended to be used directly from your code.</summary>
+    </member>
+    <member name="F:System.Web.Razor.Tokenizer.Symbols.VBKeyword.Implements">
+      <summary>This type/member supports the .NET Framework infrastructure and is not intended to be used directly from your code.</summary>
+    </member>
+    <member name="F:System.Web.Razor.Tokenizer.Symbols.VBKeyword.Imports">
+      <summary>This type/member supports the .NET Framework infrastructure and is not intended to be used directly from your code.</summary>
+    </member>
+    <member name="F:System.Web.Razor.Tokenizer.Symbols.VBKeyword.In">
+      <summary>This type/member supports the .NET Framework infrastructure and is not intended to be used directly from your code.</summary>
+    </member>
+    <member name="F:System.Web.Razor.Tokenizer.Symbols.VBKeyword.Inherits">
+      <summary>This type/member supports the .NET Framework infrastructure and is not intended to be used directly from your code.</summary>
+    </member>
+    <member name="F:System.Web.Razor.Tokenizer.Symbols.VBKeyword.Integer">
+      <summary>This type/member supports the .NET Framework infrastructure and is not intended to be used directly from your code.</summary>
+    </member>
+    <member name="F:System.Web.Razor.Tokenizer.Symbols.VBKeyword.Interface">
+      <summary>This type/member supports the .NET Framework infrastructure and is not intended to be used directly from your code.</summary>
+    </member>
+    <member name="F:System.Web.Razor.Tokenizer.Symbols.VBKeyword.Is">
+      <summary>This type/member supports the .NET Framework infrastructure and is not intended to be used directly from your code.</summary>
+    </member>
+    <member name="F:System.Web.Razor.Tokenizer.Symbols.VBKeyword.IsNot">
+      <summary>This type/member supports the .NET Framework infrastructure and is not intended to be used directly from your code.</summary>
+    </member>
+    <member name="F:System.Web.Razor.Tokenizer.Symbols.VBKeyword.Let">
+      <summary>This type/member supports the .NET Framework infrastructure and is not intended to be used directly from your code.</summary>
+    </member>
+    <member name="F:System.Web.Razor.Tokenizer.Symbols.VBKeyword.Lib">
+      <summary>This type/member supports the .NET Framework infrastructure and is not intended to be used directly from your code.</summary>
+    </member>
+    <member name="F:System.Web.Razor.Tokenizer.Symbols.VBKeyword.Like">
+      <summary>This type/member supports the .NET Framework infrastructure and is not intended to be used directly from your code.</summary>
+    </member>
+    <member name="F:System.Web.Razor.Tokenizer.Symbols.VBKeyword.Long">
+      <summary>This type/member supports the .NET Framework infrastructure and is not intended to be used directly from your code.</summary>
+    </member>
+    <member name="F:System.Web.Razor.Tokenizer.Symbols.VBKeyword.Loop">
+      <summary>This type/member supports the .NET Framework infrastructure and is not intended to be used directly from your code.</summary>
+    </member>
+    <member name="F:System.Web.Razor.Tokenizer.Symbols.VBKeyword.Me">
+      <summary>This type/member supports the .NET Framework infrastructure and is not intended to be used directly from your code.</summary>
+    </member>
+    <member name="F:System.Web.Razor.Tokenizer.Symbols.VBKeyword.Mod">
+      <summary>This type/member supports the .NET Framework infrastructure and is not intended to be used directly from your code.</summary>
+    </member>
+    <member name="F:System.Web.Razor.Tokenizer.Symbols.VBKeyword.Module">
+      <summary>This type/member supports the .NET Framework infrastructure and is not intended to be used directly from your code.</summary>
+    </member>
+    <member name="F:System.Web.Razor.Tokenizer.Symbols.VBKeyword.MustInherit">
+      <summary>This type/member supports the .NET Framework infrastructure and is not intended to be used directly from your code.</summary>
+    </member>
+    <member name="F:System.Web.Razor.Tokenizer.Symbols.VBKeyword.MustOverride">
+      <summary>This type/member supports the .NET Framework infrastructure and is not intended to be used directly from your code.</summary>
+    </member>
+    <member name="F:System.Web.Razor.Tokenizer.Symbols.VBKeyword.MyBase">
+      <summary>This type/member supports the .NET Framework infrastructure and is not intended to be used directly from your code.</summary>
+    </member>
+    <member name="F:System.Web.Razor.Tokenizer.Symbols.VBKeyword.MyClass">
+      <summary>This type/member supports the .NET Framework infrastructure and is not intended to be used directly from your code.</summary>
+    </member>
+    <member name="F:System.Web.Razor.Tokenizer.Symbols.VBKeyword.Namespace">
+      <summary>This type/member supports the .NET Framework infrastructure and is not intended to be used directly from your code.</summary>
+    </member>
+    <member name="F:System.Web.Razor.Tokenizer.Symbols.VBKeyword.Narrowing">
+      <summary>This type/member supports the .NET Framework infrastructure and is not intended to be used directly from your code.</summary>
+    </member>
+    <member name="F:System.Web.Razor.Tokenizer.Symbols.VBKeyword.New">
+      <summary>This type/member supports the .NET Framework infrastructure and is not intended to be used directly from your code.</summary>
+    </member>
+    <member name="F:System.Web.Razor.Tokenizer.Symbols.VBKeyword.Next">
+      <summary>This type/member supports the .NET Framework infrastructure and is not intended to be used directly from your code.</summary>
+    </member>
+    <member name="F:System.Web.Razor.Tokenizer.Symbols.VBKeyword.Not">
+      <summary>This type/member supports the .NET Framework infrastructure and is not intended to be used directly from your code.</summary>
+    </member>
+    <member name="F:System.Web.Razor.Tokenizer.Symbols.VBKeyword.Nothing">
+      <summary>This type/member supports the .NET Framework infrastructure and is not intended to be used directly from your code.</summary>
+    </member>
+    <member name="F:System.Web.Razor.Tokenizer.Symbols.VBKeyword.NotInheritable">
+      <summary>This type/member supports the .NET Framework infrastructure and is not intended to be used directly from your code.</summary>
+    </member>
+    <member name="F:System.Web.Razor.Tokenizer.Symbols.VBKeyword.NotOverridable">
+      <summary>This type/member supports the .NET Framework infrastructure and is not intended to be used directly from your code.</summary>
+    </member>
+    <member name="F:System.Web.Razor.Tokenizer.Symbols.VBKeyword.Object">
+      <summary>This type/member supports the .NET Framework infrastructure and is not intended to be used directly from your code.</summary>
+    </member>
+    <member name="F:System.Web.Razor.Tokenizer.Symbols.VBKeyword.Of">
+      <summary>This type/member supports the .NET Framework infrastructure and is not intended to be used directly from your code.</summary>
+    </member>
+    <member name="F:System.Web.Razor.Tokenizer.Symbols.VBKeyword.On">
+      <summary>This type/member supports the .NET Framework infrastructure and is not intended to be used directly from your code.</summary>
+    </member>
+    <member name="F:System.Web.Razor.Tokenizer.Symbols.VBKeyword.Operator">
+      <summary>This type/member supports the .NET Framework infrastructure and is not intended to be used directly from your code.</summary>
+    </member>
+    <member name="F:System.Web.Razor.Tokenizer.Symbols.VBKeyword.Option">
+      <summary>This type/member supports the .NET Framework infrastructure and is not intended to be used directly from your code.</summary>
+    </member>
+    <member name="F:System.Web.Razor.Tokenizer.Symbols.VBKeyword.Optional">
+      <summary>This type/member supports the .NET Framework infrastructure and is not intended to be used directly from your code.</summary>
+    </member>
+    <member name="F:System.Web.Razor.Tokenizer.Symbols.VBKeyword.Or">
+      <summary>This type/member supports the .NET Framework infrastructure and is not intended to be used directly from your code.</summary>
+    </member>
+    <member name="F:System.Web.Razor.Tokenizer.Symbols.VBKeyword.OrElse">
+      <summary>This type/member supports the .NET Framework infrastructure and is not intended to be used directly from your code.</summary>
+    </member>
+    <member name="F:System.Web.Razor.Tokenizer.Symbols.VBKeyword.Overloads">
+      <summary>This type/member supports the .NET Framework infrastructure and is not intended to be used directly from your code.</summary>
+    </member>
+    <member name="F:System.Web.Razor.Tokenizer.Symbols.VBKeyword.Overridable">
+      <summary>This type/member supports the .NET Framework infrastructure and is not intended to be used directly from your code.</summary>
+    </member>
+    <member name="F:System.Web.Razor.Tokenizer.Symbols.VBKeyword.Overrides">
+      <summary>This type/member supports the .NET Framework infrastructure and is not intended to be used directly from your code.</summary>
+    </member>
+    <member name="F:System.Web.Razor.Tokenizer.Symbols.VBKeyword.ParamArray">
+      <summary>This type/member supports the .NET Framework infrastructure and is not intended to be used directly from your code.</summary>
+    </member>
+    <member name="F:System.Web.Razor.Tokenizer.Symbols.VBKeyword.Partial">
+      <summary>This type/member supports the .NET Framework infrastructure and is not intended to be used directly from your code.</summary>
+    </member>
+    <member name="F:System.Web.Razor.Tokenizer.Symbols.VBKeyword.Private">
+      <summary>This type/member supports the .NET Framework infrastructure and is not intended to be used directly from your code.</summary>
+    </member>
+    <member name="F:System.Web.Razor.Tokenizer.Symbols.VBKeyword.Property">
+      <summary>This type/member supports the .NET Framework infrastructure and is not intended to be used directly from your code.</summary>
+    </member>
+    <member name="F:System.Web.Razor.Tokenizer.Symbols.VBKeyword.Protected">
+      <summary>This type/member supports the .NET Framework infrastructure and is not intended to be used directly from your code.</summary>
+    </member>
+    <member name="F:System.Web.Razor.Tokenizer.Symbols.VBKeyword.Public">
+      <summary>This type/member supports the .NET Framework infrastructure and is not intended to be used directly from your code.</summary>
+    </member>
+    <member name="F:System.Web.Razor.Tokenizer.Symbols.VBKeyword.RaiseEvent">
+      <summary>This type/member supports the .NET Framework infrastructure and is not intended to be used directly from your code.</summary>
+    </member>
+    <member name="F:System.Web.Razor.Tokenizer.Symbols.VBKeyword.ReadOnly">
+      <summary>This type/member supports the .NET Framework infrastructure and is not intended to be used directly from your code.</summary>
+    </member>
+    <member name="F:System.Web.Razor.Tokenizer.Symbols.VBKeyword.ReDim">
+      <summary>This type/member supports the .NET Framework infrastructure and is not intended to be used directly from your code.</summary>
+    </member>
+    <member name="F:System.Web.Razor.Tokenizer.Symbols.VBKeyword.Rem">
+      <summary>This type/member supports the .NET Framework infrastructure and is not intended to be used directly from your code.</summary>
+    </member>
+    <member name="F:System.Web.Razor.Tokenizer.Symbols.VBKeyword.RemoveHandler">
+      <summary>This type/member supports the .NET Framework infrastructure and is not intended to be used directly from your code.</summary>
+    </member>
+    <member name="F:System.Web.Razor.Tokenizer.Symbols.VBKeyword.Resume">
+      <summary>This type/member supports the .NET Framework infrastructure and is not intended to be used directly from your code.</summary>
+    </member>
+    <member name="F:System.Web.Razor.Tokenizer.Symbols.VBKeyword.Return">
+      <summary>This type/member supports the .NET Framework infrastructure and is not intended to be used directly from your code.</summary>
+    </member>
+    <member name="F:System.Web.Razor.Tokenizer.Symbols.VBKeyword.SByte">
+      <summary>This type/member supports the .NET Framework infrastructure and is not intended to be used directly from your code.</summary>
+    </member>
+    <member name="F:System.Web.Razor.Tokenizer.Symbols.VBKeyword.Select">
+      <summary>This type/member supports the .NET Framework infrastructure and is not intended to be used directly from your code.</summary>
+    </member>
+    <member name="F:System.Web.Razor.Tokenizer.Symbols.VBKeyword.Set">
+      <summary>This type/member supports the .NET Framework infrastructure and is not intended to be used directly from your code.</summary>
+    </member>
+    <member name="F:System.Web.Razor.Tokenizer.Symbols.VBKeyword.Shadows">
+      <summary>This type/member supports the .NET Framework infrastructure and is not intended to be used directly from your code.</summary>
+    </member>
+    <member name="F:System.Web.Razor.Tokenizer.Symbols.VBKeyword.Shared">
+      <summary>This type/member supports the .NET Framework infrastructure and is not intended to be used directly from your code.</summary>
+    </member>
+    <member name="F:System.Web.Razor.Tokenizer.Symbols.VBKeyword.Short">
+      <summary>This type/member supports the .NET Framework infrastructure and is not intended to be used directly from your code.</summary>
+    </member>
+    <member name="F:System.Web.Razor.Tokenizer.Symbols.VBKeyword.Single">
+      <summary>This type/member supports the .NET Framework infrastructure and is not intended to be used directly from your code.</summary>
+    </member>
+    <member name="F:System.Web.Razor.Tokenizer.Symbols.VBKeyword.Static">
+      <summary>This type/member supports the .NET Framework infrastructure and is not intended to be used directly from your code.</summary>
+    </member>
+    <member name="F:System.Web.Razor.Tokenizer.Symbols.VBKeyword.Step">
+      <summary>This type/member supports the .NET Framework infrastructure and is not intended to be used directly from your code.</summary>
+    </member>
+    <member name="F:System.Web.Razor.Tokenizer.Symbols.VBKeyword.Stop">
+      <summary>This type/member supports the .NET Framework infrastructure and is not intended to be used directly from your code.</summary>
+    </member>
+    <member name="F:System.Web.Razor.Tokenizer.Symbols.VBKeyword.String">
+      <summary>This type/member supports the .NET Framework infrastructure and is not intended to be used directly from your code.</summary>
+    </member>
+    <member name="F:System.Web.Razor.Tokenizer.Symbols.VBKeyword.Structure">
+      <summary>This type/member supports the .NET Framework infrastructure and is not intended to be used directly from your code.</summary>
+    </member>
+    <member name="F:System.Web.Razor.Tokenizer.Symbols.VBKeyword.Sub">
+      <summary>This type/member supports the .NET Framework infrastructure and is not intended to be used directly from your code.</summary>
+    </member>
+    <member name="F:System.Web.Razor.Tokenizer.Symbols.VBKeyword.SyncLock">
+      <summary>This type/member supports the .NET Framework infrastructure and is not intended to be used directly from your code.</summary>
+    </member>
+    <member name="F:System.Web.Razor.Tokenizer.Symbols.VBKeyword.Then">
+      <summary>This type/member supports the .NET Framework infrastructure and is not intended to be used directly from your code.</summary>
+    </member>
+    <member name="F:System.Web.Razor.Tokenizer.Symbols.VBKeyword.Throw">
+      <summary>This type/member supports the .NET Framework infrastructure and is not intended to be used directly from your code.</summary>
+    </member>
+    <member name="F:System.Web.Razor.Tokenizer.Symbols.VBKeyword.To">
+      <summary>This type/member supports the .NET Framework infrastructure and is not intended to be used directly from your code.</summary>
+    </member>
+    <member name="F:System.Web.Razor.Tokenizer.Symbols.VBKeyword.True">
+      <summary>This type/member supports the .NET Framework infrastructure and is not intended to be used directly from your code.</summary>
+    </member>
+    <member name="F:System.Web.Razor.Tokenizer.Symbols.VBKeyword.Try">
+      <summary>This type/member supports the .NET Framework infrastructure and is not intended to be used directly from your code.</summary>
+    </member>
+    <member name="F:System.Web.Razor.Tokenizer.Symbols.VBKeyword.TryCast">
+      <summary>This type/member supports the .NET Framework infrastructure and is not intended to be used directly from your code.</summary>
+    </member>
+    <member name="F:System.Web.Razor.Tokenizer.Symbols.VBKeyword.TypeOf">
+      <summary>This type/member supports the .NET Framework infrastructure and is not intended to be used directly from your code.</summary>
+    </member>
+    <member name="F:System.Web.Razor.Tokenizer.Symbols.VBKeyword.UInteger">
+      <summary>This type/member supports the .NET Framework infrastructure and is not intended to be used directly from your code.</summary>
+    </member>
+    <member name="F:System.Web.Razor.Tokenizer.Symbols.VBKeyword.ULong">
+      <summary>This type/member supports the .NET Framework infrastructure and is not intended to be used directly from your code.</summary>
+    </member>
+    <member name="F:System.Web.Razor.Tokenizer.Symbols.VBKeyword.UShort">
+      <summary>This type/member supports the .NET Framework infrastructure and is not intended to be used directly from your code.</summary>
+    </member>
+    <member name="F:System.Web.Razor.Tokenizer.Symbols.VBKeyword.Using">
+      <summary>This type/member supports the .NET Framework infrastructure and is not intended to be used directly from your code.</summary>
+    </member>
+    <member name="F:System.Web.Razor.Tokenizer.Symbols.VBKeyword.Variant">
+      <summary>This type/member supports the .NET Framework infrastructure and is not intended to be used directly from your code.</summary>
+    </member>
+    <member name="F:System.Web.Razor.Tokenizer.Symbols.VBKeyword.Wend">
+      <summary>This type/member supports the .NET Framework infrastructure and is not intended to be used directly from your code.</summary>
+    </member>
+    <member name="F:System.Web.Razor.Tokenizer.Symbols.VBKeyword.When">
+      <summary>This type/member supports the .NET Framework infrastructure and is not intended to be used directly from your code.</summary>
+    </member>
+    <member name="F:System.Web.Razor.Tokenizer.Symbols.VBKeyword.While">
+      <summary>This type/member supports the .NET Framework infrastructure and is not intended to be used directly from your code.</summary>
+    </member>
+    <member name="F:System.Web.Razor.Tokenizer.Symbols.VBKeyword.Widening">
+      <summary>This type/member supports the .NET Framework infrastructure and is not intended to be used directly from your code.</summary>
+    </member>
+    <member name="F:System.Web.Razor.Tokenizer.Symbols.VBKeyword.With">
+      <summary>This type/member supports the .NET Framework infrastructure and is not intended to be used directly from your code.</summary>
+    </member>
+    <member name="F:System.Web.Razor.Tokenizer.Symbols.VBKeyword.WithEvents">
+      <summary>This type/member supports the .NET Framework infrastructure and is not intended to be used directly from your code.</summary>
+    </member>
+    <member name="F:System.Web.Razor.Tokenizer.Symbols.VBKeyword.WriteOnly">
+      <summary>This type/member supports the .NET Framework infrastructure and is not intended to be used directly from your code.</summary>
+    </member>
+    <member name="F:System.Web.Razor.Tokenizer.Symbols.VBKeyword.Xor">
+      <summary>This type/member supports the .NET Framework infrastructure and is not intended to be used directly from your code.</summary>
+    </member>
+    <member name="T:System.Web.Razor.Tokenizer.Symbols.VBSymbol">
+      <summary>This type/member supports the .NET Framework infrastructure and is not intended to be used directly from your code.Represents the VB symbol components.</summary>
+    </member>
+    <member name="M:System.Web.Razor.Tokenizer.Symbols.VBSymbol.#ctor(System.Int32,System.Int32,System.Int32,System.String,System.Web.Razor.Tokenizer.Symbols.VBSymbolType)">
+      <summary>This type/member supports the .NET Framework infrastructure and is not intended to be used directly from your code.Initializes a new instance of the <see cref="T:System.Web.Razor.Tokenizer.Symbols.VBSymbol" /> class.</summary>
+      <param name="offset">The offset value.</param>
+      <param name="line">The line value.</param>
+      <param name="column">The column value.</param>
+      <param name="content">The content String value.</param>
+      <param name="type">The <see cref="T:System.Web.Razor.Tokenizer.Symbols.VBSymbolType" />.</param>
+    </member>
+    <member name="M:System.Web.Razor.Tokenizer.Symbols.VBSymbol.#ctor(System.Int32,System.Int32,System.Int32,System.String,System.Web.Razor.Tokenizer.Symbols.VBSymbolType,System.Collections.Generic.IEnumerable{System.Web.Razor.Parser.SyntaxTree.RazorError})">
+      <summary>This type/member supports the .NET Framework infrastructure and is not intended to be used directly from your code.Initializes a new instance of the <see cref="T:System.Web.Razor.Tokenizer.Symbols.VBSymbol" /> class.</summary>
+      <param name="offset">The offset value.</param>
+      <param name="line">The line value.</param>
+      <param name="column">The column value.</param>
+      <param name="content">The content String value.</param>
+      <param name="type">The <see cref="T:System.Web.Razor.Tokenizer.Symbols.VBSymbolType" />.</param>
+      <param name="errors">List of razor errors.</param>
+    </member>
+    <member name="M:System.Web.Razor.Tokenizer.Symbols.VBSymbol.#ctor(System.Web.Razor.Text.SourceLocation,System.String,System.Web.Razor.Tokenizer.Symbols.VBSymbolType)">
+      <summary>This type/member supports the .NET Framework infrastructure and is not intended to be used directly from your code.Initializes a new instance of the <see cref="T:System.Web.Razor.Tokenizer.Symbols.VBSymbol" /> class.</summary>
+      <param name="start">The start of the source location.</param>
+      <param name="content">The content String value.</param>
+      <param name="type">The <see cref="T:System.Web.Razor.Tokenizer.Symbols.VBSymbolType" />.</param>
+    </member>
+    <member name="M:System.Web.Razor.Tokenizer.Symbols.VBSymbol.#ctor(System.Web.Razor.Text.SourceLocation,System.String,System.Web.Razor.Tokenizer.Symbols.VBSymbolType,System.Collections.Generic.IEnumerable{System.Web.Razor.Parser.SyntaxTree.RazorError})">
+      <summary>This type/member supports the .NET Framework infrastructure and is not intended to be used directly from your code.Initializes a new instance of the <see cref="T:System.Web.Razor.Tokenizer.Symbols.VBSymbol" /> class.</summary>
+      <param name="start">The start of the source location.</param>
+      <param name="content">The content String value.</param>
+      <param name="type">The <see cref="T:System.Web.Razor.Tokenizer.Symbols.VBSymbolType" />.</param>
+      <param name="errors">List of razor errors.</param>
+    </member>
+    <member name="M:System.Web.Razor.Tokenizer.Symbols.VBSymbol.Equals(System.Object)">
+      <summary>This type/member supports the .NET Framework infrastructure and is not intended to be used directly from your code.Indicates a value whether the current object is equal to the new object.</summary>
+      <returns>true if the current object is equal to the new object; otherwise, false.</returns>
+      <param name="obj">The object to compare.</param>
+    </member>
+    <member name="M:System.Web.Razor.Tokenizer.Symbols.VBSymbol.GetHashCode">
+      <summary>This type/member supports the .NET Framework infrastructure and is not intended to be used directly from your code.Returns the hash code for this instance.</summary>
+      <returns>The hash code to return.</returns>
+    </member>
+    <member name="M:System.Web.Razor.Tokenizer.Symbols.VBSymbol.GetSample(System.Web.Razor.Tokenizer.Symbols.VBSymbolType)">
+      <summary>This type/member supports the .NET Framework infrastructure and is not intended to be used directly from your code.Gets the specified data sample from the object.</summary>
+      <returns>The specified data sample from the object.</returns>
+      <param name="type">The <see cref="T:System.Web.Razor.Tokenizer.Symbols.VBSymbolType" />.</param>
+    </member>
+    <member name="P:System.Web.Razor.Tokenizer.Symbols.VBSymbol.Keyword">
+      <summary>This type/member supports the .NET Framework infrastructure and is not intended to be used directly from your code.Gets or sets the keyword used in the VB.</summary>
+      <returns>The keyword used.</returns>
+    </member>
+    <member name="T:System.Web.Razor.Tokenizer.Symbols.VBSymbolType">
+      <summary>This type/member supports the .NET Framework infrastructure and is not intended to be used directly from your code.</summary>
+    </member>
+    <member name="F:System.Web.Razor.Tokenizer.Symbols.VBSymbolType.Add">
+      <summary>This type/member supports the .NET Framework infrastructure and is not intended to be used directly from your code.</summary>
+    </member>
+    <member name="F:System.Web.Razor.Tokenizer.Symbols.VBSymbolType.Bang">
+      <summary>This type/member supports the .NET Framework infrastructure and is not intended to be used directly from your code.</summary>
+    </member>
+    <member name="F:System.Web.Razor.Tokenizer.Symbols.VBSymbolType.CharacterLiteral">
+      <summary>This type/member supports the .NET Framework infrastructure and is not intended to be used directly from your code.</summary>
+    </member>
+    <member name="F:System.Web.Razor.Tokenizer.Symbols.VBSymbolType.Colon">
+      <summary>This type/member supports the .NET Framework infrastructure and is not intended to be used directly from your code.</summary>
+    </member>
+    <member name="F:System.Web.Razor.Tokenizer.Symbols.VBSymbolType.Comma">
+      <summary>This type/member supports the .NET Framework infrastructure and is not intended to be used directly from your code.</summary>
+    </member>
+    <member name="F:System.Web.Razor.Tokenizer.Symbols.VBSymbolType.Comment">
+      <summary>This type/member supports the .NET Framework infrastructure and is not intended to be used directly from your code.</summary>
+    </member>
+    <member name="F:System.Web.Razor.Tokenizer.Symbols.VBSymbolType.Concatenation">
+      <summary>This type/member supports the .NET Framework infrastructure and is not intended to be used directly from your code.</summary>
+    </member>
+    <member name="F:System.Web.Razor.Tokenizer.Symbols.VBSymbolType.DateLiteral">
+      <summary>This type/member supports the .NET Framework infrastructure and is not intended to be used directly from your code.</summary>
+    </member>
+    <member name="F:System.Web.Razor.Tokenizer.Symbols.VBSymbolType.Divide">
+      <summary>This type/member supports the .NET Framework infrastructure and is not intended to be used directly from your code.</summary>
+    </member>
+    <member name="F:System.Web.Razor.Tokenizer.Symbols.VBSymbolType.Dollar">
+      <summary>This type/member supports the .NET Framework infrastructure and is not intended to be used directly from your code.</summary>
+    </member>
+    <member name="F:System.Web.Razor.Tokenizer.Symbols.VBSymbolType.Dot">
+      <summary>This type/member supports the .NET Framework infrastructure and is not intended to be used directly from your code.</summary>
+    </member>
+    <member name="F:System.Web.Razor.Tokenizer.Symbols.VBSymbolType.Equal">
+      <summary>This type/member supports the .NET Framework infrastructure and is not intended to be used directly from your code.</summary>
+    </member>
+    <member name="F:System.Web.Razor.Tokenizer.Symbols.VBSymbolType.Exponentiation">
+      <summary>This type/member supports the .NET Framework infrastructure and is not intended to be used directly from your code.</summary>
+    </member>
+    <member name="F:System.Web.Razor.Tokenizer.Symbols.VBSymbolType.FloatingPointLiteral">
+      <summary>This type/member supports the .NET Framework infrastructure and is not intended to be used directly from your code.</summary>
+    </member>
+    <member name="F:System.Web.Razor.Tokenizer.Symbols.VBSymbolType.GreaterThan">
+      <summary>This type/member supports the .NET Framework infrastructure and is not intended to be used directly from your code.</summary>
+    </member>
+    <member name="F:System.Web.Razor.Tokenizer.Symbols.VBSymbolType.Hash">
+      <summary>This type/member supports the .NET Framework infrastructure and is not intended to be used directly from your code.</summary>
+    </member>
+    <member name="F:System.Web.Razor.Tokenizer.Symbols.VBSymbolType.Identifier">
+      <summary>This type/member supports the .NET Framework infrastructure and is not intended to be used directly from your code.</summary>
+    </member>
+    <member name="F:System.Web.Razor.Tokenizer.Symbols.VBSymbolType.IntegerDivide">
+      <summary>This type/member supports the .NET Framework infrastructure and is not intended to be used directly from your code.</summary>
+    </member>
+    <member name="F:System.Web.Razor.Tokenizer.Symbols.VBSymbolType.IntegerLiteral">
+      <summary>This type/member supports the .NET Framework infrastructure and is not intended to be used directly from your code.</summary>
+    </member>
+    <member name="F:System.Web.Razor.Tokenizer.Symbols.VBSymbolType.Keyword">
+      <summary>This type/member supports the .NET Framework infrastructure and is not intended to be used directly from your code.</summary>
+    </member>
+    <member name="F:System.Web.Razor.Tokenizer.Symbols.VBSymbolType.LeftBrace">
+      <summary>This type/member supports the .NET Framework infrastructure and is not intended to be used directly from your code.</summary>
+    </member>
+    <member name="F:System.Web.Razor.Tokenizer.Symbols.VBSymbolType.LeftBracket">
+      <summary>This type/member supports the .NET Framework infrastructure and is not intended to be used directly from your code.</summary>
+    </member>
+    <member name="F:System.Web.Razor.Tokenizer.Symbols.VBSymbolType.LeftParenthesis">
+      <summary>This type/member supports the .NET Framework infrastructure and is not intended to be used directly from your code.</summary>
+    </member>
+    <member name="F:System.Web.Razor.Tokenizer.Symbols.VBSymbolType.LessThan">
+      <summary>This type/member supports the .NET Framework infrastructure and is not intended to be used directly from your code.</summary>
+    </member>
+    <member name="F:System.Web.Razor.Tokenizer.Symbols.VBSymbolType.LineContinuation">
+      <summary>This type/member supports the .NET Framework infrastructure and is not intended to be used directly from your code.</summary>
+    </member>
+    <member name="F:System.Web.Razor.Tokenizer.Symbols.VBSymbolType.Multiply">
+      <summary>This type/member supports the .NET Framework infrastructure and is not intended to be used directly from your code.</summary>
+    </member>
+    <member name="F:System.Web.Razor.Tokenizer.Symbols.VBSymbolType.NewLine">
+      <summary>This type/member supports the .NET Framework infrastructure and is not intended to be used directly from your code.</summary>
+    </member>
+    <member name="F:System.Web.Razor.Tokenizer.Symbols.VBSymbolType.QuestionMark">
+      <summary>This type/member supports the .NET Framework infrastructure and is not intended to be used directly from your code.</summary>
+    </member>
+    <member name="F:System.Web.Razor.Tokenizer.Symbols.VBSymbolType.RazorComment">
+      <summary>This type/member supports the .NET Framework infrastructure and is not intended to be used directly from your code.</summary>
+    </member>
+    <member name="F:System.Web.Razor.Tokenizer.Symbols.VBSymbolType.RazorCommentStar">
+      <summary>This type/member supports the .NET Framework infrastructure and is not intended to be used directly from your code.</summary>
+    </member>
+    <member name="F:System.Web.Razor.Tokenizer.Symbols.VBSymbolType.RazorCommentTransition">
+      <summary>This type/member supports the .NET Framework infrastructure and is not intended to be used directly from your code.</summary>
+    </member>
+    <member name="F:System.Web.Razor.Tokenizer.Symbols.VBSymbolType.RightBrace">
+      <summary>This type/member supports the .NET Framework infrastructure and is not intended to be used directly from your code.</summary>
+    </member>
+    <member name="F:System.Web.Razor.Tokenizer.Symbols.VBSymbolType.RightBracket">
+      <summary>This type/member supports the .NET Framework infrastructure and is not intended to be used directly from your code.</summary>
+    </member>
+    <member name="F:System.Web.Razor.Tokenizer.Symbols.VBSymbolType.RightParenthesis">
+      <summary>This type/member supports the .NET Framework infrastructure and is not intended to be used directly from your code.</summary>
+    </member>
+    <member name="F:System.Web.Razor.Tokenizer.Symbols.VBSymbolType.StringLiteral">
+      <summary>This type/member supports the .NET Framework infrastructure and is not intended to be used directly from your code.</summary>
+    </member>
+    <member name="F:System.Web.Razor.Tokenizer.Symbols.VBSymbolType.Subtract">
+      <summary>This type/member supports the .NET Framework infrastructure and is not intended to be used directly from your code.</summary>
+    </member>
+    <member name="F:System.Web.Razor.Tokenizer.Symbols.VBSymbolType.Transition">
+      <summary>This type/member supports the .NET Framework infrastructure and is not intended to be used directly from your code.</summary>
+    </member>
+    <member name="F:System.Web.Razor.Tokenizer.Symbols.VBSymbolType.Unknown">
+      <summary>This type/member supports the .NET Framework infrastructure and is not intended to be used directly from your code.</summary>
+    </member>
+    <member name="F:System.Web.Razor.Tokenizer.Symbols.VBSymbolType.WhiteSpace">
+      <summary>This type/member supports the .NET Framework infrastructure and is not intended to be used directly from your code.</summary>
+    </member>
+  </members>
+</doc>
\ No newline at end of file
