Copying //UE4/Dev-Core to //UE4/Main

diff --git a/Engine/Source/Programs/CrashReporter/CrashReportWebSite/packages/Microsoft.AspNet.WebPages.3.2.3/lib/net45/System.Web.Helpers.xml b/Engine/Source/Programs/CrashReporter/CrashReportWebSite/packages/Microsoft.AspNet.WebPages.3.2.3/lib/net45/System.Web.Helpers.xml
new file mode 100644
--- /dev/null
+++ b/Engine/Source/Programs/CrashReporter/CrashReportWebSite/packages/Microsoft.AspNet.WebPages.3.2.3/lib/net45/System.Web.Helpers.xml
@@ -0,0 +1,839 @@
+﻿<?xml version="1.0" encoding="utf-8"?>
+<doc>
+  <assembly>
+    <name>System.Web.Helpers</name>
+  </assembly>
+  <members>
+    <member name="T:System.Web.Helpers.Chart">
+      <summary>Displays data in the form of a graphical chart.</summary>
+    </member>
+    <member name="M:System.Web.Helpers.Chart.#ctor(System.Int32,System.Int32,System.String,System.String)">
+      <summary>Initializes a new instance of the <see cref="T:System.Web.Helpers.Chart" /> class.</summary>
+      <param name="width">The width, in pixels, of the complete chart image.</param>
+      <param name="height">The height, in pixels, of the complete chart image.</param>
+      <param name="theme">(Optional) The template (theme) to apply to the chart.</param>
+      <param name="themePath">(Optional) The template (theme) path and file name to apply to the chart.</param>
+    </member>
+    <member name="M:System.Web.Helpers.Chart.AddLegend(System.String,System.String)">
+      <summary>Adds a legend to the chart.</summary>
+      <returns>The chart.</returns>
+      <param name="title">The text of the legend title.</param>
+      <param name="name">The unique name of the legend.</param>
+    </member>
+    <member name="M:System.Web.Helpers.Chart.AddSeries(System.String,System.String,System.String,System.String,System.String,System.Int32,System.Collections.IEnumerable,System.String,System.Collections.IEnumerable,System.String)">
+      <summary>Provides data points and series attributes for the chart.</summary>
+      <returns>The chart.</returns>
+      <param name="name">The unique name of the series.</param>
+      <param name="chartType">The chart type of a series.</param>
+      <param name="chartArea">The name of the chart area that is used to plot the data series.</param>
+      <param name="axisLabel">The axis label text for the series.</param>
+      <param name="legend">The name of the series that is associated with the legend.</param>
+      <param name="markerStep">The granularity of data point markers.</param>
+      <param name="xValue">The values to plot along the x-axis.</param>
+      <param name="xField">The name of the field for x-values.</param>
+      <param name="yValues">The values to plot along the y-axis.</param>
+      <param name="yFields">A comma-separated list of name or names of the field or fields for y-values.</param>
+    </member>
+    <member name="M:System.Web.Helpers.Chart.AddTitle(System.String,System.String)">
+      <summary>Adds a title to the chart.</summary>
+      <returns>The chart.</returns>
+      <param name="text">The title text.</param>
+      <param name="name">The unique name of the title.</param>
+    </member>
+    <member name="M:System.Web.Helpers.Chart.DataBindCrossTable(System.Collections.IEnumerable,System.String,System.String,System.String,System.String,System.String)">
+      <summary>Binds a chart to a data table, where one series is created for each unique value in a column.</summary>
+      <returns>The chart.</returns>
+      <param name="dataSource">The chart data source.</param>
+      <param name="groupByField">The name of the column that is used to group data into the series.</param>
+      <param name="xField">The name of the column for x-values.</param>
+      <param name="yFields">A comma-separated list of names of the columns for y-values.</param>
+      <param name="otherFields">Other data point properties that can be bound.</param>
+      <param name="pointSortOrder">The order in which the series will be sorted. The default is "Ascending".</param>
+    </member>
+    <member name="M:System.Web.Helpers.Chart.DataBindTable(System.Collections.IEnumerable,System.String)">
+      <summary>Creates and binds series data to the specified data table, and optionally populates multiple x-values.</summary>
+      <returns>The chart.</returns>
+      <param name="dataSource">The chart data source. This can be can be any <see cref="T:System.Collections.IEnumerable" /> object.</param>
+      <param name="xField">The name of the table column used for the series x-values.</param>
+    </member>
+    <member name="P:System.Web.Helpers.Chart.FileName">
+      <summary>Gets or sets the name of the file that contains the chart image.</summary>
+      <returns>The name of the file.</returns>
+    </member>
+    <member name="M:System.Web.Helpers.Chart.GetBytes(System.String)">
+      <summary>Returns a chart image as a byte array.</summary>
+      <returns>The chart.</returns>
+      <param name="format">The image format. The default is "jpeg".</param>
+    </member>
+    <member name="M:System.Web.Helpers.Chart.GetFromCache(System.String)">
+      <summary>Retrieves the specified chart from the cache.</summary>
+      <returns>The chart.</returns>
+      <param name="key">The ID of the cache item that contains the chart to retrieve. The key is set when you call the <see cref="M:System.Web.Helpers.Chart.SaveToCache(System.String,System.Int32,System.Boolean)" /> method.</param>
+    </member>
+    <member name="P:System.Web.Helpers.Chart.Height">
+      <summary>Gets or sets the height, in pixels, of the chart image.</summary>
+      <returns>The chart height.</returns>
+    </member>
+    <member name="M:System.Web.Helpers.Chart.Save(System.String,System.String)">
+      <summary>Saves a chart image to the specified file.</summary>
+      <returns>The chart.</returns>
+      <param name="path">The location and name of the image file.</param>
+      <param name="format">The image file format, such as "png" or "jpeg".</param>
+    </member>
+    <member name="M:System.Web.Helpers.Chart.SaveToCache(System.String,System.Int32,System.Boolean)">
+      <summary>Saves a chart in the system cache.</summary>
+      <returns>The ID of the cache item that contains the chart.</returns>
+      <param name="key">The ID of the chart in the cache.</param>
+      <param name="minutesToCache">The number of minutes to keep the chart image in the cache. The default is 20.</param>
+      <param name="slidingExpiration">true to indicate that the chart cache item's expiration is reset each time the item is accessed, or false to indicate that the expiration is based on an absolute interval since the time that the item was added to the cache. The default is true.</param>
+    </member>
+    <member name="M:System.Web.Helpers.Chart.SaveXml(System.String)">
+      <summary>Saves a chart as an XML file.</summary>
+      <returns>The chart.</returns>
+      <param name="path">The path and name of the XML file.</param>
+    </member>
+    <member name="M:System.Web.Helpers.Chart.SetXAxis(System.String,System.Double,System.Double)">
+      <summary>Sets values for the horizontal axis.</summary>
+      <returns>The chart.</returns>
+      <param name="title">The title of the x-axis.</param>
+      <param name="min">The minimum value for the x-axis.</param>
+      <param name="max">The maximum value for the x-axis.</param>
+    </member>
+    <member name="M:System.Web.Helpers.Chart.SetYAxis(System.String,System.Double,System.Double)">
+      <summary>Sets values for the vertical axis.</summary>
+      <returns>The chart.</returns>
+      <param name="title">The title of the y-axis.</param>
+      <param name="min">The minimum value for the y-axis.</param>
+      <param name="max">The maximum value for the y-axis.</param>
+    </member>
+    <member name="M:System.Web.Helpers.Chart.ToWebImage(System.String)">
+      <summary>Creates a <see cref="T:System.Web.Helpers.WebImage" /> object based on the current <see cref="T:System.Web.Helpers.Chart" /> object.</summary>
+      <returns>The chart.</returns>
+      <param name="format">The format of the image to save the <see cref="T:System.Web.Helpers.WebImage" /> object as. The default is "jpeg". The <paramref name="format" /> parameter is not case sensitive.</param>
+    </member>
+    <member name="P:System.Web.Helpers.Chart.Width">
+      <summary>Gets or set the width, in pixels, of the chart image.</summary>
+      <returns>The chart width.</returns>
+    </member>
+    <member name="M:System.Web.Helpers.Chart.Write(System.String)">
+      <summary>Renders the output of the <see cref="T:System.Web.Helpers.Chart" /> object as an image.</summary>
+      <returns>The chart.</returns>
+      <param name="format">The format of the image. The default is "jpeg".</param>
+    </member>
+    <member name="M:System.Web.Helpers.Chart.WriteFromCache(System.String,System.String)">
+      <summary>Renders the output of a <see cref="T:System.Web.Helpers.Chart" /> object that has been cached as an image.</summary>
+      <returns>The chart.</returns>
+      <param name="key">The ID of the chart in the cache.</param>
+      <param name="format">The format of the image. The default is "jpeg".</param>
+    </member>
+    <member name="T:System.Web.Helpers.ChartTheme">
+      <summary>Specifies visual themes for a <see cref="T:System.Web.Helpers.Chart" /> object.</summary>
+    </member>
+    <member name="F:System.Web.Helpers.ChartTheme.Blue">
+      <summary>A theme for 2D charting that features a visual container with a blue gradient, rounded edges, drop-shadowing, and high-contrast gridlines.</summary>
+    </member>
+    <member name="F:System.Web.Helpers.ChartTheme.Green">
+      <summary>A theme for 2D charting that features a visual container with a green gradient, rounded edges, drop-shadowing, and low-contrast gridlines.</summary>
+    </member>
+    <member name="F:System.Web.Helpers.ChartTheme.Vanilla">
+      <summary>A theme for 2D charting that features no visual container and no gridlines.</summary>
+    </member>
+    <member name="F:System.Web.Helpers.ChartTheme.Vanilla3D">
+      <summary>A theme for 3D charting that features no visual container, limited labeling and, sparse, high-contrast gridlines.</summary>
+    </member>
+    <member name="F:System.Web.Helpers.ChartTheme.Yellow">
+      <summary>A theme for 2D charting that features a visual container that has a yellow gradient, rounded edges, drop-shadowing, and high-contrast gridlines.</summary>
+    </member>
+    <member name="T:System.Web.Helpers.Crypto">
+      <summary>Provides methods to generate hash values and encrypt passwords or other sensitive data.</summary>
+    </member>
+    <member name="M:System.Web.Helpers.Crypto.GenerateSalt(System.Int32)">
+      <summary>Generates a cryptographically strong sequence of random byte values.</summary>
+      <returns>The generated salt value as a base-64-encoded string.</returns>
+      <param name="byteLength">The number of cryptographically random bytes to generate.</param>
+    </member>
+    <member name="M:System.Web.Helpers.Crypto.Hash(System.Byte[],System.String)">
+      <summary>Returns a hash value for the specified byte array.</summary>
+      <returns>The hash value for <paramref name="input" /> as a string of hexadecimal characters.</returns>
+      <param name="input">The data to provide a hash value for.</param>
+      <param name="algorithm">The algorithm that is used to generate the hash value. The default is "sha256".</param>
+      <exception cref="T:System.ArgumentNullException">
+        <paramref name="input" /> is null.</exception>
+    </member>
+    <member name="M:System.Web.Helpers.Crypto.Hash(System.String,System.String)">
+      <summary>Returns a hash value for the specified string.</summary>
+      <returns>The hash value for <paramref name="input" /> as a string of hexadecimal characters.</returns>
+      <param name="input">The data to provide a hash value for.</param>
+      <param name="algorithm">The algorithm that is used to generate the hash value. The default is "sha256".</param>
+      <exception cref="T:System.ArgumentNullException">
+        <paramref name="input" /> is null.</exception>
+    </member>
+    <member name="M:System.Web.Helpers.Crypto.HashPassword(System.String)">
+      <summary>Returns an RFC 2898 hash value for the specified password.</summary>
+      <returns>The hash value for <paramref name="password" /> as a base-64-encoded string.</returns>
+      <param name="password">The password to generate a hash value for.</param>
+      <exception cref="T:System.ArgumentNullException">
+        <paramref name="password" /> is null.</exception>
+    </member>
+    <member name="M:System.Web.Helpers.Crypto.SHA1(System.String)">
+      <summary>Returns a SHA-1 hash value for the specified string.</summary>
+      <returns>The SHA-1 hash value for <paramref name="input" /> as a string of hexadecimal characters.</returns>
+      <param name="input">The data to provide a hash value for.</param>
+      <exception cref="T:System.ArgumentNullException">
+        <paramref name="input" /> is null.</exception>
+    </member>
+    <member name="M:System.Web.Helpers.Crypto.SHA256(System.String)">
+      <summary>Returns a SHA-256 hash value for the specified string.</summary>
+      <returns>The SHA-256 hash value for <paramref name="input" /> as a string of hexadecimal characters.</returns>
+      <param name="input">The data to provide a hash value for.</param>
+      <exception cref="T:System.ArgumentNullException">
+        <paramref name="input" /> is null.</exception>
+    </member>
+    <member name="M:System.Web.Helpers.Crypto.VerifyHashedPassword(System.String,System.String)">
+      <summary>Determines whether the specified RFC 2898 hash and password are a cryptographic match.</summary>
+      <returns>true if the hash value is a cryptographic match for the password; otherwise, false.</returns>
+      <param name="hashedPassword">The previously-computed RFC 2898 hash value as a base-64-encoded string.</param>
+      <param name="password">The plaintext password to cryptographically compare with <paramref name="hashedPassword" />.</param>
+      <exception cref="T:System.ArgumentNullException">
+        <paramref name="hashedPassword" /> or <paramref name="password" /> is null.</exception>
+    </member>
+    <member name="T:System.Web.Helpers.DynamicJsonArray">
+      <summary>Represents a series of values as a JavaScript-like array by using the dynamic capabilities of the Dynamic Language Runtime (DLR).</summary>
+    </member>
+    <member name="M:System.Web.Helpers.DynamicJsonArray.#ctor(System.Object[])">
+      <summary>Initializes a new instance of the <see cref="T:System.Web.Helpers.DynamicJsonArray" /> class using the specified array element values.</summary>
+      <param name="arrayValues">An array of objects that contains the values to add to the <see cref="T:System.Web.Helpers.DynamicJsonArray" /> instance.</param>
+    </member>
+    <member name="M:System.Web.Helpers.DynamicJsonArray.GetEnumerator">
+      <summary>Returns an enumerator that can be used to iterate through the elements of the <see cref="T:System.Web.Helpers.DynamicJsonArray" /> instance.</summary>
+      <returns>An enumerator that can be used to iterate through the elements of the JSON array.</returns>
+    </member>
+    <member name="P:System.Web.Helpers.DynamicJsonArray.Item(System.Int32)">
+      <summary>Returns the value at the specified index in the <see cref="T:System.Web.Helpers.DynamicJsonArray" /> instance.</summary>
+      <returns>The value at the specified index.</returns>
+    </member>
+    <member name="P:System.Web.Helpers.DynamicJsonArray.Length">
+      <summary>Returns the number of elements in the <see cref="T:System.Web.Helpers.DynamicJsonArray" /> instance.</summary>
+      <returns>The number of elements in the JSON array.</returns>
+    </member>
+    <member name="M:System.Web.Helpers.DynamicJsonArray.op_Implicit(System.Web.Helpers.DynamicJsonArray)~System.Object[]">
+      <summary>Converts a <see cref="T:System.Web.Helpers.DynamicJsonArray" /> instance to an array of objects.</summary>
+      <returns>The array of objects that represents the JSON array.</returns>
+      <param name="obj">The JSON array to convert.</param>
+    </member>
+    <member name="M:System.Web.Helpers.DynamicJsonArray.op_Implicit(System.Web.Helpers.DynamicJsonArray)~System.Array">
+      <summary>Converts a <see cref="T:System.Web.Helpers.DynamicJsonArray" /> instance to an array of objects.</summary>
+      <returns>The array of objects that represents the JSON array.</returns>
+      <param name="obj">The JSON array to convert.</param>
+    </member>
+    <member name="M:System.Web.Helpers.DynamicJsonArray.System#Collections#Generic#IEnumerable{T}#GetEnumerator">
+      <summary>Returns an enumerator that can be used to iterate through a collection.</summary>
+      <returns>An enumerator that can be used to iterate through the collection.</returns>
+    </member>
+    <member name="M:System.Web.Helpers.DynamicJsonArray.TryConvert(System.Dynamic.ConvertBinder,System.Object@)">
+      <summary>Converts the <see cref="T:System.Web.Helpers.DynamicJsonArray" /> instance to a compatible type.</summary>
+      <returns>true if the conversion was successful; otherwise, false.</returns>
+      <param name="binder">Provides information about the conversion operation.</param>
+      <param name="result">When this method returns, contains the result of the type conversion operation. This parameter is passed uninitialized.</param>
+    </member>
+    <member name="M:System.Web.Helpers.DynamicJsonArray.TryGetMember(System.Dynamic.GetMemberBinder,System.Object@)">
+      <summary>Tests the <see cref="T:System.Web.Helpers.DynamicJsonArray" /> instance for dynamic members (which are not supported) in a way that does not cause an exception to be thrown.</summary>
+      <returns>true in all cases.</returns>
+      <param name="binder">Provides information about the get operation.</param>
+      <param name="result">When this method returns, contains null. This parameter is passed uninitialized.</param>
+    </member>
+    <member name="T:System.Web.Helpers.DynamicJsonObject">
+      <summary>Represents a collection of values as a JavaScript-like object by using the capabilities of the Dynamic Language Runtime.</summary>
+    </member>
+    <member name="M:System.Web.Helpers.DynamicJsonObject.#ctor(System.Collections.Generic.IDictionary{System.String,System.Object})">
+      <summary>Initializes a new instance of the <see cref="T:System.Web.Helpers.DynamicJsonObject" /> class using the specified field values.</summary>
+      <param name="values">A dictionary of property names and values to add to the <see cref="T:System.Web.Helpers.DynamicJsonObject" /> instance as dynamic members.</param>
+    </member>
+    <member name="M:System.Web.Helpers.DynamicJsonObject.GetDynamicMemberNames">
+      <summary>Returns a list that contains the name of all dynamic members (JSON fields) of the <see cref="T:System.Web.Helpers.DynamicJsonObject" /> instance.</summary>
+      <returns>A list that contains the name of every dynamic member (JSON field).</returns>
+    </member>
+    <member name="M:System.Web.Helpers.DynamicJsonObject.TryConvert(System.Dynamic.ConvertBinder,System.Object@)">
+      <summary>Converts the <see cref="T:System.Web.Helpers.DynamicJsonObject" /> instance to a compatible type.</summary>
+      <returns>true in all cases.</returns>
+      <param name="binder">Provides information about the conversion operation.</param>
+      <param name="result">When this method returns, contains the result of the type conversion operation. This parameter is passed uninitialized.</param>
+      <exception cref="T:System.InvalidOperationException">The <see cref="T:System.Web.Helpers.DynamicJsonObject" /> instance could not be converted to the specified type.</exception>
+    </member>
+    <member name="M:System.Web.Helpers.DynamicJsonObject.TryGetIndex(System.Dynamic.GetIndexBinder,System.Object[],System.Object@)">
+      <summary>Gets the value of a <see cref="T:System.Web.Helpers.DynamicJsonObject" /> field using the specified index.</summary>
+      <returns>true in all cases.</returns>
+      <param name="binder">Provides information about the indexed get operation.</param>
+      <param name="indexes">An array that contains a single object that indexes the field by name. The object must be convertible to a string that specifies the name of the JSON field to return. If multiple indexes are specified, <paramref name="result" /> contains null when this method returns.</param>
+      <param name="result">When this method returns, contains the value of the indexed field, or null if the get operation was unsuccessful. This parameter is passed uninitialized.</param>
+    </member>
+    <member name="M:System.Web.Helpers.DynamicJsonObject.TryGetMember(System.Dynamic.GetMemberBinder,System.Object@)">
+      <summary>Gets the value of a <see cref="T:System.Web.Helpers.DynamicJsonObject" /> field using the specified name.</summary>
+      <returns>true in all cases.</returns>
+      <param name="binder">Provides information about the get operation.</param>
+      <param name="result">When this method returns, contains the value of the field, or null if the get operation was unsuccessful. This parameter is passed uninitialized.</param>
+    </member>
+    <member name="M:System.Web.Helpers.DynamicJsonObject.TrySetIndex(System.Dynamic.SetIndexBinder,System.Object[],System.Object)">
+      <summary>Sets the value of a <see cref="T:System.Web.Helpers.DynamicJsonObject" /> field using the specified index.</summary>
+      <returns>true in all cases.</returns>
+      <param name="binder">Provides information about the indexed set operation.</param>
+      <param name="indexes">An array that contains a single object that indexes the field by name. The object must be convertible to a string that specifies the name of the JSON field to return. If multiple indexes are specified, no field is changed or added.</param>
+      <param name="value">The value to set the field to.</param>
+    </member>
+    <member name="M:System.Web.Helpers.DynamicJsonObject.TrySetMember(System.Dynamic.SetMemberBinder,System.Object)">
+      <summary>Sets the value of a <see cref="T:System.Web.Helpers.DynamicJsonObject" /> field using the specified name.</summary>
+      <returns>true in all cases.</returns>
+      <param name="binder">Provides information about the set operation.</param>
+      <param name="value">The value to set the field to.</param>
+    </member>
+    <member name="T:System.Web.Helpers.Json">
+      <summary>Provides methods for working with data in JavaScript Object Notation (JSON) format.</summary>
+    </member>
+    <member name="M:System.Web.Helpers.Json.Decode``1(System.String)">
+      <summary>Converts data in JavaScript Object Notation (JSON) format into the specified strongly typed data list.</summary>
+      <returns>The JSON-encoded data converted to a strongly typed list.</returns>
+      <param name="value">The JSON-encoded string to convert.</param>
+      <typeparam name="T">The type of the strongly typed list to convert JSON data into.</typeparam>
+    </member>
+    <member name="M:System.Web.Helpers.Json.Decode(System.String)">
+      <summary>Converts data in JavaScript Object Notation (JSON) format into a data object.</summary>
+      <returns>The JSON-encoded data converted to a data object.</returns>
+      <param name="value">The JSON-encoded string to convert.</param>
+    </member>
+    <member name="M:System.Web.Helpers.Json.Decode(System.String,System.Type)">
+      <summary>Converts data in JavaScript Object Notation (JSON) format into a data object of a specified type.</summary>
+      <returns>The JSON-encoded data converted to the specified type.</returns>
+      <param name="value">The JSON-encoded string to convert.</param>
+      <param name="targetType">The type that the <paramref name="value" /> data should be converted to.</param>
+    </member>
+    <member name="M:System.Web.Helpers.Json.Encode(System.Object)">
+      <summary>Converts a data object to a string that is in the JavaScript Object Notation (JSON) format.</summary>
+      <returns>Returns a string of data converted to the JSON format.</returns>
+      <param name="value">The data object to convert.</param>
+    </member>
+    <member name="M:System.Web.Helpers.Json.Write(System.Object,System.IO.TextWriter)">
+      <summary>Converts a data object to a string in JavaScript Object Notation (JSON) format and adds the string to the specified <see cref="T:System.IO.TextWriter" /> object.</summary>
+      <param name="value">The data object to convert.</param>
+      <param name="writer">The object that contains the converted JSON data.</param>
+    </member>
+    <member name="T:System.Web.Helpers.ObjectInfo">
+      <summary>Renders the property names and values of the specified object and of any subobjects that it references.</summary>
+    </member>
+    <member name="M:System.Web.Helpers.ObjectInfo.Print(System.Object,System.Int32,System.Int32)">
+      <summary>Renders the property names and values of the specified object and of any subobjects.</summary>
+      <returns>For a simple variable, returns the type and the value. For an object that contains multiple items, returns the property name or key and the value for each property.</returns>
+      <param name="value">The object to render information for.</param>
+      <param name="depth">Optional. Specifies the depth of nested subobjects to render information for. The default is 10.</param>
+      <param name="enumerationLength">Optional. Specifies the maximum number of characters that the method displays for object values. The default is 1000.</param>
+      <exception cref="T:System.ArgumentOutOfRangeException">
+        <paramref name="depth" /> is less than zero.</exception>
+      <exception cref="T:System.ArgumentOutOfRangeException">
+        <paramref name="enumerationLength" /> is less than or equal to zero.</exception>
+    </member>
+    <member name="T:System.Web.Helpers.ServerInfo">
+      <summary>Displays information about the web server environment that hosts the current web page.</summary>
+    </member>
+    <member name="M:System.Web.Helpers.ServerInfo.GetHtml">
+      <summary>Displays information about the web server environment.</summary>
+      <returns>A string of name-value pairs that contains information about the web server. </returns>
+    </member>
+    <member name="T:System.Web.Helpers.SortDirection">
+      <summary>Specifies the direction in which to sort a list of items.</summary>
+    </member>
+    <member name="F:System.Web.Helpers.SortDirection.Ascending">
+      <summary>Sort from smallest to largest —for example, from 1 to 10.</summary>
+    </member>
+    <member name="F:System.Web.Helpers.SortDirection.Descending">
+      <summary>Sort from largest to smallest — for example, from 10 to 1.</summary>
+    </member>
+    <member name="T:System.Web.Helpers.WebCache">
+      <summary>Provides a cache to store frequently accessed data.</summary>
+    </member>
+    <member name="M:System.Web.Helpers.WebCache.Get(System.String)">
+      <summary>Retrieves the specified item from the <see cref="T:System.Web.Helpers.WebCache" /> object.</summary>
+      <returns>The item retrieved from the cache, or null if the item is not found.</returns>
+      <param name="key">The identifier for the cache item to retrieve.</param>
+    </member>
+    <member name="M:System.Web.Helpers.WebCache.Remove(System.String)">
+      <summary>Removes the specified item from the <see cref="T:System.Web.Helpers.WebCache" /> object.</summary>
+      <returns>The item removed from the <see cref="T:System.Web.Helpers.WebCache" /> object. If the item is not found, returns null.</returns>
+      <param name="key">The identifier for the cache item to remove.</param>
+    </member>
+    <member name="M:System.Web.Helpers.WebCache.Set(System.String,System.Object,System.Int32,System.Boolean)">
+      <summary>Inserts an item into the <see cref="T:System.Web.Helpers.WebCache" /> object.</summary>
+      <param name="key">The identifier for the cache item.</param>
+      <param name="value">The data to insert into the cache.</param>
+      <param name="minutesToCache">Optional. The number of minutes to keep an item in the cache. The default is 20.</param>
+      <param name="slidingExpiration">Optional. true to indicate that the cache item expiration is reset each time the item is accessed, or false to indicate that the expiration is based the absolute time since the item was added to the cache. The default is true. In that case, if you also use the default value for the <paramref name="minutesToCache" /> parameter, a cached item expires 20 minutes after it was last accessed.</param>
+      <exception cref="T:System.ArgumentOutOfRangeException">The value of <paramref name="minutesToCache" /> is less than or equal to zero.</exception>
+      <exception cref="T:System.ArgumentOutOfRangeException">Sliding expiration is enabled and the value of <paramref name="minutesToCache" /> is greater than a year.</exception>
+    </member>
+    <member name="T:System.Web.Helpers.WebGrid">
+      <summary>Displays data on a web page using an HTML table element.</summary>
+    </member>
+    <member name="M:System.Web.Helpers.WebGrid.#ctor(System.Collections.Generic.IEnumerable{System.Object},System.Collections.Generic.IEnumerable{System.String},System.String,System.Int32,System.Boolean,System.Boolean,System.String,System.String,System.String,System.String,System.String,System.String,System.String)">
+      <summary>Initializes a new instance of the <see cref="T:System.Web.Helpers.WebGrid" /> class.</summary>
+      <param name="source">The data to display.</param>
+      <param name="columnNames">A collection that contains the names of the data columns to display. By default, this value is auto-populated according to the values in the <paramref name="source" /> parameter.</param>
+      <param name="defaultSort">The name of the data column that is used to sort the grid by default.</param>
+      <param name="rowsPerPage">The number of rows that are displayed on each page of the grid when paging is enabled. The default is 10.</param>
+      <param name="canPage">true to specify that paging is enabled for the <see cref="T:System.Web.Helpers.WebGrid" /> instance; otherwise false. The default is true. </param>
+      <param name="canSort">true to specify that sorting is enabled for the <see cref="T:System.Web.Helpers.WebGrid" /> instance; otherwise, false. The default is true.</param>
+      <param name="ajaxUpdateContainerId">The value of the HTML id attribute that is used to mark the HTML element that gets dynamic Ajax updates that are associated with the <see cref="T:System.Web.Helpers.WebGrid" /> instance.</param>
+      <param name="ajaxUpdateCallback">The name of the JavaScript function that is called after the HTML element specified by the <see cref="P:System.Web.Helpers.WebGrid.AjaxUpdateContainerId" /> property has been updated. If the name of a function is not provided, no function will be called. If the specified function does not exist, a JavaScript error will occur if it is invoked.</param>
+      <param name="fieldNamePrefix">The prefix that is applied to all query-string fields that are associated with the <see cref="T:System.Web.Helpers.WebGrid" /> instance. This value is used in order to support multiple <see cref="T:System.Web.Helpers.WebGrid" /> instances on the same web page.</param>
+      <param name="pageFieldName">The name of the query-string field that is used to specify the current page of the <see cref="T:System.Web.Helpers.WebGrid" /> instance.</param>
+      <param name="selectionFieldName">The name of the query-string field that is used to specify the currently selected row of the <see cref="T:System.Web.Helpers.WebGrid" /> instance.</param>
+      <param name="sortFieldName">The name of the query-string field that is used to specify the name of the data column that the <see cref="T:System.Web.Helpers.WebGrid" /> instance is sorted by.</param>
+      <param name="sortDirectionFieldName">The name of the query-string field that is used to specify the direction in which the <see cref="T:System.Web.Helpers.WebGrid" /> instance is sorted.</param>
+    </member>
+    <member name="M:System.Web.Helpers.WebGrid.AddSorter``2(System.String,System.Linq.Expressions.Expression{System.Func{``0,``1}})">
+      <summary>Adds a specific sort function for a given column.</summary>
+      <returns>The current grid, with the new custom sorter applied.</returns>
+      <param name="columnName">The column name (as used for sorting)</param>
+      <param name="keySelector">The function used to select a key to sort by, for each element in the grid's source.</param>
+      <typeparam name="TElement">The type of elements in the grid's source.</typeparam>
+      <typeparam name="TProperty">The column type, usually inferred from the keySelector function's return type.</typeparam>
+    </member>
+    <member name="P:System.Web.Helpers.WebGrid.AjaxUpdateCallback">
+      <summary>Gets the name of the JavaScript function to call after the HTML element that is associated with the <see cref="T:System.Web.Helpers.WebGrid" /> instance has been updated in response to an Ajax update request.</summary>
+      <returns>The name of the function.</returns>
+    </member>
+    <member name="P:System.Web.Helpers.WebGrid.AjaxUpdateContainerId">
+      <summary>Gets the value of the HTML id attribute that marks an HTML element on the web page that gets dynamic Ajax updates that are associated with the <see cref="T:System.Web.Helpers.WebGrid" /> instance.</summary>
+      <returns>The value of the id attribute.</returns>
+    </member>
+    <member name="M:System.Web.Helpers.WebGrid.Bind(System.Collections.Generic.IEnumerable{System.Object},System.Collections.Generic.IEnumerable{System.String},System.Boolean,System.Int32)">
+      <summary>Binds the specified data to the <see cref="T:System.Web.Helpers.WebGrid" /> instance.</summary>
+      <returns>The bound and populated <see cref="T:System.Web.Helpers.WebGrid" /> instance.</returns>
+      <param name="source">The data to display.</param>
+      <param name="columnNames">A collection that contains the names of the data columns to bind.</param>
+      <param name="autoSortAndPage">true to enable sorting and paging of the <see cref="T:System.Web.Helpers.WebGrid" /> instance; otherwise, false.</param>
+      <param name="rowCount">The number of rows to display on each page of the grid.</param>
+    </member>
+    <member name="P:System.Web.Helpers.WebGrid.CanSort">
+      <summary>Gets a value that indicates whether the <see cref="T:System.Web.Helpers.WebGrid" /> instance supports sorting.</summary>
+      <returns>true if the instance supports sorting; otherwise, false.</returns>
+    </member>
+    <member name="M:System.Web.Helpers.WebGrid.Column(System.String,System.String,System.Func{System.Object,System.Object},System.String,System.Boolean)">
+      <summary>Creates a new <see cref="T:System.Web.Helpers.WebGridColumn" /> instance.</summary>
+      <returns>The new column.</returns>
+      <param name="columnName">The name of the data column to associate with the <see cref="T:System.Web.Helpers.WebGridColumn" /> instance.</param>
+      <param name="header">The text that is rendered in the header of the HTML table column that is associated with the <see cref="T:System.Web.Helpers.WebGridColumn" /> instance.</param>
+      <param name="format">The function that is used to format the data values that are associated with the <see cref="T:System.Web.Helpers.WebGridColumn" /> instance.</param>
+      <param name="style">A string that specifies the name of the CSS class that is used to style the HTML table cells that are associated with the <see cref="T:System.Web.Helpers.WebGridColumn" /> instance.</param>
+      <param name="canSort">true to enable sorting in the <see cref="T:System.Web.Helpers.WebGrid" /> instance by the data values that are associated with the <see cref="T:System.Web.Helpers.WebGridColumn" /> instance; otherwise, false. The default is true.</param>
+    </member>
+    <member name="P:System.Web.Helpers.WebGrid.ColumnNames">
+      <summary>Gets a collection that contains the name of each data column that is bound to the <see cref="T:System.Web.Helpers.WebGrid" /> instance.</summary>
+      <returns>The collection of data column names.</returns>
+    </member>
+    <member name="M:System.Web.Helpers.WebGrid.Columns(System.Web.Helpers.WebGridColumn[])">
+      <summary>Returns an array that contains the specified <see cref="T:System.Web.Helpers.WebGridColumn" /> instances.</summary>
+      <returns>An array of columns.</returns>
+      <param name="columnSet">A variable number of <see cref="T:System.Web.Helpers.WebGridColumn" /> column instances.</param>
+    </member>
+    <member name="P:System.Web.Helpers.WebGrid.FieldNamePrefix">
+      <summary>Gets the prefix that is applied to all query-string fields that are associated with the <see cref="T:System.Web.Helpers.WebGrid" /> instance.</summary>
+      <returns>The query-string field prefix of the <see cref="T:System.Web.Helpers.WebGrid" /> instance.</returns>
+    </member>
+    <member name="M:System.Web.Helpers.WebGrid.GetContainerUpdateScript(System.String)">
+      <summary>Returns a JavaScript statement that can be used to update the HTML element that is associated with the <see cref="T:System.Web.Helpers.WebGrid" /> instance on the specified web page.</summary>
+      <returns>A JavaScript statement that can be used to update the HTML element in a web page that is associated with the <see cref="T:System.Web.Helpers.WebGrid" /> instance.</returns>
+      <param name="path">The URL of the web page that contains the <see cref="T:System.Web.Helpers.WebGrid" /> instance that is being updated. The URL can include query-string arguments.</param>
+    </member>
+    <member name="M:System.Web.Helpers.WebGrid.GetHtml(System.String,System.String,System.String,System.String,System.String,System.String,System.String,System.Boolean,System.Boolean,System.String,System.Collections.Generic.IEnumerable{System.Web.Helpers.WebGridColumn},System.Collections.Generic.IEnumerable{System.String},System.Web.Helpers.WebGridPagerModes,System.String,System.String,System.String,System.String,System.Int32,System.Object)">
+      <summary>Returns the HTML markup that is used to render the <see cref="T:System.Web.Helpers.WebGrid" /> instance and using the specified paging options.</summary>
+      <returns>The HTML markup that represents the fully-populated <see cref="T:System.Web.Helpers.WebGrid" /> instance.</returns>
+      <param name="tableStyle">The name of the CSS class that is used to style the whole table.</param>
+      <param name="headerStyle">The name of the CSS class that is used to style the table header.</param>
+      <param name="footerStyle">The name of the CSS class that is used to style the table footer.</param>
+      <param name="rowStyle">The name of the CSS class that is used to style each table row.</param>
+      <param name="alternatingRowStyle">The name of the CSS class that is used to style even-numbered table rows.</param>
+      <param name="selectedRowStyle">The name of the CSS class that is used to style the selected table row. (Only one row can be selected at a time.)</param>
+      <param name="caption">The table caption.</param>
+      <param name="displayHeader">true to display the table header; otherwise, false. The default is true.</param>
+      <param name="fillEmptyRows">true to insert additional rows in the last page when there are insufficient data items to fill the last page; otherwise, false. The default is false. Additional rows are populated using the text specified by the <paramref name="emptyRowCellValue" /> parameter.</param>
+      <param name="emptyRowCellValue">The text that is used to populate additional rows in a page when there are insufficient data items to fill the last page. The <paramref name="fillEmptyRows" /> parameter must be set to true to display these additional rows.</param>
+      <param name="columns">A collection of <see cref="T:System.Web.Helpers.WebGridColumn" /> instances that specify how each column is displayed. This includes which data column is associated with each grid column, and how to format the data values that each grid column contains.</param>
+      <param name="exclusions">A collection that contains the names of the data columns to exclude when the grid auto-populates columns.</param>
+      <param name="mode">A bitwise combination of the enumeration values that specify methods that are provided for moving between pages of the <see cref="T:System.Web.Helpers.WebGrid" /> instance.</param>
+      <param name="firstText">The text for the HTML link element that is used to link to the first page of the <see cref="T:System.Web.Helpers.WebGrid" /> instance. The <see cref="F:System.Web.Helpers.WebGridPagerModes.FirstLast" /> flag of the <paramref name="mode" /> parameter must be set to display this page navigation element.</param>
+      <param name="previousText">The text for the HTML link element that is used to link to previous page of the <see cref="T:System.Web.Helpers.WebGrid" /> instance. The <see cref="F:System.Web.Helpers.WebGridPagerModes.NextPrevious" /> flag of the <paramref name="mode" /> parameter must be set to display this page navigation element.</param>
+      <param name="nextText">The text for the HTML link element that is used to link to the next page of the <see cref="T:System.Web.Helpers.WebGrid" /> instance. The <see cref="F:System.Web.Helpers.WebGridPagerModes.NextPrevious" /> flag of the <paramref name="mode" /> parameter must be set to display this page navigation element.</param>
+      <param name="lastText">The text for the HTML link element that is used to link to the last page of the <see cref="T:System.Web.Helpers.WebGrid" /> instance. The <see cref="F:System.Web.Helpers.WebGridPagerModes.FirstLast" /> flag of the <paramref name="mode" /> parameter must be set to display this page navigation element.</param>
+      <param name="numericLinksCount">The number of numeric page links that are provided to nearby <see cref="T:System.Web.Helpers.WebGrid" /> pages. The text of each numeric page link contains the page number. The <see cref="F:System.Web.Helpers.WebGridPagerModes.Numeric" /> flag of the <paramref name="mode" /> parameter must be set to display these page navigation elements.</param>
+      <param name="htmlAttributes">An object that represents a collection of attributes (names and values) to set for the HTML table element that represents the <see cref="T:System.Web.Helpers.WebGrid" /> instance.</param>
+    </member>
+    <member name="M:System.Web.Helpers.WebGrid.GetPageUrl(System.Int32)">
+      <summary>Returns a URL that can be used to display the specified data page of the <see cref="T:System.Web.Helpers.WebGrid" /> instance.</summary>
+      <returns>A URL that can be used to display the specified data page of the grid.</returns>
+      <param name="pageIndex">The index of the <see cref="T:System.Web.Helpers.WebGrid" /> page to display.</param>
+    </member>
+    <member name="M:System.Web.Helpers.WebGrid.GetSortUrl(System.String)">
+      <summary>Returns a URL that can be used to sort the <see cref="T:System.Web.Helpers.WebGrid" /> instance by the specified column.</summary>
+      <returns>A URL that can be used to sort the grid.</returns>
+      <param name="column">The name of the data column to sort by.</param>
+    </member>
+    <member name="P:System.Web.Helpers.WebGrid.HasSelection">
+      <summary>Gets a value that indicates whether a row in the <see cref="T:System.Web.Helpers.WebGrid" /> instance is selected.</summary>
+      <returns>true if a row is currently selected; otherwise, false.</returns>
+    </member>
+    <member name="P:System.Web.Helpers.WebGrid.IsAjaxEnabled">
+      <summary>Returns a value that indicates whether the <see cref="T:System.Web.Helpers.WebGrid" /> instance can use Ajax calls to refresh the display.</summary>
+      <returns>true if the instance supports Ajax calls; otherwise, false..</returns>
+    </member>
+    <member name="P:System.Web.Helpers.WebGrid.PageCount">
+      <summary>Gets the number of pages that the <see cref="T:System.Web.Helpers.WebGrid" /> instance contains.</summary>
+      <returns>The page count.</returns>
+    </member>
+    <member name="P:System.Web.Helpers.WebGrid.PageFieldName">
+      <summary>Gets the full name of the query-string field that is used to specify the current page of the <see cref="T:System.Web.Helpers.WebGrid" /> instance.</summary>
+      <returns>The full name of the query string field that is used to specify the current page of the grid.</returns>
+    </member>
+    <member name="P:System.Web.Helpers.WebGrid.PageIndex">
+      <summary>Gets or sets the index of the current page of the <see cref="T:System.Web.Helpers.WebGrid" /> instance.</summary>
+      <returns>The index of the current page.</returns>
+    </member>
+    <member name="M:System.Web.Helpers.WebGrid.Pager(System.Web.Helpers.WebGridPagerModes,System.String,System.String,System.String,System.String,System.Int32)">
+      <summary>Returns the HTML markup that is used to provide the specified paging support for the <see cref="T:System.Web.Helpers.WebGrid" /> instance.</summary>
+      <returns>The HTML markup that provides paging support for the grid.</returns>
+      <param name="mode">A bitwise combination of the enumeration values that specify the methods that are provided for moving between the pages of the grid. The default is the bitwise OR of the <see cref="F:System.Web.Helpers.WebGridPagerModes.NextPrevious" /> and <see cref="F:System.Web.Helpers.WebGridPagerModes.Numeric" /> flags.</param>
+      <param name="firstText">The text for the HTML link element that navigates to the first page of the grid.</param>
+      <param name="previousText">The text for the HTML link element that navigates to the previous page of the grid.</param>
+      <param name="nextText">The text for the HTML link element that navigates to the next page of the grid.</param>
+      <param name="lastText">The text for the HTML link element that navigates to the last page of the grid.</param>
+      <param name="numericLinksCount">The number of numeric page links to display. The default is 5.</param>
+    </member>
+    <member name="P:System.Web.Helpers.WebGrid.Rows">
+      <summary>Gets a list that contains the rows that are on the current page of  the <see cref="T:System.Web.Helpers.WebGrid" /> instance after the grid has been sorted.</summary>
+      <returns>The list of rows.</returns>
+    </member>
+    <member name="P:System.Web.Helpers.WebGrid.RowsPerPage">
+      <summary>Gets the number of rows that are displayed on each page of the <see cref="T:System.Web.Helpers.WebGrid" /> instance.</summary>
+      <returns>The number of rows that are displayed on each page of the grid.</returns>
+    </member>
+    <member name="P:System.Web.Helpers.WebGrid.SelectedIndex">
+      <summary>Gets or sets the index of the selected row relative to the current page of the <see cref="T:System.Web.Helpers.WebGrid" /> instance.</summary>
+      <returns>The index of the selected row relative to the current page.</returns>
+    </member>
+    <member name="P:System.Web.Helpers.WebGrid.SelectedRow">
+      <summary>Gets the currently selected row of the <see cref="T:System.Web.Helpers.WebGrid" /> instance.</summary>
+      <returns>The currently selected row.</returns>
+    </member>
+    <member name="P:System.Web.Helpers.WebGrid.SelectionFieldName">
+      <summary>Gets the full name of the query-string field that is used to specify the selected row of the <see cref="T:System.Web.Helpers.WebGrid" /> instance.</summary>
+      <returns>The full name of the query string field that is used to specify the selected row of the grid.</returns>
+    </member>
+    <member name="P:System.Web.Helpers.WebGrid.SortColumn">
+      <summary>Gets or sets the name of the data column that the <see cref="T:System.Web.Helpers.WebGrid" /> instance is sorted by.</summary>
+      <returns>The name of the data column that is used to sort the grid.</returns>
+    </member>
+    <member name="P:System.Web.Helpers.WebGrid.SortDirection">
+      <summary>Gets or sets the direction in which the <see cref="T:System.Web.Helpers.WebGrid" /> instance is sorted.</summary>
+      <returns>The sort direction.</returns>
+    </member>
+    <member name="P:System.Web.Helpers.WebGrid.SortDirectionFieldName">
+      <summary>Gets the full name of the query-string field that is used to specify the sort direction of the <see cref="T:System.Web.Helpers.WebGrid" /> instance.</summary>
+      <returns>The full name of the query string field that is used to specify the sort direction of the grid.</returns>
+    </member>
+    <member name="P:System.Web.Helpers.WebGrid.SortFieldName">
+      <summary>Gets the full name of the query-string field that is used to specify the name of the data column that the <see cref="T:System.Web.Helpers.WebGrid" /> instance is sorted by.</summary>
+      <returns>The full name of the query-string field that is used to specify the name of the data column that the grid is sorted by.</returns>
+    </member>
+    <member name="M:System.Web.Helpers.WebGrid.Table(System.String,System.String,System.String,System.String,System.String,System.String,System.String,System.Boolean,System.Boolean,System.String,System.Collections.Generic.IEnumerable{System.Web.Helpers.WebGridColumn},System.Collections.Generic.IEnumerable{System.String},System.Func{System.Object,System.Object},System.Object)">
+      <summary>Returns the HTML markup that is used to render the <see cref="T:System.Web.Helpers.WebGrid" /> instance.</summary>
+      <returns>The HTML markup that represents the fully-populated <see cref="T:System.Web.Helpers.WebGrid" /> instance.</returns>
+      <param name="tableStyle">The name of the CSS class that is used to style the whole table.</param>
+      <param name="headerStyle">The name of the CSS class that is used to style the table header.</param>
+      <param name="footerStyle">The name of the CSS class that is used to style the table footer.</param>
+      <param name="rowStyle">The name of the CSS class that is used to style each table row.</param>
+      <param name="alternatingRowStyle">The name of the CSS class that is used to style even-numbered table rows.</param>
+      <param name="selectedRowStyle">The name of the CSS class that is used use to style the selected table row.</param>
+      <param name="caption">The table caption.</param>
+      <param name="displayHeader">true to display the table header; otherwise, false. The default is true.</param>
+      <param name="fillEmptyRows">true to insert additional rows in the last page when there are insufficient data items to fill the last page; otherwise, false. The default is false. Additional rows are populated using the text specified by the <paramref name="emptyRowCellValue" /> parameter.</param>
+      <param name="emptyRowCellValue">The text that is used to populate additional rows in the last page when there are insufficient data items to fill the last page. The <paramref name="fillEmptyRows" /> parameter must be set to true to display these additional rows.</param>
+      <param name="columns">A collection of <see cref="T:System.Web.Helpers.WebGridColumn" /> instances that specify how each column is displayed. This includes which data column is associated with each grid column, and how to format the data values that each grid column contains.</param>
+      <param name="exclusions">A collection that contains the names of the data columns to exclude when the grid auto-populates columns.</param>
+      <param name="footer">A function that returns the HTML markup that is used to render the table footer.</param>
+      <param name="htmlAttributes">An object that represents a collection of attributes (names and values) to set for the HTML table element that represents the <see cref="T:System.Web.Helpers.WebGrid" /> instance.</param>
+    </member>
+    <member name="P:System.Web.Helpers.WebGrid.TotalRowCount">
+      <summary>Gets the total number of rows that the <see cref="T:System.Web.Helpers.WebGrid" /> instance contains.</summary>
+      <returns>The total number of rows in the grid. This value includes all rows from every page, but does not include the additional rows inserted in the last page when there are insufficient data items to fill the last page.</returns>
+    </member>
+    <member name="T:System.Web.Helpers.WebGridColumn">
+      <summary>Represents a column in a <see cref="T:System.Web.Helpers.WebGrid" /> instance.</summary>
+    </member>
+    <member name="M:System.Web.Helpers.WebGridColumn.#ctor">
+      <summary>Initializes a new instance of the <see cref="T:System.Web.Helpers.WebGridColumn" /> class.</summary>
+    </member>
+    <member name="P:System.Web.Helpers.WebGridColumn.CanSort">
+      <summary>Gets or sets a value that indicates whether the <see cref="T:System.Web.Helpers.WebGrid" /> column can be sorted.</summary>
+      <returns>true to indicate that the column can be sorted; otherwise, false.</returns>
+    </member>
+    <member name="P:System.Web.Helpers.WebGridColumn.ColumnName">
+      <summary>Gets or sets the name of the data item that is associated with the <see cref="T:System.Web.Helpers.WebGrid" /> column.</summary>
+      <returns>The name of the data item.</returns>
+    </member>
+    <member name="P:System.Web.Helpers.WebGridColumn.Format">
+      <summary>Gets or sets a function that is used to format the data item that is associated with the <see cref="T:System.Web.Helpers.WebGrid" /> column.</summary>
+      <returns>The function that is used to format that data item that is associated with the column.</returns>
+    </member>
+    <member name="P:System.Web.Helpers.WebGridColumn.Header">
+      <summary>Gets or sets the text that is rendered in the header of the <see cref="T:System.Web.Helpers.WebGrid" /> column.</summary>
+      <returns>The text that is rendered to the column header.</returns>
+    </member>
+    <member name="P:System.Web.Helpers.WebGridColumn.Style">
+      <summary>Gets or sets the CSS class attribute that is rendered as part of the HTML table cells that are associated with the <see cref="T:System.Web.Helpers.WebGrid" /> column.</summary>
+      <returns>The CSS class attribute that is applied to cells that are associated with the column.</returns>
+    </member>
+    <member name="T:System.Web.Helpers.WebGridPagerModes">
+      <summary>Specifies flags that describe the methods that are provided for moving between the pages of a <see cref="T:System.Web.Helpers.WebGrid" /> instance.This enumeration has a <see cref="T:System.FlagsAttribute" /> attribute that allows a bitwise combination of its member values.</summary>
+    </member>
+    <member name="F:System.Web.Helpers.WebGridPagerModes.All">
+      <summary>Indicates that all methods for moving between <see cref="T:System.Web.Helpers.WebGrid" /> pages are provided.</summary>
+    </member>
+    <member name="F:System.Web.Helpers.WebGridPagerModes.FirstLast">
+      <summary>Indicates that methods for moving directly to the first or last <see cref="F:System.Web.Helpers.WebGrid" /> page are provided.</summary>
+    </member>
+    <member name="F:System.Web.Helpers.WebGridPagerModes.NextPrevious">
+      <summary>Indicates that methods for moving to the next or previous <see cref="F:System.Web.Helpers.WebGrid" /> page are provided.</summary>
+    </member>
+    <member name="F:System.Web.Helpers.WebGridPagerModes.Numeric">
+      <summary>Indicates that methods for moving to a nearby <see cref="F:System.Web.Helpers.WebGrid" /> page by using a page number are provided.</summary>
+    </member>
+    <member name="T:System.Web.Helpers.WebGridRow">
+      <summary>Represents a row in a <see cref="T:System.Web.Helpers.WebGrid" /> instance.</summary>
+    </member>
+    <member name="M:System.Web.Helpers.WebGridRow.#ctor(System.Web.Helpers.WebGrid,System.Object,System.Int32)">
+      <summary>Initializes a new instance of the <see cref="T:System.Web.Helpers.WebGridRow" /> class using the specified <see cref="T:System.Web.Helpers.WebGrid" /> instance, row value, and index.</summary>
+      <param name="webGrid">The <see cref="T:System.Web.Helpers.WebGrid" /> instance that contains the row.</param>
+      <param name="value">An object that contains a property member for each value in the row.</param>
+      <param name="rowIndex">The index of the row.</param>
+    </member>
+    <member name="M:System.Web.Helpers.WebGridRow.GetEnumerator">
+      <summary>Returns an enumerator that can be used to iterate through the values of the <see cref="T:System.Web.Helpers.WebGridRow" /> instance.</summary>
+      <returns>An enumerator that can be used to iterate through the values of the row.</returns>
+    </member>
+    <member name="M:System.Web.Helpers.WebGridRow.GetSelectLink(System.String)">
+      <summary>Returns an HTML element (a link) that users can use to select the row.</summary>
+      <returns>The link that users can click to select the row.</returns>
+      <param name="text">The inner text of the link element. If <paramref name="text" /> is empty or null, "Select" is used.</param>
+    </member>
+    <member name="M:System.Web.Helpers.WebGridRow.GetSelectUrl">
+      <summary>Returns the URL that can be used to select the row.</summary>
+      <returns>The URL that is used to select a row.</returns>
+    </member>
+    <member name="P:System.Web.Helpers.WebGridRow.Item(System.Int32)">
+      <summary>Returns the value at the specified index in the <see cref="T:System.Web.Helpers.WebGridRow" /> instance.</summary>
+      <returns>The value at the specified index.</returns>
+      <param name="index">The zero-based index of the value in the row to return.</param>
+      <exception cref="T:System.ArgumentOutOfRangeException">
+        <paramref name="index" /> is less than 0 or greater than or equal to the number of values in the row.</exception>
+    </member>
+    <member name="P:System.Web.Helpers.WebGridRow.Item(System.String)">
+      <summary>Returns the value that has the specified name in the <see cref="T:System.Web.Helpers.WebGridRow" /> instance.</summary>
+      <returns>The specified value.</returns>
+      <param name="name">The name of the value in the row to return.</param>
+      <exception cref="T:System.ArgumentException">
+        <paramref name="name" /> is Nothing or empty.</exception>
+      <exception cref="T:System.InvalidOperationException">
+        <paramref name="name" /> specifies a value that does not exist.</exception>
+    </member>
+    <member name="M:System.Web.Helpers.WebGridRow.System#Collections#IEnumerable#GetEnumerator">
+      <summary>Returns an enumerator that can be used to iterate through a collection.</summary>
+      <returns>An enumerator that can be used to iterate through the collection.</returns>
+    </member>
+    <member name="M:System.Web.Helpers.WebGridRow.ToString">
+      <summary>Returns a string that represents all of the values of the <see cref="T:System.Web.Helpers.WebGridRow" /> instance.</summary>
+      <returns>A string that represents the row's values.</returns>
+    </member>
+    <member name="M:System.Web.Helpers.WebGridRow.TryGetMember(System.Dynamic.GetMemberBinder,System.Object@)">
+      <summary>Returns the value of a <see cref="T:System.Web.Helpers.WebGridRow" /> member that is described by the specified binder.</summary>
+      <returns>true if the value of the item was successfully retrieved; otherwise, false.</returns>
+      <param name="binder">The getter of the bound property member.</param>
+      <param name="result">When this method returns, contains an object that holds the value of the item described by <paramref name="binder" />. This parameter is passed uninitialized.</param>
+    </member>
+    <member name="P:System.Web.Helpers.WebGridRow.Value">
+      <summary>Gets an object that contains a property member for each value in the row.</summary>
+      <returns>An object that contains each value in the row as a property.</returns>
+    </member>
+    <member name="P:System.Web.Helpers.WebGridRow.WebGrid">
+      <summary>Gets the <see cref="T:System.Web.Helpers.WebGrid" /> instance that the row belongs to.</summary>
+      <returns>The <see cref="T:System.Web.Helpers.WebGrid" /> instance that contains the row.</returns>
+    </member>
+    <member name="T:System.Web.Helpers.WebImage">
+      <summary>Represents an object that lets you display and manage images in a web page.</summary>
+    </member>
+    <member name="M:System.Web.Helpers.WebImage.#ctor(System.Byte[])">
+      <summary>Initializes a new instance of the <see cref="T:System.Web.Helpers.WebImage" /> class using a byte array to represent the image.</summary>
+      <param name="content">The image.</param>
+    </member>
+    <member name="M:System.Web.Helpers.WebImage.#ctor(System.IO.Stream)">
+      <summary>Initializes a new instance of the <see cref="T:System.Web.Helpers.WebImage" /> class using a stream to represent the image.</summary>
+      <param name="imageStream">The image.</param>
+    </member>
+    <member name="M:System.Web.Helpers.WebImage.#ctor(System.String)">
+      <summary>Initializes a new instance of the <see cref="T:System.Web.Helpers.WebImage" /> class using a path to represent the image location.</summary>
+      <param name="filePath">The path of the file that contains the image.</param>
+    </member>
+    <member name="M:System.Web.Helpers.WebImage.AddImageWatermark(System.String,System.Int32,System.Int32,System.String,System.String,System.Int32,System.Int32)">
+      <summary>Adds a watermark image using a path to the watermark image.</summary>
+      <returns>The watermarked image.</returns>
+      <param name="watermarkImageFilePath">The path of a file that contains the watermark image.</param>
+      <param name="width">The width, in pixels, of the watermark image.</param>
+      <param name="height">The height, in pixels, of the watermark image.</param>
+      <param name="horizontalAlign">The horizontal alignment for watermark image. Values can be "Left", "Right", or "Center".</param>
+      <param name="verticalAlign">The vertical alignment for the watermark image. Values can be "Top", "Middle", or "Bottom".</param>
+      <param name="opacity">The opacity for the watermark image, specified as a value between 0 and 100.</param>
+      <param name="padding">The size, in pixels, of the padding around the watermark image.</param>
+    </member>
+    <member name="M:System.Web.Helpers.WebImage.AddImageWatermark(System.Web.Helpers.WebImage,System.Int32,System.Int32,System.String,System.String,System.Int32,System.Int32)">
+      <summary>Adds a watermark image using the specified image object.</summary>
+      <returns>The watermarked image.</returns>
+      <param name="watermarkImage">A <see cref="T:System.Web.Helpers.WebImage" /> object.</param>
+      <param name="width">The width, in pixels, of the watermark image.</param>
+      <param name="height">The height, in pixels, of the watermark image.</param>
+      <param name="horizontalAlign">The horizontal alignment for watermark image. Values can be "Left", "Right", or "Center".</param>
+      <param name="verticalAlign">The vertical alignment for the watermark image. Values can be "Top", "Middle", or "Bottom".</param>
+      <param name="opacity">The opacity for the watermark image, specified as a value between 0 and 100.</param>
+      <param name="padding">The size, in pixels, of the padding around the watermark image.</param>
+    </member>
+    <member name="M:System.Web.Helpers.WebImage.AddTextWatermark(System.String,System.String,System.Int32,System.String,System.String,System.String,System.String,System.Int32,System.Int32)">
+      <summary>Adds watermark text to the image.</summary>
+      <returns>The watermarked image.</returns>
+      <param name="text">The text to use as a watermark.</param>
+      <param name="fontColor">The color of the watermark text.</param>
+      <param name="fontSize">The font size of the watermark text.</param>
+      <param name="fontStyle">The font style of the watermark text.</param>
+      <param name="fontFamily">The font type of the watermark text.</param>
+      <param name="horizontalAlign">The horizontal alignment for watermark text. Values can be "Left", "Right", or "Center".</param>
+      <param name="verticalAlign">The vertical alignment for the watermark text. Values can be "Top", "Middle", or "Bottom".</param>
+      <param name="opacity">The opacity for the watermark image, specified as a value between 0 and 100.</param>
+      <param name="padding">The size, in pixels, of the padding around the watermark text.</param>
+    </member>
+    <member name="M:System.Web.Helpers.WebImage.Clone">
+      <summary>Copies the <see cref="T:System.Web.Helpers.WebImage" /> object.</summary>
+      <returns>The image.</returns>
+    </member>
+    <member name="M:System.Web.Helpers.WebImage.Crop(System.Int32,System.Int32,System.Int32,System.Int32)">
+      <summary>Crops an image.</summary>
+      <returns>The cropped image.</returns>
+      <param name="top">The number of pixels to remove from the top.</param>
+      <param name="left">The number of pixels to remove from the left.</param>
+      <param name="bottom">The number of pixels to remove from the bottom.</param>
+      <param name="right">The number of pixels to remove from the right.</param>
+    </member>
+    <member name="P:System.Web.Helpers.WebImage.FileName">
+      <summary>Gets or sets the file name of the <see cref="T:System.Web.Helpers.WebImage" /> object.</summary>
+      <returns>The file name.</returns>
+    </member>
+    <member name="M:System.Web.Helpers.WebImage.FlipHorizontal">
+      <summary>Flips an image horizontally.</summary>
+      <returns>The flipped image.</returns>
+    </member>
+    <member name="M:System.Web.Helpers.WebImage.FlipVertical">
+      <summary>Flips an image vertically.</summary>
+      <returns>The flipped image.</returns>
+    </member>
+    <member name="M:System.Web.Helpers.WebImage.GetBytes(System.String)">
+      <summary>Returns the image as a byte array.</summary>
+      <returns>The image.</returns>
+      <param name="requestedFormat">The <see cref="P:System.Web.Helpers.WebImage.ImageFormat" /> value of the <see cref="T:System.Web.Helpers.WebImage" /> object.</param>
+    </member>
+    <member name="M:System.Web.Helpers.WebImage.GetImageFromRequest(System.String)">
+      <summary>Returns an image that has been uploaded using the browser.</summary>
+      <returns>The image.</returns>
+      <param name="postedFileName">(Optional) The name of the file that has been posted. If no file name is specified, the first file that was uploaded is returned.</param>
+    </member>
+    <member name="P:System.Web.Helpers.WebImage.Height">
+      <summary>Gets the height, in pixels, of the image.</summary>
+      <returns>The height.</returns>
+    </member>
+    <member name="P:System.Web.Helpers.WebImage.ImageFormat">
+      <summary>Gets the format of the image (for example, "jpeg" or "png").</summary>
+      <returns>The file format of the image.</returns>
+    </member>
+    <member name="M:System.Web.Helpers.WebImage.Resize(System.Int32,System.Int32,System.Boolean,System.Boolean)">
+      <summary>Resizes an image.</summary>
+      <returns>The resized image.</returns>
+      <param name="width">The width, in pixels, of the <see cref="T:System.Web.Helpers.WebImage" /> object.</param>
+      <param name="height">The height, in pixels, of the <see cref="T:System.Web.Helpers.WebImage" /> object.</param>
+      <param name="preserveAspectRatio">true to preserve the aspect ratio of the image; otherwise, false.</param>
+      <param name="preventEnlarge">true to prevent the enlargement of the image; otherwise, false.</param>
+    </member>
+    <member name="M:System.Web.Helpers.WebImage.RotateLeft">
+      <summary>Rotates an image to the left.</summary>
+      <returns>The rotated image.</returns>
+    </member>
+    <member name="M:System.Web.Helpers.WebImage.RotateRight">
+      <summary>Rotates an image to the right.</summary>
+      <returns>The rotated image.</returns>
+    </member>
+    <member name="M:System.Web.Helpers.WebImage.Save(System.String,System.String,System.Boolean)">
+      <summary>Saves the image using the specified file name.</summary>
+      <returns>The image.</returns>
+      <param name="filePath">The path to save the image to.</param>
+      <param name="imageFormat">The format to use when the image file is saved, such as "gif", or "png".</param>
+      <param name="forceCorrectExtension">true to force the correct file-name extension to be used for the format that is specified in <paramref name="imageFormat" />; otherwise, false. If there is a mismatch between the file type and the specified file-name extension, and if <paramref name="forceCorrectExtension" /> is true, the correct extension will be appended to the file name. For example, a PNG file named Photograph.txt is saved using the name Photograph.txt.png.</param>
+    </member>
+    <member name="P:System.Web.Helpers.WebImage.Width">
+      <summary>Gets the width, in pixels, of the image.</summary>
+      <returns>The width.</returns>
+    </member>
+    <member name="M:System.Web.Helpers.WebImage.Write(System.String)">
+      <summary>Renders an image to the browser.</summary>
+      <returns>The image.</returns>
+      <param name="requestedFormat">(Optional) The file format to use when the image is written.</param>
+    </member>
+    <member name="T:System.Web.Helpers.WebMail">
+      <summary>Provides a way to construct and send an email message using Simple Mail Transfer Protocol (SMTP).</summary>
+    </member>
+    <member name="P:System.Web.Helpers.WebMail.EnableSsl">
+      <summary>Gets or sets a value that indicates whether Secure Sockets Layer (SSL) is used to encrypt the connection when an email message is sent.</summary>
+      <returns>true if SSL is used to encrypt the connection; otherwise, false.</returns>
+    </member>
+    <member name="P:System.Web.Helpers.WebMail.From">
+      <summary>Gets or sets the email address of the sender.</summary>
+      <returns>The email address of the sender.</returns>
+    </member>
+    <member name="P:System.Web.Helpers.WebMail.Password">
+      <summary>Gets or sets the password of the sender's email account.</summary>
+      <returns>The sender's password.</returns>
+    </member>
+    <member name="M:System.Web.Helpers.WebMail.Send(System.String,System.String,System.String,System.String,System.String,System.Collections.Generic.IEnumerable{System.String},System.Boolean,System.Collections.Generic.IEnumerable{System.String},System.String,System.String,System.String,System.String,System.String)">
+      <summary>Sends the specified message to an SMTP server for delivery.</summary>
+      <param name="to">The email address of the recipient or recipients. Separate multiple recipients using a semicolon (;).</param>
+      <param name="subject">The subject line for the email message.</param>
+      <param name="body">The body of the email message. If <paramref name="isBodyHtml" /> is true, HTML in the body is interpreted as markup.</param>
+      <param name="from">(Optional) The email address of the message sender, or null to not specify a sender. The default value is null.</param>
+      <param name="cc">(Optional) The email addresses of additional recipients to send a copy of the message to, or null if there are no additional recipients. Separate multiple recipients using a semicolon (;). The default value is null.</param>
+      <param name="filesToAttach">(Optional) A collection of file names that specifies the files to attach to the email message, or null if there are no files to attach. The default value is null.</param>
+      <param name="isBodyHtml">(Optional) true to specify that the email message body is in HTML format; false to indicate that the body is in plain-text format. The default value is true.</param>
+      <param name="additionalHeaders">(Optional) A collection of headers to add to the normal SMTP headers included in this email message, or null to send no additional headers. The default value is null.</param>
+      <param name="bcc">(Optional) The email addresses of additional recipients to send a "blind" copy of the message to, or null if there are no additional recipients. Separate multiple recipients using a semicolon (;). The default value is null.</param>
+      <param name="contentEncoding">(Optional) The encoding to use for the body of the message. Possible values are property values for the <see cref="T:System.Text.Encoding" /> class, such as <see cref="P:System.Text.Encoding.UTF8" />. The default value is null.</param>
+      <param name="headerEncoding">(Optional) The encoding to use for the header of the message. Possible values are property values for the <see cref="T:System.Text.Encoding" /> class, such as <see cref="P:System.Text.Encoding.UTF8" />. The default value is null.</param>
+      <param name="priority">(Optional) A value ("Normal", "Low", "High") that specifies the priority of the message. The default is "Normal".</param>
+      <param name="replyTo">(Optional) The email address that will be used when the recipient replies to the message. The default value is null, which indicates that the reply address is the value of the From property. </param>
+    </member>
+    <member name="P:System.Web.Helpers.WebMail.SmtpPort">
+      <summary>Gets or sets the port that is used for SMTP transactions.</summary>
+      <returns>The port that is used for SMTP transactions.</returns>
+    </member>
+    <member name="P:System.Web.Helpers.WebMail.SmtpServer">
+      <summary>Gets or sets the name of the SMTP server that is used to transmit the email message.</summary>
+      <returns>The SMTP server.</returns>
+    </member>
+    <member name="P:System.Web.Helpers.WebMail.SmtpUseDefaultCredentials">
+      <summary>Gets or sets a value that indicates whether the default credentials are sent with the requests.</summary>
+      <returns>true if credentials are sent with the email message; otherwise, false.</returns>
+    </member>
+    <member name="P:System.Web.Helpers.WebMail.UserName">
+      <summary>Gets or sets the name of email account that is used to send email.</summary>
+      <returns>The name of the user account.</returns>
+    </member>
+  </members>
+</doc>
\ No newline at end of file
diff --git a/Engine/Source/Programs/CrashReporter/CrashReportWebSite/packages/Microsoft.AspNet.WebPages.3.2.3/lib/net45/System.Web.WebPages.Deployment.xml b/Engine/Source/Programs/CrashReporter/CrashReportWebSite/packages/Microsoft.AspNet.WebPages.3.2.3/lib/net45/System.Web.WebPages.Deployment.xml
new file mode 100644
--- /dev/null
+++ b/Engine/Source/Programs/CrashReporter/CrashReportWebSite/packages/Microsoft.AspNet.WebPages.3.2.3/lib/net45/System.Web.WebPages.Deployment.xml
@@ -0,0 +1,60 @@
+﻿<?xml version="1.0" encoding="utf-8"?>
+<doc>
+  <assembly>
+    <name>System.Web.WebPages.Deployment</name>
+  </assembly>
+  <members>
+    <member name="T:System.Web.WebPages.Deployment.PreApplicationStartCode">
+      <summary>Provides a registration point for pre-application start code for Web Pages deployment.</summary>
+    </member>
+    <member name="M:System.Web.WebPages.Deployment.PreApplicationStartCode.Start">
+      <summary>Registers pre-application start code for Web Pages deployment.</summary>
+    </member>
+    <member name="T:System.Web.WebPages.Deployment.WebPagesDeployment">
+      <summary>This type/member supports the .NET Framework infrastructure and is not intended to be used directly from your code.Provides methods that are used to get deployment information about the Web application.</summary>
+    </member>
+    <member name="M:System.Web.WebPages.Deployment.WebPagesDeployment.GetAssemblyPath(System.Version)">
+      <summary>This type/member supports the .NET Framework infrastructure and is not intended to be used directly from your code.Gets the assembly path for the Web Pages deployment.</summary>
+      <returns>The assembly path for the Web Pages deployment.</returns>
+      <param name="version">The Web Pages version.</param>
+    </member>
+    <member name="M:System.Web.WebPages.Deployment.WebPagesDeployment.GetExplicitWebPagesVersion(System.String)">
+      <summary>This type/member supports the .NET Framework infrastructure and is not intended to be used directly from your code.Gets the Web Pages version from the given binary path.</summary>
+      <returns>The Web Pages version.</returns>
+      <param name="path">The binary path for the Web Pages.</param>
+    </member>
+    <member name="M:System.Web.WebPages.Deployment.WebPagesDeployment.GetIncompatibleDependencies(System.String)">
+      <summary>This type/member supports the .NET Framework infrastructure and is not intended to be used directly from your code.Gets the assembly references from the given path regardless of the Web Pages version.</summary>
+      <returns>The dictionary containing the assembly references of the Web Pages and its version.</returns>
+      <param name="appPath">The path to the Web Pages application.</param>
+    </member>
+    <member name="M:System.Web.WebPages.Deployment.WebPagesDeployment.GetMaxVersion">
+      <summary>This type/member supports the .NET Framework infrastructure and is not intended to be used directly from your code.Gets the maximum version of the Web Pages loaded assemblies.</summary>
+      <returns>The maximum version of the Web Pages loaded assemblies.</returns>
+    </member>
+    <member name="M:System.Web.WebPages.Deployment.WebPagesDeployment.GetVersion(System.String)">
+      <summary>Gets the Web Pages version from the given path.</summary>
+      <returns>The Web Pages version.</returns>
+      <param name="path">The path of the root directory for the application.</param>
+    </member>
+    <member name="M:System.Web.WebPages.Deployment.WebPagesDeployment.GetVersionWithoutEnabledCheck(System.String)">
+      <summary>This type/member supports the .NET Framework infrastructure and is not intended to be used directly from your code.Gets the Web Pages version using the configuration settings with the specified path.</summary>
+      <returns>The Web Pages version.</returns>
+      <param name="path">The path to the application settings.</param>
+    </member>
+    <member name="M:System.Web.WebPages.Deployment.WebPagesDeployment.GetWebPagesAssemblies">
+      <summary>This type/member supports the .NET Framework infrastructure and is not intended to be used directly from your code.Returns the assemblies for this Web Pages deployment.</summary>
+      <returns>A list containing the assemblies for this Web Pages deployment.</returns>
+    </member>
+    <member name="M:System.Web.WebPages.Deployment.WebPagesDeployment.IsEnabled(System.String)">
+      <summary>This type/member supports the .NET Framework infrastructure and is not intended to be used directly from your code.Indicates whether the Web Pages deployment is enabled.</summary>
+      <returns>true if the Web Pages deployment is enabled; otherwise, false.</returns>
+      <param name="path">The path to the Web Pages deployment.</param>
+    </member>
+    <member name="M:System.Web.WebPages.Deployment.WebPagesDeployment.IsExplicitlyDisabled(System.String)">
+      <summary>This type/member supports the .NET Framework infrastructure and is not intended to be used directly from your code.Indicates whether the Web Pages deployment is explicitly disabled.</summary>
+      <returns>true if the Web Pages deployment is explicitly disabled; otherwise, false.</returns>
+      <param name="path">The path to the Web Pages deployment.</param>
+    </member>
+  </members>
+</doc>
\ No newline at end of file
