Copying //UE4/Dev-Core to //UE4/Main

diff --git a/Engine/Source/Programs/CrashReporter/CrashReportWebSite/packages/Microsoft.AspNet.WebPages.3.2.3/lib/net45/System.Web.WebPages.Razor.xml b/Engine/Source/Programs/CrashReporter/CrashReportWebSite/packages/Microsoft.AspNet.WebPages.3.2.3/lib/net45/System.Web.WebPages.Razor.xml
new file mode 100644
--- /dev/null
+++ b/Engine/Source/Programs/CrashReporter/CrashReportWebSite/packages/Microsoft.AspNet.WebPages.3.2.3/lib/net45/System.Web.WebPages.Razor.xml
@@ -0,0 +1,292 @@
+﻿<?xml version="1.0" encoding="utf-8"?>
+<doc>
+  <assembly>
+    <name>System.Web.WebPages.Razor</name>
+  </assembly>
+  <members>
+    <member name="T:System.Web.WebPages.Razor.CompilingPathEventArgs">
+      <summary>This type/member supports the .NET Framework infrastructure and is not intended to be used directly from your code.Represents the base class for the compiling path that contains event data.</summary>
+    </member>
+    <member name="M:System.Web.WebPages.Razor.CompilingPathEventArgs.#ctor(System.String,System.Web.WebPages.Razor.WebPageRazorHost)">
+      <summary>This type/member supports the .NET Framework infrastructure and is not intended to be used directly from your code.Initializes a new instance of the <see cref="T:System.Web.WebPages.Razor.CompilingPathEventArgs" /> class.</summary>
+      <param name="virtualPath">The string of virtual path.</param>
+      <param name="host">The host for the webpage razor.</param>
+    </member>
+    <member name="P:System.Web.WebPages.Razor.CompilingPathEventArgs.Host">
+      <summary>This type/member supports the .NET Framework infrastructure and is not intended to be used directly from your code.Gets or sets the host for the webpage razor.</summary>
+      <returns>The host for the webpage razor.</returns>
+    </member>
+    <member name="P:System.Web.WebPages.Razor.CompilingPathEventArgs.VirtualPath">
+      <summary>This type/member supports the .NET Framework infrastructure and is not intended to be used directly from your code.Gets the virtual path for the webpage.</summary>
+      <returns>The virtual path for the webpage.</returns>
+    </member>
+    <member name="T:System.Web.WebPages.Razor.PreApplicationStartCode">
+      <summary>This type/member supports the .NET Framework infrastructure and is not intended to be used directly from your code.</summary>
+    </member>
+    <member name="M:System.Web.WebPages.Razor.PreApplicationStartCode.Start">
+      <summary>This type/member supports the .NET Framework infrastructure and is not intended to be used directly from your code.</summary>
+    </member>
+    <member name="T:System.Web.WebPages.Razor.RazorBuildProvider">
+      <summary>This type/member supports the .NET Framework infrastructure and is not intended to be used directly from your code.Represents a build provider for Razor.</summary>
+    </member>
+    <member name="M:System.Web.WebPages.Razor.RazorBuildProvider.#ctor">
+      <summary>This type/member supports the .NET Framework infrastructure and is not intended to be used directly from your code.Initializes a new instance of the <see cref="T:System.Web.WebPages.Razor.RazorBuildProvider" /> class.</summary>
+    </member>
+    <member name="M:System.Web.WebPages.Razor.RazorBuildProvider.AddVirtualPathDependency(System.String)">
+      <summary>This type/member supports the .NET Framework infrastructure and is not intended to be used directly from your code.Adds a virtual path dependency to the collection.</summary>
+      <param name="dependency">A virtual path dependency to add.</param>
+    </member>
+    <member name="P:System.Web.WebPages.Razor.RazorBuildProvider.AssemblyBuilder">
+      <summary>This type/member supports the .NET Framework infrastructure and is not intended to be used directly from your code.Gets the assembly builder for Razor environment.</summary>
+      <returns>The assembly builder.</returns>
+    </member>
+    <member name="P:System.Web.WebPages.Razor.RazorBuildProvider.CodeCompilerType">
+      <summary>This type/member supports the .NET Framework infrastructure and is not intended to be used directly from your code.Gets the compiler settings for Razor environment.</summary>
+    </member>
+    <member name="E:System.Web.WebPages.Razor.RazorBuildProvider.CodeGenerationCompleted">
+      <summary>This type/member supports the .NET Framework infrastructure and is not intended to be used directly from your code.Occurs when code generation is completed.</summary>
+    </member>
+    <member name="E:System.Web.WebPages.Razor.RazorBuildProvider.CodeGenerationStarted">
+      <summary>This type/member supports the .NET Framework infrastructure and is not intended to be used directly from your code.Occurs when code generation is started.</summary>
+    </member>
+    <member name="E:System.Web.WebPages.Razor.RazorBuildProvider.CompilingPath">
+      <summary>This type/member supports the .NET Framework infrastructure and is not intended to be used directly from your code.Occurs when compiling with a new virtual path.</summary>
+    </member>
+    <member name="M:System.Web.WebPages.Razor.RazorBuildProvider.CreateHost">
+      <summary>This type/member supports the .NET Framework infrastructure and is not intended to be used directly from your code.Creates a Razor engine host instance base on web configuration.</summary>
+      <returns>A Razor engine host instance.</returns>
+    </member>
+    <member name="M:System.Web.WebPages.Razor.RazorBuildProvider.GenerateCode(System.Web.Compilation.AssemblyBuilder)">
+      <summary>This type/member supports the .NET Framework infrastructure and is not intended to be used directly from your code.Generates the code using the provided assembly builder.</summary>
+      <param name="assemblyBuilder">The assembly builder.</param>
+    </member>
+    <member name="M:System.Web.WebPages.Razor.RazorBuildProvider.GetGeneratedType(System.CodeDom.Compiler.CompilerResults)">
+      <summary>This type/member supports the .NET Framework infrastructure and is not intended to be used directly from your code.Gets the type of the generated code.</summary>
+      <returns>The type of the generated code.</returns>
+      <param name="results">The results of the code compilation.</param>
+    </member>
+    <member name="M:System.Web.WebPages.Razor.RazorBuildProvider.GetHostFromConfig">
+      <summary>This type/member supports the .NET Framework infrastructure and is not intended to be used directly from your code.Creates the Razor engine host instance based on the web configuration.</summary>
+      <returns>The Razor engine host instance.</returns>
+    </member>
+    <member name="M:System.Web.WebPages.Razor.RazorBuildProvider.InternalOpenReader">
+      <summary>This type/member supports the .NET Framework infrastructure and is not intended to be used directly from your code.Opens an internal text reader.</summary>
+      <returns>An internal text reader.</returns>
+    </member>
+    <member name="M:System.Web.WebPages.Razor.RazorBuildProvider.OnBeforeCompilePath(System.Web.WebPages.Razor.CompilingPathEventArgs)">
+      <summary>This type/member supports the .NET Framework infrastructure and is not intended to be used directly from your code.Raises the CompilingPath event.</summary>
+      <param name="args">The data provided for the CompilingPath event.</param>
+    </member>
+    <member name="P:System.Web.WebPages.Razor.RazorBuildProvider.VirtualPath">
+      <summary>This type/member supports the .NET Framework infrastructure and is not intended to be used directly from your code.Gets the virtual path of the source code.</summary>
+      <returns>The virtual path of the source code.</returns>
+    </member>
+    <member name="P:System.Web.WebPages.Razor.RazorBuildProvider.VirtualPathDependencies">
+      <summary>This type/member supports the .NET Framework infrastructure and is not intended to be used directly from your code.Gets the collection of virtual path for the dependencies.</summary>
+      <returns>The collection of virtual path for the dependencies.</returns>
+    </member>
+    <member name="T:System.Web.WebPages.Razor.WebCodeRazorHost">
+      <summary>This type/member supports the .NET Framework infrastructure and is not intended to be used directly from your code.Represents a web code razor host for the web pages.</summary>
+    </member>
+    <member name="M:System.Web.WebPages.Razor.WebCodeRazorHost.#ctor(System.String)">
+      <summary>This type/member supports the .NET Framework infrastructure and is not intended to be used directly from your code.Initializes a new instance of the <see cref="T:System.Web.WebPages.Razor.WebCodeRazorHost" /> class.</summary>
+      <param name="virtualPath">The virtual path.</param>
+    </member>
+    <member name="M:System.Web.WebPages.Razor.WebCodeRazorHost.#ctor(System.String,System.String)">
+      <summary>This type/member supports the .NET Framework infrastructure and is not intended to be used directly from your code.Initializes a new instance of the <see cref="T:System.Web.WebPages.Razor.WebCodeRazorHost" /> class.</summary>
+      <param name="virtualPath">The virtual path.</param>
+      <param name="physicalPath">The physical path.</param>
+    </member>
+    <member name="M:System.Web.WebPages.Razor.WebCodeRazorHost.GetClassName(System.String)">
+      <summary>This type/member supports the .NET Framework infrastructure and is not intended to be used directly from your code.Returns the class name of this instance.</summary>
+      <returns>The class name of this instance.</returns>
+      <param name="virtualPath">The virtual path.</param>
+    </member>
+    <member name="M:System.Web.WebPages.Razor.WebCodeRazorHost.PostProcessGeneratedCode(System.Web.Razor.Generator.CodeGeneratorContext)">
+      <summary>Generates a post process code for the web code razor host.</summary>
+      <param name="context">The generator code context.</param>
+    </member>
+    <member name="T:System.Web.WebPages.Razor.WebPageRazorHost">
+      <summary>This type/member supports the .NET Framework infrastructure and is not intended to be used directly from your code.Represents the razor hosts in a webpage.</summary>
+    </member>
+    <member name="M:System.Web.WebPages.Razor.WebPageRazorHost.#ctor(System.String)">
+      <summary>This type/member supports the .NET Framework infrastructure and is not intended to be used directly from your code.Initializes a new instance of the <see cref="T:System.Web.WebPages.Razor.WebPageRazorHost" /> class with the specified virtual file path.</summary>
+      <param name="virtualPath">The virtual file path.</param>
+    </member>
+    <member name="M:System.Web.WebPages.Razor.WebPageRazorHost.#ctor(System.String,System.String)">
+      <summary>This type/member supports the .NET Framework infrastructure and is not intended to be used directly from your code.Initializes a new instance of the <see cref="T:System.Web.WebPages.Razor.WebPageRazorHost" /> class with the specified virtual and physical file path.</summary>
+      <param name="virtualPath">The virtual file path.</param>
+      <param name="physicalPath">The physical file path.</param>
+    </member>
+    <member name="M:System.Web.WebPages.Razor.WebPageRazorHost.AddGlobalImport(System.String)">
+      <summary>This type/member supports the .NET Framework infrastructure and is not intended to be used directly from your code.Adds a global import on the webpage.</summary>
+      <param name="ns">The notification service name.</param>
+    </member>
+    <member name="P:System.Web.WebPages.Razor.WebPageRazorHost.CodeLanguage">
+      <summary>This type/member supports the .NET Framework infrastructure and is not intended to be used directly from your code.Gets the <see cref="T:System.Web.Razor.RazorCodeLanguage" />.</summary>
+      <returns>The <see cref="T:System.Web.Razor.RazorCodeLanguage" />.</returns>
+    </member>
+    <member name="M:System.Web.WebPages.Razor.WebPageRazorHost.CreateMarkupParser">
+      <summary>This type/member supports the .NET Framework infrastructure and is not intended to be used directly from your code.Creates a markup parser.</summary>
+      <returns>A markup parser.</returns>
+    </member>
+    <member name="P:System.Web.WebPages.Razor.WebPageRazorHost.DefaultBaseClass">
+      <summary>This type/member supports the .NET Framework infrastructure and is not intended to be used directly from your code.Gets or sets a value for the DefaultBaseClass.</summary>
+      <returns>A value for the DefaultBaseClass.</returns>
+    </member>
+    <member name="P:System.Web.WebPages.Razor.WebPageRazorHost.DefaultClassName">
+      <summary>This type/member supports the .NET Framework infrastructure and is not intended to be used directly from your code.Gets or sets the name of the default class.</summary>
+      <returns>The name of the default class.</returns>
+    </member>
+    <member name="P:System.Web.WebPages.Razor.WebPageRazorHost.DefaultDebugCompilation">
+      <summary>This type/member supports the .NET Framework infrastructure and is not intended to be used directly from your code.Gets or sets a value that indicates whether the debug compilation is set to default.</summary>
+      <returns>true if the debug compilation is set to default; otherwise, false.</returns>
+    </member>
+    <member name="P:System.Web.WebPages.Razor.WebPageRazorHost.DefaultPageBaseClass">
+      <summary>This type/member supports the .NET Framework infrastructure and is not intended to be used directly from your code.Gets or sets the base class of the default page.</summary>
+      <returns>The base class of the default page.</returns>
+    </member>
+    <member name="M:System.Web.WebPages.Razor.WebPageRazorHost.GetClassName(System.String)">
+      <summary>This type/member supports the .NET Framework infrastructure and is not intended to be used directly from your code.Retrieves the name of the class to which the specified webpage belongs.</summary>
+      <returns>The name of the class to which the specified webpage belongs.</returns>
+      <param name="virtualPath">The virtual file path.</param>
+    </member>
+    <member name="M:System.Web.WebPages.Razor.WebPageRazorHost.GetCodeLanguage">
+      <summary>This type/member supports the .NET Framework infrastructure and is not intended to be used directly from your code.Gets the code language specified in the webpage.</summary>
+      <returns>The code language specified in the webpage.</returns>
+    </member>
+    <member name="M:System.Web.WebPages.Razor.WebPageRazorHost.GetGlobalImports">
+      <summary>This type/member supports the .NET Framework infrastructure and is not intended to be used directly from your code.Gets the global imports for the webpage.</summary>
+      <returns>The global imports for the webpage.</returns>
+    </member>
+    <member name="P:System.Web.WebPages.Razor.WebPageRazorHost.InstrumentedSourceFilePath">
+      <summary>This type/member supports the .NET Framework infrastructure and is not intended to be used directly from your code.Gets or sets the file path of the instrumental source.</summary>
+      <returns>The file path of the instrumental source.</returns>
+    </member>
+    <member name="P:System.Web.WebPages.Razor.WebPageRazorHost.IsSpecialPage">
+      <summary>This type/member supports the .NET Framework infrastructure and is not intended to be used directly from your code.Gets a value that indicates whether the webpage is a special page.</summary>
+      <returns>true if the webpage is a special page; otherwise, false.</returns>
+    </member>
+    <member name="P:System.Web.WebPages.Razor.WebPageRazorHost.PhysicalPath">
+      <summary>This type/member supports the .NET Framework infrastructure and is not intended to be used directly from your code.Gets the physical file system path of the razor host.</summary>
+      <returns>They physical file system path of the razor host.</returns>
+    </member>
+    <member name="M:System.Web.WebPages.Razor.WebPageRazorHost.PostProcessGeneratedCode(System.Web.Razor.Generator.CodeGeneratorContext)">
+      <summary>This type/member supports the .NET Framework infrastructure and is not intended to be used directly from your code.Gets the generated code after the process.</summary>
+      <param name="context">The <see cref="T:System.Web.Razor.Generator.CodeGeneratorContext" />.</param>
+    </member>
+    <member name="M:System.Web.WebPages.Razor.WebPageRazorHost.RegisterSpecialFile(System.String,System.String)">
+      <summary>This type/member supports the .NET Framework infrastructure and is not intended to be used directly from your code.Registers the special file with the specified file name and base type name.</summary>
+      <param name="fileName">The file name.</param>
+      <param name="baseTypeName">The base type name.</param>
+    </member>
+    <member name="M:System.Web.WebPages.Razor.WebPageRazorHost.RegisterSpecialFile(System.String,System.Type)">
+      <summary>This type/member supports the .NET Framework infrastructure and is not intended to be used directly from your code.Registers the special file with the specified file name and base type.</summary>
+      <param name="fileName">The file name.</param>
+      <param name="baseType">The type of base file.</param>
+    </member>
+    <member name="P:System.Web.WebPages.Razor.WebPageRazorHost.VirtualPath">
+      <summary>This type/member supports the .NET Framework infrastructure and is not intended to be used directly from your code.Gets the virtual file path.</summary>
+      <returns>The virtual file path.</returns>
+    </member>
+    <member name="T:System.Web.WebPages.Razor.WebRazorHostFactory">
+      <summary>This type/member supports the .NET Framework infrastructure and is not intended to be used directly from your code.Creates instances of the host files.</summary>
+    </member>
+    <member name="M:System.Web.WebPages.Razor.WebRazorHostFactory.#ctor">
+      <summary>This type/member supports the .NET Framework infrastructure and is not intended to be used directly from your code.Initializes a new instance of the <see cref="T:System.Web.WebPages.Razor.WebRazorHostFactory" /> class.</summary>
+    </member>
+    <member name="M:System.Web.WebPages.Razor.WebRazorHostFactory.ApplyConfigurationToHost(System.Web.WebPages.Razor.Configuration.RazorPagesSection,System.Web.WebPages.Razor.WebPageRazorHost)">
+      <summary>This type/member supports the .NET Framework infrastructure and is not intended to be used directly from your code.Loads the service description information from the configuration file and applies it to the host.</summary>
+      <param name="config">The configuration.</param>
+      <param name="host">The webpage razor host.</param>
+    </member>
+    <member name="M:System.Web.WebPages.Razor.WebRazorHostFactory.CreateDefaultHost(System.String)">
+      <summary>This type/member supports the .NET Framework infrastructure and is not intended to be used directly from your code.Creates a default host with the specified virtual path.</summary>
+      <returns>A default host.</returns>
+      <param name="virtualPath">The virtual path of the file.</param>
+    </member>
+    <member name="M:System.Web.WebPages.Razor.WebRazorHostFactory.CreateDefaultHost(System.String,System.String)">
+      <summary>This type/member supports the .NET Framework infrastructure and is not intended to be used directly from your code.Creates a default host with the specified virtual and physical path.</summary>
+      <returns>A default host.</returns>
+      <param name="virtualPath">The virtual path of the file.</param>
+      <param name="physicalPath">The physical file system path.</param>
+    </member>
+    <member name="M:System.Web.WebPages.Razor.WebRazorHostFactory.CreateHost(System.String,System.String)">
+      <summary>This type/member supports the .NET Framework infrastructure and is not intended to be used directly from your code.Creates a Razor host.</summary>
+      <returns>A razor host.</returns>
+      <param name="virtualPath">The virtual path to the target file.</param>
+      <param name="physicalPath">The physical path to the target file.</param>
+    </member>
+    <member name="M:System.Web.WebPages.Razor.WebRazorHostFactory.CreateHostFromConfig(System.String)">
+      <summary>This type/member supports the .NET Framework infrastructure and is not intended to be used directly from your code.Creates a host from the configuration.</summary>
+      <returns>A host from the configuration.</returns>
+      <param name="virtualPath">The virtual path to the target file.</param>
+    </member>
+    <member name="M:System.Web.WebPages.Razor.WebRazorHostFactory.CreateHostFromConfig(System.String,System.String)">
+      <summary>This type/member supports the .NET Framework infrastructure and is not intended to be used directly from your code.Creates a host from the configuration.</summary>
+      <returns>A host from the configuration.</returns>
+      <param name="virtualPath">The virtual path of the file.</param>
+      <param name="physicalPath">The physical file system path.</param>
+    </member>
+    <member name="M:System.Web.WebPages.Razor.WebRazorHostFactory.CreateHostFromConfig(System.Web.WebPages.Razor.Configuration.RazorWebSectionGroup,System.String)">
+      <summary>This type/member supports the .NET Framework infrastructure and is not intended to be used directly from your code.Creates a host from the configuration.</summary>
+      <returns>A host from the configuration.</returns>
+      <param name="config">The configuration.</param>
+      <param name="virtualPath">The virtual path of the file.</param>
+    </member>
+    <member name="M:System.Web.WebPages.Razor.WebRazorHostFactory.CreateHostFromConfig(System.Web.WebPages.Razor.Configuration.RazorWebSectionGroup,System.String,System.String)">
+      <summary>This type/member supports the .NET Framework infrastructure and is not intended to be used directly from your code.Creates a host from the configuration.</summary>
+      <returns>A host from the configuration.</returns>
+      <param name="config">The configuration.</param>
+      <param name="virtualPath">The virtual path of the file.</param>
+      <param name="physicalPath">The physical file system path.</param>
+    </member>
+    <member name="T:System.Web.WebPages.Razor.Configuration.HostSection">
+      <summary>Provides configuration system support for the host configuration section.</summary>
+    </member>
+    <member name="M:System.Web.WebPages.Razor.Configuration.HostSection.#ctor">
+      <summary>Initializes a new instance of the <see cref="T:System.Web.WebPages.Razor.Configuration.HostSection" /> class.</summary>
+    </member>
+    <member name="P:System.Web.WebPages.Razor.Configuration.HostSection.FactoryType">
+      <summary>Gets or sets the host factory.</summary>
+      <returns>The host factory.</returns>
+    </member>
+    <member name="F:System.Web.WebPages.Razor.Configuration.HostSection.SectionName">
+      <summary>Represents the name of the configuration section for a Razor host environment.</summary>
+    </member>
+    <member name="T:System.Web.WebPages.Razor.Configuration.RazorPagesSection">
+      <summary>Provides configuration system support for the pages configuration section.</summary>
+    </member>
+    <member name="M:System.Web.WebPages.Razor.Configuration.RazorPagesSection.#ctor">
+      <summary>Initializes a new instance of the <see cref="T:System.Web.WebPages.Razor.Configuration.RazorPagesSection" /> class.</summary>
+    </member>
+    <member name="P:System.Web.WebPages.Razor.Configuration.RazorPagesSection.Namespaces">
+      <summary>Gets or sets the collection of namespaces to add to Web Pages pages in the current application.</summary>
+      <returns>The collection of namespaces.</returns>
+    </member>
+    <member name="P:System.Web.WebPages.Razor.Configuration.RazorPagesSection.PageBaseType">
+      <summary>Gets or sets the name of the page base type class.</summary>
+      <returns>The name of the page base type class.</returns>
+    </member>
+    <member name="F:System.Web.WebPages.Razor.Configuration.RazorPagesSection.SectionName">
+      <summary>Represents the name of the configuration section for Razor pages.</summary>
+    </member>
+    <member name="T:System.Web.WebPages.Razor.Configuration.RazorWebSectionGroup">
+      <summary>Provides configuration system support for the system.web.webPages.razor configuration section.</summary>
+    </member>
+    <member name="M:System.Web.WebPages.Razor.Configuration.RazorWebSectionGroup.#ctor">
+      <summary>Initializes a new instance of the <see cref="T:System.Web.WebPages.Razor.Configuration.RazorWebSectionGroup" /> class.</summary>
+    </member>
+    <member name="F:System.Web.WebPages.Razor.Configuration.RazorWebSectionGroup.GroupName">
+      <summary>Represents the name of the configuration section for Razor Web section. Contains the static, read-only string "system.web.webPages.razor".</summary>
+    </member>
+    <member name="P:System.Web.WebPages.Razor.Configuration.RazorWebSectionGroup.Host">
+      <summary>Gets or sets the host value for system.web.webPages.razor section group.</summary>
+      <returns>The host value.</returns>
+    </member>
+    <member name="P:System.Web.WebPages.Razor.Configuration.RazorWebSectionGroup.Pages">
+      <summary>Gets or sets the value of the pages element for the system.web.webPages.razor section.</summary>
+      <returns>The pages element value.</returns>
+    </member>
+  </members>
+</doc>
\ No newline at end of file
