Copying //UE4/Dev-Core to //UE4/Main

diff --git a/Engine/Source/Programs/CrashReporter/CrashReportWebSite/packages/Microsoft.AspNet.WebPages.3.2.3/lib/net45/System.Web.WebPages.xml b/Engine/Source/Programs/CrashReporter/CrashReportWebSite/packages/Microsoft.AspNet.WebPages.3.2.3/lib/net45/System.Web.WebPages.xml
new file mode 100644
--- /dev/null
+++ b/Engine/Source/Programs/CrashReporter/CrashReportWebSite/packages/Microsoft.AspNet.WebPages.3.2.3/lib/net45/System.Web.WebPages.xml
@@ -0,0 +1,2706 @@
+﻿<?xml version="1.0" encoding="utf-8"?>
+<doc>
+  <assembly>
+    <name>System.Web.WebPages</name>
+  </assembly>
+  <members>
+    <member name="T:System.Web.Helpers.AntiForgery">
+      <summary>Helps prevent malicious scripts from submitting forged page requests.</summary>
+    </member>
+    <member name="M:System.Web.Helpers.AntiForgery.GetHtml">
+      <summary>Adds an authenticating token to a form to help protect against request forgery.</summary>
+      <returns>Returns a string that contains the encrypted token value in a hidden HTML field.</returns>
+      <exception cref="T:System.ArgumentException">The current <see cref="T:System.Web.HttpContext" /> object is null.</exception>
+    </member>
+    <member name="M:System.Web.Helpers.AntiForgery.GetHtml(System.Web.HttpContextBase,System.String,System.String,System.String)">
+      <summary>Adds an authenticating token to a form to help protect against request forgery and lets callers specify authentication details.</summary>
+      <returns>Returns the encrypted token value in a hidden HTML field.</returns>
+      <param name="httpContext">The HTTP context data for a request.</param>
+      <param name="salt">An optional string of random characters (such as Z*7g1&amp;p4) that is used to add complexity to the encryption for extra safety. The default is null.</param>
+      <param name="domain">The domain of a web application that a request is submitted from.</param>
+      <param name="path">The virtual root path of a web application that a request is submitted from.</param>
+      <exception cref="T:System.ArgumentNullException">
+        <paramref name="httpContext" /> is null.</exception>
+    </member>
+    <member name="M:System.Web.Helpers.AntiForgery.GetTokens(System.String,System.String@,System.String@)">
+      <summary>Gets the search tokens.</summary>
+      <param name="oldCookieToken">The previous cookie token.</param>
+      <param name="newCookieToken">The new cookie token.</param>
+      <param name="formToken">The form of the token.</param>
+    </member>
+    <member name="M:System.Web.Helpers.AntiForgery.Validate">
+      <summary>Validates that input data from an HTML form field comes from the user who submitted the data.</summary>
+      <exception cref="T:System.ArgumentException">The current <see cref="T:System.Web.HttpContext" /> value is null.</exception>
+      <exception cref="T:System.Web.Helpers.HttpAntiForgeryException">The HTTP cookie token that accompanies a valid request is missing-or-The form token is missing.-or-The form token value does not match the cookie token value.-or-The form token value does not match the cookie token value.</exception>
+    </member>
+    <member name="M:System.Web.Helpers.AntiForgery.Validate(System.String,System.String)">
+      <summary>Validates that input data from an HTML form field comes from the user who submitted the data.</summary>
+      <param name="cookieToken">The cookie token value.</param>
+      <param name="formToken">The token form.</param>
+    </member>
+    <member name="M:System.Web.Helpers.AntiForgery.Validate(System.Web.HttpContextBase,System.String)">
+      <summary>Validates that input data from an HTML form field comes from the user who submitted the data and lets callers specify additional validation details.</summary>
+      <param name="httpContext">The HTTP context data for a request.</param>
+      <param name="salt">An optional string of random characters (such as Z*7g1&amp;p4) that is used to decrypt an authentication token created by the <see cref="T:System.Web.Helpers.AntiForgery" /> class. The default is null.</param>
+      <exception cref="T:System.ArgumentException">The current <see cref="T:System.Web.HttpContext" /> value is null.</exception>
+      <exception cref="T:System.Web.Helpers.HttpAntiForgeryException">The HTTP cookie token that accompanies a valid request is missing.-or-The form token is missing.-or-The form token value does not match the cookie token value.-or-The form token value does not match the cookie token value.-or-The <paramref name="salt" /> value supplied does not match the <paramref name="salt" /> value that was used to create the form token.</exception>
+    </member>
+    <member name="T:System.Web.Helpers.AntiForgeryConfig">
+      <summary>Provides programmatic configuration for the anti-forgery token system.</summary>
+    </member>
+    <member name="P:System.Web.Helpers.AntiForgeryConfig.AdditionalDataProvider">
+      <summary>Gets a data provider that can provide additional data to put into all generated tokens and that can validate additional data in incoming tokens.</summary>
+      <returns>The data provider.</returns>
+    </member>
+    <member name="P:System.Web.Helpers.AntiForgeryConfig.CookieName">
+      <summary>Gets or sets the name of the cookie that is used by the anti-forgery system.</summary>
+      <returns>The cookie name.</returns>
+    </member>
+    <member name="P:System.Web.Helpers.AntiForgeryConfig.RequireSsl">
+      <summary>Gets or sets a value that indicates whether the anti-forgery cookie requires SSL in order to be returned to the server.</summary>
+      <returns>true if SSL is required to return the anti-forgery cookie to the server; otherwise, false. </returns>
+    </member>
+    <member name="P:System.Web.Helpers.AntiForgeryConfig.SuppressIdentityHeuristicChecks">
+      <summary>Gets or sets a value that indicates whether the anti-forgery system should skip checking for conditions that might indicate misuse of the system.</summary>
+      <returns>true if the anti-forgery system should not check for possible misuse; otherwise, false.</returns>
+    </member>
+    <member name="P:System.Web.Helpers.AntiForgeryConfig.SuppressXFrameOptionsHeader">
+      <summary>Specifies whether to suppress the generation of X-Frame-Options header which is used to prevent ClickJacking. By default, the X-Frame-Options header is generated with the value SAMEORIGIN. If this setting is 'true', the X-Frame-Options header will not be generated for the response.</summary>
+    </member>
+    <member name="P:System.Web.Helpers.AntiForgeryConfig.UniqueClaimTypeIdentifier">
+      <summary>If claims-based authorization is in use, gets or sets the claim type from the identity that is used to uniquely identify the user.</summary>
+      <returns>The claim type.</returns>
+    </member>
+    <member name="T:System.Web.Helpers.IAntiForgeryAdditionalDataProvider">
+      <summary>Provides a way to include or validate custom data for anti-forgery tokens.</summary>
+    </member>
+    <member name="M:System.Web.Helpers.IAntiForgeryAdditionalDataProvider.GetAdditionalData(System.Web.HttpContextBase)">
+      <summary>Provides additional data to store for the anti-forgery tokens that are generated during this request.</summary>
+      <returns>The supplemental data to embed in the anti-forgery token.</returns>
+      <param name="context">Information about the current request.</param>
+    </member>
+    <member name="M:System.Web.Helpers.IAntiForgeryAdditionalDataProvider.ValidateAdditionalData(System.Web.HttpContextBase,System.String)">
+      <summary>Validates additional data that was embedded inside an incoming anti-forgery token.</summary>
+      <returns>true if the data is valid, or false if the data is invalid.</returns>
+      <param name="context">Information about the current request.</param>
+      <param name="additionalData">The supplemental data that was embedded in the token.</param>
+    </member>
+    <member name="T:System.Web.Helpers.UnvalidatedRequestValues">
+      <summary>Provides access to unvalidated form values in the <see cref="T:System.Web.HttpRequest" /> object.</summary>
+    </member>
+    <member name="P:System.Web.Helpers.UnvalidatedRequestValues.Form">
+      <summary>Gets a collection of unvalidated form values that were posted from the browser.</summary>
+      <returns>An unvalidated collection of form values.</returns>
+    </member>
+    <member name="P:System.Web.Helpers.UnvalidatedRequestValues.Item(System.String)">
+      <summary>Gets the specified unvalidated object from the collection of posted values in the <see cref="T:System.Web.HttpRequest" /> object.</summary>
+      <returns>The specified member, or null if the specified item is not found.</returns>
+    </member>
+    <member name="P:System.Web.Helpers.UnvalidatedRequestValues.QueryString">
+      <summary>Gets a collection of unvalidated query-string values.</summary>
+      <returns>A collection of unvalidated query-string values.</returns>
+    </member>
+    <member name="T:System.Web.Helpers.Validation">
+      <summary>Excludes fields of the Request object from being checked for potentially unsafe HTML markup and client script.</summary>
+    </member>
+    <member name="M:System.Web.Helpers.Validation.Unvalidated(System.Web.HttpRequest)">
+      <summary>Returns a version of form values, cookies, and query-string variables without checking them first for HTML markup and client script.</summary>
+      <returns>An object that contains unvalidated versions of the form and query-string values.</returns>
+      <param name="request">The <see cref="T:System.Web.HttpRequest" /> object that contains values to exclude from request validation.</param>
+    </member>
+    <member name="M:System.Web.Helpers.Validation.Unvalidated(System.Web.HttpRequest,System.String)">
+      <summary>Returns a value from the specified form field, cookie, or query-string variable without checking it first for HTML markup and client script.</summary>
+      <returns>A string that contains unvalidated text from the specified field, cookie, or query-string value.</returns>
+      <param name="request">The <see cref="T:System.Web.HttpRequest" /> object that contains values to exclude from validation.</param>
+      <param name="key">The name of the field to exclude from validation. <paramref name="key" /> can refer to a form field, to a cookie, or to the query-string variable.</param>
+    </member>
+    <member name="M:System.Web.Helpers.Validation.Unvalidated(System.Web.HttpRequestBase)">
+      <summary>Returns all values from the Request object (including form fields, cookies, and the query string) without checking them first for HTML markup and client script.</summary>
+      <returns>An object that contains unvalidated versions of the form, cookie, and query-string values.</returns>
+      <param name="request">The <see cref="T:System.Web.HttpRequest" /> object that contains values to exclude from validation.</param>
+    </member>
+    <member name="M:System.Web.Helpers.Validation.Unvalidated(System.Web.HttpRequestBase,System.String)">
+      <summary>Returns the specified value from the Request object without checking it first for HTML markup and client script.</summary>
+      <returns>A string that contains unvalidated text from the specified field, cookie, or query-string value.</returns>
+      <param name="request">The <see cref="T:System.Web.HttpRequestBase" /> object that contains values to exclude from validation.</param>
+      <param name="key">The name of the field to exclude from validation. <paramref name="key" /> can refer to a form field, to a cookie, or to the query-string variable.</param>
+    </member>
+    <member name="T:System.Web.Mvc.HttpAntiForgeryException">
+      <summary>This type/member supports the .NET Framework infrastructure and is not intended to be used directly from your code.</summary>
+    </member>
+    <member name="M:System.Web.Mvc.HttpAntiForgeryException.#ctor">
+      <summary>This member supports the .NET Framework infrastructure and is not intended to be used directly from your code.Initializes a new instance of the <see cref="T:System.Web.Mvc.HttpAntiForgeryException" /> class.</summary>
+    </member>
+    <member name="M:System.Web.Mvc.HttpAntiForgeryException.#ctor(System.String)">
+      <summary>This member supports the .NET Framework infrastructure and is not intended to be used directly from your code. Initializes a new instance of the <see cref="T:System.Web.Mvc.HttpAntiForgeryException" /> class.</summary>
+      <param name="message">The containing message.</param>
+    </member>
+    <member name="M:System.Web.Mvc.HttpAntiForgeryException.#ctor(System.String,System.Exception)">
+      <summary>This member supports the .NET Framework infrastructure and is not intended to be used directly from your code. Initializes a new instance of the <see cref="T:System.Web.Mvc.HttpAntiForgeryException" /> class.</summary>
+      <param name="message">The message.</param>
+      <param name="innerException">The inner exception.</param>
+    </member>
+    <member name="T:System.Web.Mvc.ModelClientValidationEqualToRule">
+      <summary>This type/member supports the .NET Framework infrastructure and is not intended to be used directly from your code.</summary>
+    </member>
+    <member name="M:System.Web.Mvc.ModelClientValidationEqualToRule.#ctor(System.String,System.Object)">
+      <summary>This type/member supports the .NET Framework infrastructure and is not intended to be used directly from your code.</summary>
+      <param name="errorMessage">The error message.</param>
+      <param name="other">The other.</param>
+    </member>
+    <member name="T:System.Web.Mvc.ModelClientValidationMaxLengthRule"></member>
+    <member name="M:System.Web.Mvc.ModelClientValidationMaxLengthRule.#ctor(System.String,System.Int32)"></member>
+    <member name="T:System.Web.Mvc.ModelClientValidationMinLengthRule"></member>
+    <member name="M:System.Web.Mvc.ModelClientValidationMinLengthRule.#ctor(System.String,System.Int32)"></member>
+    <member name="T:System.Web.Mvc.ModelClientValidationRangeRule">
+      <summary>This type/member supports the .NET Framework infrastructure and is not intended to be used directly from your code.</summary>
+    </member>
+    <member name="M:System.Web.Mvc.ModelClientValidationRangeRule.#ctor(System.String,System.Object,System.Object)">
+      <summary>This type/member supports the .NET Framework infrastructure and is not intended to be used directly from your code.</summary>
+      <param name="errorMessage">The error message.</param>
+      <param name="minValue">The minimum value.</param>
+      <param name="maxValue">The maximum value.</param>
+    </member>
+    <member name="T:System.Web.Mvc.ModelClientValidationRegexRule">
+      <summary>This type/member supports the .NET Framework infrastructure and is not intended to be used directly from your code.</summary>
+    </member>
+    <member name="M:System.Web.Mvc.ModelClientValidationRegexRule.#ctor(System.String,System.String)">
+      <summary>This type/member supports the .NET Framework infrastructure and is not intended to be used directly from your code. Initializes a new instance of the <see cref="T:System.Web.Mvc.ModelClientValidationRegexRule" /> class.</summary>
+      <param name="errorMessage">The exception message.</param>
+      <param name="pattern">The pattern.</param>
+    </member>
+    <member name="T:System.Web.Mvc.ModelClientValidationRemoteRule">
+      <summary>Represents the remote rule for the validation of the model client.</summary>
+    </member>
+    <member name="M:System.Web.Mvc.ModelClientValidationRemoteRule.#ctor(System.String,System.String,System.String,System.String)">
+      <summary>Initializes a new instance of the <see cref="T:System.Web.Mvc.ModelClientValidationRemoteRule" /> class.</summary>
+      <param name="errorMessage">The error message.</param>
+      <param name="url">The URL of the rule.</param>
+      <param name="httpMethod">The HTTP method.</param>
+      <param name="additionalFields">The additional fields used.</param>
+    </member>
+    <member name="T:System.Web.Mvc.ModelClientValidationRequiredRule">
+      <summary>Represents the required rule for the validation of the model client.</summary>
+    </member>
+    <member name="M:System.Web.Mvc.ModelClientValidationRequiredRule.#ctor(System.String)">
+      <summary>Initializes a new instance of the <see cref="T:System.Web.Mvc.ModelClientValidationRequiredRule" /> class.</summary>
+      <param name="errorMessage">The error message </param>
+    </member>
+    <member name="T:System.Web.Mvc.ModelClientValidationRule">
+      <summary>This type/member supports the .NET Framework infrastructure and is not intended to be used directly from your code.</summary>
+    </member>
+    <member name="M:System.Web.Mvc.ModelClientValidationRule.#ctor">
+      <summary>This type/member supports the .NET Framework infrastructure and is not intended to be used directly from your code.</summary>
+    </member>
+    <member name="P:System.Web.Mvc.ModelClientValidationRule.ErrorMessage">
+      <summary>This type/member supports the .NET Framework infrastructure and is not intended to be used directly from your code.</summary>
+    </member>
+    <member name="P:System.Web.Mvc.ModelClientValidationRule.ValidationParameters">
+      <summary>This type/member supports the .NET Framework infrastructure and is not intended to be used directly from your code.</summary>
+    </member>
+    <member name="P:System.Web.Mvc.ModelClientValidationRule.ValidationType">
+      <summary>This type/member supports the .NET Framework infrastructure and is not intended to be used directly from your code.</summary>
+    </member>
+    <member name="T:System.Web.Mvc.ModelClientValidationStringLengthRule">
+      <summary>This type/member supports the .NET Framework infrastructure and is not intended to be used directly from your code.Represents a length of the validation rule of the model client.</summary>
+    </member>
+    <member name="M:System.Web.Mvc.ModelClientValidationStringLengthRule.#ctor(System.String,System.Int32,System.Int32)">
+      <summary>This type/member supports the .NET Framework infrastructure and is not intended to be used directly from your code.Initializes a new instance of the <see cref="T:System.Web.Mvc.ModelClientValidationStringLengthRule" /> class.</summary>
+      <param name="errorMessage">The error message.</param>
+      <param name="minimumLength">The minimum length of the validation rule.</param>
+      <param name="maximumLength">The maximum length of the validation rule.</param>
+    </member>
+    <member name="T:System.Web.Mvc.TagBuilder">
+      <summary>Contains classes and properties that are used to create HTML elements. This class is used to write helpers, such as those found in the <see cref="N:System.Web.Helpers" /> namespace.</summary>
+    </member>
+    <member name="M:System.Web.Mvc.TagBuilder.#ctor(System.String)">
+      <summary>Creates a new tag that has the specified tag name.</summary>
+      <param name="tagName">The tag name without the "&lt;", "/", or "&gt;" delimiters.</param>
+      <exception cref="T:System.ArgumentException">
+        <paramref name="tagName" /> is null or empty.</exception>
+    </member>
+    <member name="M:System.Web.Mvc.TagBuilder.AddCssClass(System.String)">
+      <summary>Adds a CSS class to the list of CSS classes in the tag.</summary>
+      <param name="value">The CSS class to add.</param>
+    </member>
+    <member name="P:System.Web.Mvc.TagBuilder.Attributes">
+      <summary>Gets the collection of attributes.</summary>
+      <returns>The collection of attributes.</returns>
+    </member>
+    <member name="M:System.Web.Mvc.TagBuilder.CreateSanitizedId(System.String)">
+      <summary>Replaces each invalid character in the tag ID with a valid HTML character.</summary>
+      <returns>The sanitized tag ID, or null if <paramref name="originalId" /> is null or empty, or if <paramref name="originalId" /> does not begin with a letter.</returns>
+      <param name="originalId">The ID that might contain characters to replace.</param>
+    </member>
+    <member name="M:System.Web.Mvc.TagBuilder.CreateSanitizedId(System.String,System.String)">
+      <summary>Replaces each invalid character in the tag ID with the specified replacement string.</summary>
+      <returns>The sanitized tag ID, or null if <paramref name="originalId" /> is null or empty, or if <paramref name="originalId" /> does not begin with a letter.</returns>
+      <param name="originalId">The ID that might contain characters to replace.</param>
+      <param name="invalidCharReplacement">The replacement string.</param>
+      <exception cref="T:System.ArgumentNullException">
+        <paramref name="invalidCharReplacement" /> is null.</exception>
+    </member>
+    <member name="M:System.Web.Mvc.TagBuilder.GenerateId(System.String)">
+      <summary>Generates a sanitized ID attribute for the tag by using the specified name.</summary>
+      <param name="name">The name to use to generate an ID attribute.</param>
+    </member>
+    <member name="P:System.Web.Mvc.TagBuilder.IdAttributeDotReplacement">
+      <summary>Gets or sets a string that can be used to replace invalid HTML characters.</summary>
+      <returns>The string to use to replace invalid HTML characters.</returns>
+    </member>
+    <member name="P:System.Web.Mvc.TagBuilder.InnerHtml">
+      <summary>Gets or sets the inner HTML value for the element.</summary>
+      <returns>The inner HTML value for the element.</returns>
+    </member>
+    <member name="M:System.Web.Mvc.TagBuilder.MergeAttribute(System.String,System.String)">
+      <summary>Adds a new attribute to the tag.</summary>
+      <param name="key">The key for the attribute.</param>
+      <param name="value">The value of the attribute.</param>
+    </member>
+    <member name="M:System.Web.Mvc.TagBuilder.MergeAttribute(System.String,System.String,System.Boolean)">
+      <summary>Adds a new attribute or optionally replaces an existing attribute in the opening tag.</summary>
+      <param name="key">The key for the attribute.</param>
+      <param name="value">The value of the attribute.</param>
+      <param name="replaceExisting">true to replace an existing attribute if an attribute exists that has the specified <paramref name="key" /> value, or false to leave the original attribute unchanged.</param>
+    </member>
+    <member name="M:System.Web.Mvc.TagBuilder.MergeAttributes``2(System.Collections.Generic.IDictionary{``0,``1})">
+      <summary>Adds new attributes to the tag.</summary>
+      <param name="attributes">The collection of attributes to add.</param>
+      <typeparam name="TKey">The type of the key object.</typeparam>
+      <typeparam name="TValue">The type of the value object.</typeparam>
+    </member>
+    <member name="M:System.Web.Mvc.TagBuilder.MergeAttributes``2(System.Collections.Generic.IDictionary{``0,``1},System.Boolean)">
+      <summary>Adds new attributes or optionally replaces existing attributes in the tag.</summary>
+      <param name="attributes">The collection of attributes to add or replace.</param>
+      <param name="replaceExisting">For each attribute in <paramref name="attributes" />, true to replace the attribute if an attribute already exists that has the same key, or false to leave the original attribute unchanged.</param>
+      <typeparam name="TKey">The type of the key object.</typeparam>
+      <typeparam name="TValue">The type of the value object.</typeparam>
+    </member>
+    <member name="M:System.Web.Mvc.TagBuilder.SetInnerText(System.String)">
+      <summary>Sets the <see cref="P:System.Web.Mvc.TagBuilder.InnerHtml" /> property of the element to an HTML-encoded version of the specified string.</summary>
+      <param name="innerText">The string to HTML-encode.</param>
+    </member>
+    <member name="P:System.Web.Mvc.TagBuilder.TagName">
+      <summary>Gets the tag name for this tag.</summary>
+      <returns>The name.</returns>
+    </member>
+    <member name="M:System.Web.Mvc.TagBuilder.ToString">
+      <summary>Renders the element as a <see cref="F:System.Web.Mvc.TagRenderMode.Normal" /> element.</summary>
+    </member>
+    <member name="M:System.Web.Mvc.TagBuilder.ToString(System.Web.Mvc.TagRenderMode)">
+      <summary>Renders the HTML tag by using the specified render mode.</summary>
+      <returns>The rendered HTML tag.</returns>
+      <param name="renderMode">The render mode.</param>
+    </member>
+    <member name="T:System.Web.Mvc.TagRenderMode">
+      <summary>Enumerates the modes that are available for rendering HTML tags.</summary>
+    </member>
+    <member name="F:System.Web.Mvc.TagRenderMode.EndTag">
+      <summary>Represents the mode for rendering a closing tag (for example, &lt;/tag&gt;).</summary>
+    </member>
+    <member name="F:System.Web.Mvc.TagRenderMode.Normal">
+      <summary>Represents the mode for rendering normal text.</summary>
+    </member>
+    <member name="F:System.Web.Mvc.TagRenderMode.SelfClosing">
+      <summary>Represents the mode for rendering a self-closing tag (for example, &lt;tag /&gt;).</summary>
+    </member>
+    <member name="F:System.Web.Mvc.TagRenderMode.StartTag">
+      <summary>Represents the mode for rendering an opening tag (for example, &lt;tag&gt;).</summary>
+    </member>
+    <member name="T:System.Web.Mvc.UnobtrusiveValidationAttributesGenerator">
+      <summary>This type/member supports the .NET Framework infrastructure and is not intended to be used directly from your code.</summary>
+    </member>
+    <member name="M:System.Web.Mvc.UnobtrusiveValidationAttributesGenerator.GetValidationAttributes(System.Collections.Generic.IEnumerable{System.Web.Mvc.ModelClientValidationRule},System.Collections.Generic.IDictionary{System.String,System.Object})">
+      <summary>This type/member supports the .NET Framework infrastructure and is not intended to be used directly from your code.Gets the validation attributes from the structure or content of <see cref="T:System.Web.Mvc.UnobtrusiveValidationAttributesGenerator" />.</summary>
+      <param name="clientRules">The <see cref="T:System.Web.Mvc.ModelClientValidationRule" /> to be implemented.</param>
+      <param name="results">The result of the validation.</param>
+    </member>
+    <member name="T:System.Web.WebPages.ApplicationPart">
+      <summary>Contains methods to register assemblies as application parts.</summary>
+    </member>
+    <member name="M:System.Web.WebPages.ApplicationPart.#ctor(System.Reflection.Assembly,System.String)">
+      <summary>Initializes a new instance of the <see cref="T:System.Web.WebPages.ApplicationPart" /> class by using the specified assembly and root virtual path.</summary>
+      <param name="assembly">The assembly.</param>
+      <param name="rootVirtualPath">The root virtual path.</param>
+      <exception cref="T:System.ArgumentException">
+        <paramref name="rootVirtualPath" /> is null or empty.</exception>
+    </member>
+    <member name="M:System.Web.WebPages.ApplicationPart.ProcessVirtualPath(System.Reflection.Assembly,System.String,System.String)">
+      <summary>Resolves a path to the specified assembly or resource within an assembly by using the specified base virtual path and specified virtual path.</summary>
+      <returns>The path of the assembly or resource.</returns>
+      <param name="assembly">The assembly.</param>
+      <param name="baseVirtualPath">The base virtual path.</param>
+      <param name="virtualPath">The virtual path.</param>
+      <exception cref="T:System.InvalidOperationException">
+        <paramref name="assembly" /> is not registered.</exception>
+    </member>
+    <member name="M:System.Web.WebPages.ApplicationPart.Register(System.Web.WebPages.ApplicationPart)">
+      <summary>Adds an assembly and all web pages within the assembly to the list of available application parts.</summary>
+      <param name="applicationPart">The application part.</param>
+      <exception cref="T:System.InvalidOperationException">
+        <paramref name="applicationPart" /> is already registered.</exception>
+    </member>
+    <member name="T:System.Web.WebPages.ApplicationStartPage">
+      <summary>Provides objects and methods that are used to execute and render ASP.NET Web Pages application start pages (_AppStart.cshtml or _AppStart.vbhtml files).</summary>
+    </member>
+    <member name="M:System.Web.WebPages.ApplicationStartPage.#ctor">
+      <summary>Initializes a new instance of the <see cref="T:System.Web.WebPages.ApplicationStartPage" /> class.</summary>
+    </member>
+    <member name="P:System.Web.WebPages.ApplicationStartPage.Application">
+      <summary>Gets the HTTP application object that references this application startup page.</summary>
+      <returns>The HTTP application object that references this application startup page.</returns>
+    </member>
+    <member name="F:System.Web.WebPages.ApplicationStartPage.CacheKeyPrefix">
+      <summary>The prefix that is applied to all keys that are added to the cache by the application start page.</summary>
+    </member>
+    <member name="P:System.Web.WebPages.ApplicationStartPage.Context">
+      <summary>Gets the <see cref="T:System.Web.HttpContextBase" /> object that represents context data that is associated with this page.</summary>
+      <returns>The current context data.</returns>
+    </member>
+    <member name="M:System.Web.WebPages.ApplicationStartPage.GetOutputWriter">
+      <summary>Returns the text writer instance that is used to render the page.</summary>
+      <returns>The text writer.</returns>
+    </member>
+    <member name="P:System.Web.WebPages.ApplicationStartPage.Markup">
+      <summary>Gets the output from the application start page as an HTML-encoded string.</summary>
+      <returns>The output from the application start page as an HTML-encoded string.</returns>
+    </member>
+    <member name="P:System.Web.WebPages.ApplicationStartPage.Output">
+      <summary>Gets the text writer for the page.</summary>
+      <returns>The text writer for the page.</returns>
+    </member>
+    <member name="F:System.Web.WebPages.ApplicationStartPage.StartPageVirtualPath">
+      <summary>The path to the application start page.</summary>
+    </member>
+    <member name="P:System.Web.WebPages.ApplicationStartPage.VirtualPath">
+      <summary>Gets or sets the virtual path of the page.</summary>
+      <returns>The virtual path.</returns>
+    </member>
+    <member name="M:System.Web.WebPages.ApplicationStartPage.Write(System.Object)">
+      <summary>Writes the string representation of the specified object as an HTML-encoded string.</summary>
+      <param name="value">The object to encode and write.</param>
+    </member>
+    <member name="M:System.Web.WebPages.ApplicationStartPage.Write(System.Web.WebPages.HelperResult)">
+      <summary>Writes the specified <see cref="T:System.Web.WebPages.HelperResult" /> object as an HTML-encoded string.</summary>
+      <param name="result">The helper result to encode and write.</param>
+    </member>
+    <member name="M:System.Web.WebPages.ApplicationStartPage.WriteLiteral(System.Object)">
+      <summary>Writes the specified object without HTML encoding.</summary>
+      <param name="value">The object to write.</param>
+    </member>
+    <member name="T:System.Web.WebPages.AttributeValue">
+      <summary>Stores the value for an attribute.This type/member supports the .NET Framework infrastructure and is not intended to be used directly from your code.</summary>
+    </member>
+    <member name="M:System.Web.WebPages.AttributeValue.#ctor(System.Web.WebPages.Instrumentation.PositionTagged{System.String},System.Web.WebPages.Instrumentation.PositionTagged{System.Object},System.Boolean)">
+      <summary>Initializes a new instance of the <see cref="T:System.Web.WebPages.AttributeValue" /> class.This type/member supports the .NET Framework infrastructure and is not intended to be used directly from your code.</summary>
+      <param name="prefix">The namespace prefix for the attribute.</param>
+      <param name="value">The value for the attribute.</param>
+      <param name="literal">true to indicate that the value is a literal value; otherwise, false.</param>
+    </member>
+    <member name="M:System.Web.WebPages.AttributeValue.FromTuple(System.Tuple{System.Tuple{System.String,System.Int32},System.Tuple{System.Object,System.Int32},System.Boolean})">
+      <summary>Creates an attribute value from the specified tuple object.</summary>
+      <returns>The created attribute value.</returns>
+      <param name="value">The tuple object from which to create from.</param>
+    </member>
+    <member name="M:System.Web.WebPages.AttributeValue.FromTuple(System.Tuple{System.Tuple{System.String,System.Int32},System.Tuple{System.String,System.Int32},System.Boolean})">
+      <summary>Creates an attribute value from the specified tuple object.</summary>
+      <returns>The created attribute value.</returns>
+      <param name="value">The tuple object from which to create from.</param>
+    </member>
+    <member name="P:System.Web.WebPages.AttributeValue.Literal">
+      <summary>Gets or sets a value that indicates whether the value is a literal value.This type/member supports the .NET Framework infrastructure and is not intended to be used directly from your code.</summary>
+      <returns>true if the value is a literal value; otherwise, false.</returns>
+    </member>
+    <member name="M:System.Web.WebPages.AttributeValue.op_Implicit(System.Tuple{System.Tuple{System.String,System.Int32},System.Tuple{System.Object,System.Int32},System.Boolean})~System.Web.WebPages.AttributeValue">
+      <summary>Creates an attribute value from the specified tuple object.</summary>
+      <returns>The created attribute value.</returns>
+      <param name="value">The tuple object from which to create from.</param>
+    </member>
+    <member name="M:System.Web.WebPages.AttributeValue.op_Implicit(System.Tuple{System.Tuple{System.String,System.Int32},System.Tuple{System.String,System.Int32},System.Boolean})~System.Web.WebPages.AttributeValue">
+      <summary>Creates an attribute value from the specified tuple object.</summary>
+      <returns>The created attribute value.</returns>
+      <param name="value">The tuple object from which to create from.</param>
+    </member>
+    <member name="P:System.Web.WebPages.AttributeValue.Prefix">
+      <summary>Gets or sets the namespace prefix for the attribute.This type/member supports the .NET Framework infrastructure and is not intended to be used directly from your code.</summary>
+      <returns>The namespace prefix for the attribute.</returns>
+    </member>
+    <member name="P:System.Web.WebPages.AttributeValue.Value">
+      <summary>Gets or set the value for the attribute.This type/member supports the .NET Framework infrastructure and is not intended to be used directly from your code.</summary>
+      <returns>The value for the attribute.</returns>
+    </member>
+    <member name="T:System.Web.WebPages.BrowserHelpers">
+      <summary>Provides a way to specify custom browser (user agent) information.</summary>
+    </member>
+    <member name="M:System.Web.WebPages.BrowserHelpers.ClearOverriddenBrowser(System.Web.HttpContextBase)">
+      <summary>Removes any overridden user agent for the current request.</summary>
+      <param name="httpContext">The current context.</param>
+    </member>
+    <member name="M:System.Web.WebPages.BrowserHelpers.GetOverriddenBrowser(System.Web.HttpContextBase)">
+      <summary>Returns the browser capabilities object for the overridden browser capabilities or for the actual browser if no override has been specified.</summary>
+      <returns>The browser capabilities.</returns>
+      <param name="httpContext">The current context.</param>
+    </member>
+    <member name="M:System.Web.WebPages.BrowserHelpers.GetOverriddenUserAgent(System.Web.HttpContextBase)">
+      <summary>Returns the overridden user agent value or the actual user agent string if no override has been specified.</summary>
+      <returns>The user agent string</returns>
+      <param name="httpContext">The current context.</param>
+    </member>
+    <member name="M:System.Web.WebPages.BrowserHelpers.GetVaryByCustomStringForOverriddenBrowser(System.Web.HttpContext)">
+      <summary>Gets a string that varies based on the type of the browser.</summary>
+      <returns>A string that identifies the browser.</returns>
+      <param name="httpContext">The current context.</param>
+    </member>
+    <member name="M:System.Web.WebPages.BrowserHelpers.GetVaryByCustomStringForOverriddenBrowser(System.Web.HttpContextBase)">
+      <summary>Gets a string that varies based on the type of the browser. </summary>
+      <returns>A string that identifies the browser.</returns>
+      <param name="httpContext">The current context base.</param>
+    </member>
+    <member name="M:System.Web.WebPages.BrowserHelpers.SetOverriddenBrowser(System.Web.HttpContextBase,System.String)">
+      <summary>Overrides the request's actual user agent value using the specified user agent.</summary>
+      <param name="httpContext">The current context.</param>
+      <param name="userAgent">The user agent to use.</param>
+    </member>
+    <member name="M:System.Web.WebPages.BrowserHelpers.SetOverriddenBrowser(System.Web.HttpContextBase,System.Web.WebPages.BrowserOverride)">
+      <summary>Overrides the request's actual user agent value using the specified browser override information.</summary>
+      <param name="httpContext">The current context.</param>
+      <param name="browserOverride">One of the enumeration values that represents the browser override information to use.</param>
+    </member>
+    <member name="T:System.Web.WebPages.BrowserOverride">
+      <summary>Specifies browser types that can be defined for the <see cref="M:System.Web.WebPages.BrowserHelpers.SetOverriddenBrowser(System.Web.HttpContextBase,System.Web.WebPages.BrowserOverride)" /> method.</summary>
+    </member>
+    <member name="F:System.Web.WebPages.BrowserOverride.Desktop">
+      <summary>Specifies a desktop browser.</summary>
+    </member>
+    <member name="F:System.Web.WebPages.BrowserOverride.Mobile">
+      <summary>Specifies a mobile browser.</summary>
+    </member>
+    <member name="T:System.Web.WebPages.BrowserOverrideStore">
+      <summary>This type/member supports the .NET Framework infrastructure and is not intended to be used directly from your code.The current BrowserOverrideStore is used to get and set the user agent of a request.</summary>
+    </member>
+    <member name="M:System.Web.WebPages.BrowserOverrideStore.#ctor">
+      <summary>This type/member supports the .NET Framework infrastructure and is not intended to be used directly from your code.Initializes a new instance of the <see cref="T:System.Web.WebPages.BrowserOverrideStore" /> class.</summary>
+    </member>
+    <member name="M:System.Web.WebPages.BrowserOverrideStore.GetOverriddenUserAgent(System.Web.HttpContextBase)">
+      <summary>This type/member supports the .NET Framework infrastructure and is not intended to be used directly from your code.Looks for a user agent by searching for the browser override cookie.</summary>
+      <returns>The user agent.</returns>
+      <param name="httpContext">The HTTP context.</param>
+    </member>
+    <member name="M:System.Web.WebPages.BrowserOverrideStore.SetOverriddenUserAgent(System.Web.HttpContextBase,System.String)">
+      <summary>This type/member supports the .NET Framework infrastructure and is not intended to be used directly from your code.Adds a browser override cookie with the set user agent to the response of the current request.</summary>
+      <param name="httpContext">The HTTP context.</param>
+      <param name="userAgent">The user agent.</param>
+    </member>
+    <member name="T:System.Web.WebPages.BrowserOverrideStores">
+      <summary>This type/member supports the .NET Framework infrastructure and is not intended to be used directly from your code.</summary>
+    </member>
+    <member name="M:System.Web.WebPages.BrowserOverrideStores.#ctor">
+      <summary>This type/member supports the .NET Framework infrastructure and is not intended to be used directly from your code.</summary>
+    </member>
+    <member name="P:System.Web.WebPages.BrowserOverrideStores.Current">
+      <summary>This type/member supports the .NET Framework infrastructure and is not intended to be used directly from your code.</summary>
+    </member>
+    <member name="T:System.Web.WebPages.CookieBrowserOverrideStore">
+      <summary>This type/member supports the .NET Framework infrastructure and is not intended to be used directly from your code.Gets overridden user agent for a request from a cookie. Creates a cookie to set the overridden user agent.</summary>
+    </member>
+    <member name="M:System.Web.WebPages.CookieBrowserOverrideStore.#ctor">
+      <summary>This type/member supports the .NET Framework infrastructure and is not intended to be used directly from your code.Initializes a new instance of the <see cref="T:System.Web.WebPages.CookieBrowserOverrideStore" /> class.</summary>
+    </member>
+    <member name="M:System.Web.WebPages.CookieBrowserOverrideStore.#ctor(System.Int32)">
+      <summary>This type/member supports the .NET Framework infrastructure and is not intended to be used directly from your code.Initializes a new instance of the <see cref="T:System.Web.WebPages.CookieBrowserOverrideStore" /> class.</summary>
+      <param name="daysToExpire">The days to expire.</param>
+    </member>
+    <member name="M:System.Web.WebPages.CookieBrowserOverrideStore.GetOverriddenUserAgent(System.Web.HttpContextBase)">
+      <summary>This type/member supports the .NET Framework infrastructure and is not intended to be used directly from your code.Looks for a user agent by searching for the browser override cookie.</summary>
+      <returns>The user agent.</returns>
+      <param name="httpContext">The HTTP context.</param>
+    </member>
+    <member name="M:System.Web.WebPages.CookieBrowserOverrideStore.SetOverriddenUserAgent(System.Web.HttpContextBase,System.String)">
+      <summary>This type/member supports the .NET Framework infrastructure and is not intended to be used directly from your code.Adds a browser override cookie with the set user agent to the response of the current request.</summary>
+      <param name="httpContext">The HTTP context.</param>
+      <param name="userAgent">The user agent.</param>
+    </member>
+    <member name="T:System.Web.WebPages.DefaultDisplayMode">
+      <summary>This type/member supports the .NET Framework infrastructure and is not intended to be used directly from your code.Represents the default display mode of the web pages.</summary>
+    </member>
+    <member name="M:System.Web.WebPages.DefaultDisplayMode.#ctor">
+      <summary>This type/member supports the .NET Framework infrastructure and is not intended to be used directly from your code.Initializes a new instance of the <see cref="T:System.Web.WebPages.DefaultDisplayMode" /> class.</summary>
+    </member>
+    <member name="M:System.Web.WebPages.DefaultDisplayMode.#ctor(System.String)">
+      <summary>This type/member supports the .NET Framework infrastructure and is not intended to be used directly from your code.Initializes a new instance of the <see cref="T:System.Web.WebPages.DefaultDisplayMode" /> class.</summary>
+      <param name="suffix">The suffix.</param>
+    </member>
+    <member name="M:System.Web.WebPages.DefaultDisplayMode.CanHandleContext(System.Web.HttpContextBase)">
+      <summary>This type/member supports the .NET Framework infrastructure and is not intended to be used directly from your code.Indicates a value whether the <see cref="T:System.Web.HttpContextBase" /> can handle context.</summary>
+      <returns>true if the <see cref="T:System.Web.HttpContextBase" /> can handle context; otherwise, false.</returns>
+      <param name="httpContext">The specified http context.</param>
+    </member>
+    <member name="P:System.Web.WebPages.DefaultDisplayMode.ContextCondition">
+      <summary>This type/member supports the .NET Framework infrastructure and is not intended to be used directly from your code.Gets or sets a value that indicates whether the context condition displays a default mode.</summary>
+      <returns>true if the context condition displays a default mode; otherwise, false.</returns>
+    </member>
+    <member name="P:System.Web.WebPages.DefaultDisplayMode.DisplayModeId">
+      <summary>This type/member supports the .NET Framework infrastructure and is not intended to be used directly from your code.Gets the display mode identifier of the web pages.</summary>
+      <returns>The display mode identifier of the web pages.</returns>
+    </member>
+    <member name="M:System.Web.WebPages.DefaultDisplayMode.GetDisplayInfo(System.Web.HttpContextBase,System.String,System.Func{System.String,System.Boolean})">
+      <summary>Retrieves the display information about an item in the result pane.</summary>
+      <returns>The display information about an item in the result pane.</returns>
+      <param name="httpContext">The http context.</param>
+      <param name="virtualPath">The virtual path.</param>
+      <param name="virtualPathExists">true if the virtual path exists; otherwise, false.</param>
+    </member>
+    <member name="M:System.Web.WebPages.DefaultDisplayMode.TransformPath(System.String,System.String)">
+      <summary>This type/member supports the .NET Framework infrastructure and is not intended to be used directly from your code.Transforms the path of the display mode.</summary>
+      <returns>The path of the display mode to transform.</returns>
+      <param name="virtualPath">The virtual path.</param>
+      <param name="suffix">The suffix.</param>
+    </member>
+    <member name="T:System.Web.WebPages.DisplayInfo">
+      <summary>This type/member supports the .NET Framework infrastructure and is not intended to be used directly from your code.Represents a property’s display information.</summary>
+    </member>
+    <member name="M:System.Web.WebPages.DisplayInfo.#ctor(System.String,System.Web.WebPages.IDisplayMode)">
+      <summary>This type/member supports the .NET Framework infrastructure and is not intended to be used directly from your code.Initializes a new instance of the <see cref="T:System.Web.WebPages.DisplayInfo" /> class.</summary>
+      <param name="filePath">The virtual path.</param>
+      <param name="displayMode">The active display mode.</param>
+    </member>
+    <member name="P:System.Web.WebPages.DisplayInfo.DisplayMode">
+      <summary>This type/member supports the .NET Framework infrastructure and is not intended to be used directly from your code.Gets the active display mode for a Web page.</summary>
+      <returns>The active display mode for a Web page.</returns>
+    </member>
+    <member name="P:System.Web.WebPages.DisplayInfo.FilePath">
+      <summary>This type/member supports the .NET Framework infrastructure and is not intended to be used directly from your code.Gets the virtual path of the current Web page.</summary>
+      <returns>The virtual path of the current Web page.</returns>
+    </member>
+    <member name="T:System.Web.WebPages.DisplayModeProvider">
+      <summary>This type/member supports the .NET Framework infrastructure and is not intended to be used directly from your code.Represents the modes of display for the provider.</summary>
+    </member>
+    <member name="F:System.Web.WebPages.DisplayModeProvider.DefaultDisplayModeId">
+      <summary>This type/member supports the .NET Framework infrastructure and is not intended to be used directly from your code.Defines the default display mode identifier.</summary>
+    </member>
+    <member name="M:System.Web.WebPages.DisplayModeProvider.GetAvailableDisplayModesForContext(System.Web.HttpContextBase,System.Web.WebPages.IDisplayMode)">
+      <summary>This type/member supports the .NET Framework infrastructure and is not intended to be used directly from your code.Gets a list of the available display modes for the context base.</summary>
+      <returns>A list of the available display modes for the context base.</returns>
+      <param name="httpContext">The http context base.</param>
+      <param name="currentDisplayMode">The current display mode.</param>
+    </member>
+    <member name="M:System.Web.WebPages.DisplayModeProvider.GetDisplayInfoForVirtualPath(System.String,System.Web.HttpContextBase,System.Func{System.String,System.Boolean},System.Web.WebPages.IDisplayMode)">
+      <summary>This type/member supports the .NET Framework infrastructure and is not intended to be used directly from your code.Gets the <see cref="T:System.Web.WebPages.DisplayInfo" /> for the virtual path.</summary>
+      <returns>The <see cref="T:System.Web.WebPages.DisplayInfo" /> for the virtual path.</returns>
+      <param name="virtualPath">The virtual path.</param>
+      <param name="httpContext">The http context base.</param>
+      <param name="virtualPathExists">true if the virtual path exists; otherwise, false.</param>
+      <param name="currentDisplayMode">The current display mode.</param>
+    </member>
+    <member name="P:System.Web.WebPages.DisplayModeProvider.Instance">
+      <summary>This type/member supports the .NET Framework infrastructure and is not intended to be used directly from your code.Gets the instance of the <see cref="T:System.Web.WebPages.DisplayModeProvider" />.</summary>
+      <returns>The instance of the <see cref="T:System.Web.WebPages.DisplayModeProvider" />.</returns>
+    </member>
+    <member name="F:System.Web.WebPages.DisplayModeProvider.MobileDisplayModeId">
+      <summary>This type/member supports the .NET Framework infrastructure and is not intended to be used directly from your code.Defines the mobile display mode identifier.</summary>
+    </member>
+    <member name="P:System.Web.WebPages.DisplayModeProvider.Modes">
+      <summary>This type/member supports the .NET Framework infrastructure and is not intended to be used directly from your code.Gets a list of modes of the <see cref="T:System.Web.WebPages.DisplayModeProvider" />.</summary>
+      <returns>A list of modes of the <see cref="T:System.Web.WebPages.DisplayModeProvider" />.</returns>
+    </member>
+    <member name="P:System.Web.WebPages.DisplayModeProvider.RequireConsistentDisplayMode">
+      <summary>This type/member supports the .NET Framework infrastructure and is not intended to be used directly from your code.Gets or sets a value that indicates whether the Web page requires consistent display mode.</summary>
+      <returns>true if the Web page requires consistent display mode; otherwise, false.</returns>
+    </member>
+    <member name="T:System.Web.WebPages.HelperPage">
+      <summary>Represents a base class for pages that is used when ASP.NET compiles a .cshtml or .vbhtml file and that exposes page-level and application-level properties and methods.</summary>
+    </member>
+    <member name="M:System.Web.WebPages.HelperPage.#ctor">
+      <summary>This type/member supports the .NET Framework infrastructure and is not intended to be used directly from your code.Initializes a new instance of the <see cref="T:System.Web.WebPages.HelperPage" /> class.</summary>
+    </member>
+    <member name="P:System.Web.WebPages.HelperPage.App">
+      <summary>Gets the application-state data as a <see cref="T:System.Dynamic.DynamicObject" /> object that callers can use to create and access custom application-scoped properties.</summary>
+      <returns>The application-state data.</returns>
+    </member>
+    <member name="P:System.Web.WebPages.HelperPage.AppState">
+      <summary>Gets a reference to global application-state data that can be shared across sessions and requests in an ASP.NET application.</summary>
+      <returns>The application-state data.</returns>
+    </member>
+    <member name="M:System.Web.WebPages.HelperPage.BeginContext(System.IO.TextWriter,System.String,System.Int32,System.Int32,System.Boolean)">
+      <summary>This type/member supports the .NET Framework infrastructure and is not intended to be used directly from your code.Puts all the helper statements into the context of the helper page.</summary>
+      <param name="writer">The text writer.</param>
+      <param name="virtualPath">The helper virtual path.</param>
+      <param name="startPosition">The starting position.</param>
+      <param name="length">The length of the context.</param>
+      <param name="isLiteral">true of the context has a literal attribute; otherwise, false.</param>
+    </member>
+    <member name="M:System.Web.WebPages.HelperPage.BeginContext(System.String,System.Int32,System.Int32,System.Boolean)">
+      <summary>This type/member supports the .NET Framework infrastructure and is not intended to be used directly from your code.Puts all the helper statements into the context of the helper page.</summary>
+      <param name="virtualPath">The helper virtual path.</param>
+      <param name="startPosition">The starting position.</param>
+      <param name="length">The length of the context.</param>
+      <param name="isLiteral">true of the context has a literal attribute; otherwise, false.</param>
+    </member>
+    <member name="P:System.Web.WebPages.HelperPage.Cache">
+      <summary>Gets the cache object for the current application domain.</summary>
+      <returns>The cache object.</returns>
+    </member>
+    <member name="P:System.Web.WebPages.HelperPage.Context">
+      <summary>Gets the <see cref="T:System.Web.HttpContextBase" /> object that is associated with a page.</summary>
+      <returns>The current context data.</returns>
+    </member>
+    <member name="P:System.Web.WebPages.HelperPage.CurrentPage">
+      <summary>Gets the current page for this helper page.</summary>
+      <returns>The current page.</returns>
+    </member>
+    <member name="M:System.Web.WebPages.HelperPage.EndContext(System.IO.TextWriter,System.String,System.Int32,System.Int32,System.Boolean)">
+      <summary>This type/member supports the .NET Framework infrastructure and is not intended to be used directly from your code.Indicates the end of context block.</summary>
+      <param name="writer">The text writer.</param>
+      <param name="virtualPath">The helper virtual path.</param>
+      <param name="startPosition">The starting position.</param>
+      <param name="length">The length of the context.</param>
+      <param name="isLiteral">true of the context has a literal attribute; otherwise, false.</param>
+    </member>
+    <member name="M:System.Web.WebPages.HelperPage.EndContext(System.String,System.Int32,System.Int32,System.Boolean)">
+      <summary>This type/member supports the .NET Framework infrastructure and is not intended to be used directly from your code.Indicates the end of context block.</summary>
+      <param name="virtualPath">The helper virtual path.</param>
+      <param name="startPosition">The starting position.</param>
+      <param name="length">The length of the context.</param>
+      <param name="isLiteral">true of the context has a literal attribute; otherwise, false.</param>
+    </member>
+    <member name="P:System.Web.WebPages.HelperPage.HelperVirtualPath">
+      <summary>This type/member supports the .NET Framework infrastructure and is not intended to be used directly from your code.Gets or sets the path of the helper page.</summary>
+      <returns>The path of the helper page.</returns>
+    </member>
+    <member name="M:System.Web.WebPages.HelperPage.Href(System.String,System.Object[])">
+      <summary>Builds an absolute URL from an application-relative URL by using the specified parameters.</summary>
+      <returns>The absolute URL.</returns>
+      <param name="path">The initial path to use in the URL.</param>
+      <param name="pathParts">Additional path information, such as folders and subfolders.</param>
+    </member>
+    <member name="P:System.Web.WebPages.HelperPage.Html">
+      <summary>Gets the <see cref="T:System.Web.WebPages.Html.HtmlHelper" /> object that is associated with a page.</summary>
+      <returns>An object that supports rendering HTML form controls in a page.</returns>
+    </member>
+    <member name="P:System.Web.WebPages.HelperPage.IsAjax">
+      <summary>Gets a value that indicates whether Ajax is being used during the request of the web page.</summary>
+      <returns>true if Ajax is being used during the request; otherwise, false.</returns>
+    </member>
+    <member name="P:System.Web.WebPages.HelperPage.IsPost">
+      <summary>Gets a value that indicates whether the current request is a post (submitted using the HTTP POST verb).</summary>
+      <returns>true if the HTTP verb is POST; otherwise, false.</returns>
+    </member>
+    <member name="P:System.Web.WebPages.HelperPage.Model">
+      <summary>Gets the model that is associated with a page.</summary>
+      <returns>An object that represents a model that is associated with the view data for a page.</returns>
+    </member>
+    <member name="P:System.Web.WebPages.HelperPage.ModelState">
+      <summary>Gets the state data for the model that is associated with a page.</summary>
+      <returns>The state of the model.</returns>
+    </member>
+    <member name="P:System.Web.WebPages.HelperPage.Page">
+      <summary>Gets property-like access to page data that is shared between pages, layout pages, and partial pages.</summary>
+      <returns>An object that contains page data.</returns>
+    </member>
+    <member name="P:System.Web.WebPages.HelperPage.PageContext">
+      <summary>Gets and sets the HTTP context for the web page.</summary>
+      <returns>The HTTP context for the web page.</returns>
+    </member>
+    <member name="P:System.Web.WebPages.HelperPage.PageData">
+      <summary>Gets array-like access to page data that is shared between pages, layout pages, and partial pages.</summary>
+      <returns>An object that provides array-like access to page data.</returns>
+    </member>
+    <member name="P:System.Web.WebPages.HelperPage.Request">
+      <summary>Gets the <see cref="T:System.Web.HttpRequest" /> object for the current HTTP request.</summary>
+      <returns>An <see cref="T:System.Web.HttpRequest" /> object that contains the HTTP values that were sent by a client during a web request.</returns>
+    </member>
+    <member name="P:System.Web.WebPages.HelperPage.Response">
+      <summary>Gets the <see cref="T:System.Web.HttpResponse" /> object for the current HTTP response.</summary>
+      <returns>An <see cref="T:System.Web.HttpResponse" /> object that contains the HTTP-response information from an ASP.NET operation.</returns>
+    </member>
+    <member name="P:System.Web.WebPages.HelperPage.Server">
+      <summary>Gets the <see cref="T:System.Web.HttpServerUtility" /> object that provides methods that can be used as part of web-page processing.</summary>
+      <returns>The <see cref="T:System.Web.HttpServerUtility" /> object.</returns>
+    </member>
+    <member name="P:System.Web.WebPages.HelperPage.Session">
+      <summary>Gets the <see cref="T:System.Web.HttpSessionState" /> object for the current HTTP request.</summary>
+      <returns>The <see cref="T:System.Web.HttpSessionState" /> object for the current HTTP request.</returns>
+    </member>
+    <member name="P:System.Web.WebPages.HelperPage.UrlData">
+      <summary>Gets data related to the URL path.</summary>
+      <returns>Data related to the URL path.</returns>
+    </member>
+    <member name="P:System.Web.WebPages.HelperPage.User">
+      <summary>Gets a user value based on the HTTP context.</summary>
+      <returns>A user value based on the HTTP context.</returns>
+    </member>
+    <member name="P:System.Web.WebPages.HelperPage.VirtualPath">
+      <summary>Gets the virtual path of the page.</summary>
+      <returns>The virtual path.</returns>
+    </member>
+    <member name="M:System.Web.WebPages.HelperPage.WriteAttributeTo(System.IO.TextWriter,System.String,System.Web.WebPages.Instrumentation.PositionTagged{System.String},System.Web.WebPages.Instrumentation.PositionTagged{System.String},System.Web.WebPages.AttributeValue[])">
+      <summary>This type/member supports the .NET Framework infrastructure and is not intended to be used directly from your code.Writes an attribute associated with the helper.</summary>
+      <param name="writer">The text writer.</param>
+      <param name="name">The name of the attribute.</param>
+      <param name="prefix">The prefix.</param>
+      <param name="suffix">The suffix.</param>
+      <param name="values">The attribute value.</param>
+    </member>
+    <member name="M:System.Web.WebPages.HelperPage.WriteLiteralTo(System.IO.TextWriter,System.Object)">
+      <summary>This type/member supports the .NET Framework infrastructure and is not intended to be used directly from your code.Writes a literal object to the helper.</summary>
+      <param name="writer">The text writer.</param>
+      <param name="value">The value of the object.</param>
+    </member>
+    <member name="M:System.Web.WebPages.HelperPage.WriteLiteralTo(System.IO.TextWriter,System.Web.WebPages.HelperResult)">
+      <summary>This type/member supports the .NET Framework infrastructure and is not intended to be used directly from your code.Writes a helper result object to the helper.</summary>
+      <param name="writer">The text writer</param>
+      <param name="value">The helper result.</param>
+    </member>
+    <member name="M:System.Web.WebPages.HelperPage.WriteTo(System.IO.TextWriter,System.Object)">
+      <summary>This type/member supports the .NET Framework infrastructure and is not intended to be used directly from your code.Writes an object to the helper.</summary>
+      <param name="writer">The text writer.</param>
+      <param name="value">The object value.</param>
+    </member>
+    <member name="M:System.Web.WebPages.HelperPage.WriteTo(System.IO.TextWriter,System.Web.WebPages.HelperResult)">
+      <summary>This type/member supports the .NET Framework infrastructure and is not intended to be used directly from your code.Writes a helper result object to the helper.</summary>
+      <param name="writer">The text writer.</param>
+      <param name="value">The helper result value.</param>
+    </member>
+    <member name="T:System.Web.WebPages.HelperResult">
+      <summary>This type/member supports the .NET Framework infrastructure and is not intended to be used directly from your code.</summary>
+    </member>
+    <member name="M:System.Web.WebPages.HelperResult.#ctor(System.Action{System.IO.TextWriter})">
+      <summary>This type/member supports the .NET Framework infrastructure and is not intended to be used directly from your code.</summary>
+    </member>
+    <member name="M:System.Web.WebPages.HelperResult.ToHtmlString">
+      <summary>This type/member supports the .NET Framework infrastructure and is not intended to be used directly from your code.</summary>
+    </member>
+    <member name="M:System.Web.WebPages.HelperResult.ToString">
+      <summary>This type/member supports the .NET Framework infrastructure and is not intended to be used directly from your code.</summary>
+    </member>
+    <member name="M:System.Web.WebPages.HelperResult.WriteTo(System.IO.TextWriter)">
+      <summary>This type/member supports the .NET Framework infrastructure and is not intended to be used directly from your code.</summary>
+    </member>
+    <member name="T:System.Web.WebPages.HttpContextExtensions">
+      <summary>This type/member supports the .NET Framework infrastructure and is not intended to be used directly from your code.</summary>
+    </member>
+    <member name="M:System.Web.WebPages.HttpContextExtensions.RedirectLocal(System.Web.HttpContextBase,System.String)">
+      <summary>This type/member supports the .NET Framework infrastructure and is not intended to be used directly from your code.</summary>
+    </member>
+    <member name="M:System.Web.WebPages.HttpContextExtensions.RegisterForDispose(System.Web.HttpContextBase,System.IDisposable)">
+      <summary>This type/member supports the .NET Framework infrastructure and is not intended to be used directly from your code..</summary>
+    </member>
+    <member name="T:System.Web.WebPages.IDisplayMode">
+      <summary>This type/member supports the .NET Framework infrastructure and is not intended to be used directly from your code.Represents the display mode interface for the web pages.</summary>
+    </member>
+    <member name="M:System.Web.WebPages.IDisplayMode.CanHandleContext(System.Web.HttpContextBase)">
+      <summary>This type/member supports the .NET Framework infrastructure and is not intended to be used directly from your code.Indicates a value whether the web pages can handle HTTP context.</summary>
+      <returns>true if the web pages can handle HTTP context; otherwise, false.</returns>
+      <param name="httpContext">The HTTP context.</param>
+    </member>
+    <member name="P:System.Web.WebPages.IDisplayMode.DisplayModeId">
+      <summary>This type/member supports the .NET Framework infrastructure and is not intended to be used directly from your code.Gets the display mode id for the web pages.</summary>
+      <returns>The display mode id for the web pages.</returns>
+    </member>
+    <member name="M:System.Web.WebPages.IDisplayMode.GetDisplayInfo(System.Web.HttpContextBase,System.String,System.Func{System.String,System.Boolean})">
+      <summary>Returns this method to display all the information for the web pages.</summary>
+      <returns>The method to display all the information for the web pages.</returns>
+      <param name="httpContext">The HTTP context.</param>
+      <param name="virtualPath">The virtual path.</param>
+      <param name="virtualPathExists">true if the virtual path exists; otherwise, false.</param>
+    </member>
+    <member name="T:System.Web.WebPages.ITemplateFile">
+      <summary>This type/member supports the .NET Framework infrastructure and is not intended to be used directly from your code.</summary>
+    </member>
+    <member name="P:System.Web.WebPages.ITemplateFile.TemplateInfo">
+      <summary>This type/member supports the .NET Framework infrastructure and is not intended to be used directly from your code.</summary>
+    </member>
+    <member name="T:System.Web.WebPages.IValidator">
+      <summary>This type/member supports the .NET Framework infrastructure and is not intended to be used directly from your code.Defines the properties and methods that objects that participate in webpages.</summary>
+    </member>
+    <member name="P:System.Web.WebPages.IValidator.ClientValidationRule">
+      <summary>This type/member supports the .NET Framework infrastructure and is not intended to be used directly from your code.Gets a container for client validation for the required field.</summary>
+      <returns>A container for client validation for the required field.</returns>
+    </member>
+    <member name="M:System.Web.WebPages.IValidator.Validate(System.ComponentModel.DataAnnotations.ValidationContext)">
+      <summary>This type/member supports the .NET Framework infrastructure and is not intended to be used directly from your code.Evaluates the condition it checks and updates the validation context.</summary>
+      <returns>The condition it checks and updates the validation context.</returns>
+      <param name="validationContext">The validation context.</param>
+    </member>
+    <member name="T:System.Web.WebPages.IVirtualPathFactory">
+      <summary>Defines methods that are implemented by virtual path handler factories.</summary>
+    </member>
+    <member name="M:System.Web.WebPages.IVirtualPathFactory.CreateInstance(System.String)">
+      <summary>Creates a handler factory for the specified virtual path.</summary>
+      <returns>A handler factory for the specified virtual path.</returns>
+      <param name="virtualPath">The virtual path.</param>
+    </member>
+    <member name="M:System.Web.WebPages.IVirtualPathFactory.Exists(System.String)">
+      <summary>Determines whether the specified virtual path is associated with a handler factory.</summary>
+      <returns>true if a handler factory exists for the specified virtual path; otherwise, false.</returns>
+      <param name="virtualPath">The virtual path.</param>
+    </member>
+    <member name="T:System.Web.WebPages.IWebPageRequestExecutor">
+      <summary>Defines methods to implement an executor class that can execute the code on a web page.</summary>
+    </member>
+    <member name="M:System.Web.WebPages.IWebPageRequestExecutor.Execute(System.Web.WebPages.WebPage)">
+      <summary>Executes the code on the specified web page.</summary>
+      <returns>true if the executor took over execution of the web page; otherwise, false.</returns>
+      <param name="page">The web page.</param>
+    </member>
+    <member name="T:System.Web.WebPages.PageVirtualPathAttribute">
+      <summary>Represents a path attribute for a web page class.</summary>
+    </member>
+    <member name="M:System.Web.WebPages.PageVirtualPathAttribute.#ctor(System.String)">
+      <summary>Initializes a new instance of the <see cref="T:System.Web.WebPages.PageVirtualPathAttribute" /> class by using the specified virtual path.</summary>
+      <param name="virtualPath">The virtual path.</param>
+    </member>
+    <member name="P:System.Web.WebPages.PageVirtualPathAttribute.VirtualPath">
+      <summary>Gets the virtual path of the current web page.</summary>
+      <returns>The virtual path.</returns>
+    </member>
+    <member name="T:System.Web.WebPages.PreApplicationStartCode">
+      <summary>Provides a registration point for pre-application start code for web pages.</summary>
+    </member>
+    <member name="M:System.Web.WebPages.PreApplicationStartCode.Start">
+      <summary>Registers pre-application start code for web pages.</summary>
+    </member>
+    <member name="T:System.Web.WebPages.RequestExtensions">
+      <summary>Defines extension methods for the <see cref="T:System.Web.HttpRequestBase" /> class.</summary>
+    </member>
+    <member name="M:System.Web.WebPages.RequestExtensions.IsUrlLocalToHost(System.Web.HttpRequestBase,System.String)">
+      <summary>Determines whether the specified URL references the local computer.</summary>
+      <returns>true if the specified URL references the local computer; otherwise, false.</returns>
+      <param name="request">The HTTP request object.</param>
+      <param name="url">The URL to test.</param>
+    </member>
+    <member name="T:System.Web.WebPages.RequestFieldValidatorBase">
+      <summary>Serves as the abstract base class for the validation helper classes.</summary>
+    </member>
+    <member name="M:System.Web.WebPages.RequestFieldValidatorBase.#ctor(System.String)">
+      <summary>Initializes a new instance of the derived class and specifies the name of the HTML element that is being validated.</summary>
+      <param name="errorMessage">The name (value of the name attribute) of the user input element to validate.</param>
+    </member>
+    <member name="M:System.Web.WebPages.RequestFieldValidatorBase.#ctor(System.String,System.Boolean)">
+      <summary>Initializes a new instance of the derived class, registers the specified string as the error message to display if no value is supplied, and specifies whether the method can use unvalidated data.</summary>
+      <param name="errorMessage">The error message.</param>
+      <param name="useUnvalidatedValues">true to use unvalidated user input; false to reject unvalidated data. This parameter is set to true by calling methods in circumstances when the actual value of the user input is not important, such as for required fields.</param>
+    </member>
+    <member name="P:System.Web.WebPages.RequestFieldValidatorBase.ClientValidationRule">
+      <summary>When implemented in a derived class, gets a container for client validation for the required field.</summary>
+      <returns>The container.</returns>
+    </member>
+    <member name="M:System.Web.WebPages.RequestFieldValidatorBase.GetHttpContext(System.ComponentModel.DataAnnotations.ValidationContext)">
+      <summary>Returns the HTTP context of the current request.</summary>
+      <returns>The context.</returns>
+      <param name="validationContext">The validation context.</param>
+    </member>
+    <member name="M:System.Web.WebPages.RequestFieldValidatorBase.GetRequestValue(System.Web.HttpRequestBase,System.String)">
+      <summary>Returns the value to validate.</summary>
+      <returns>The value to validate.</returns>
+      <param name="request">The current request.</param>
+      <param name="field">The name of the field from the current request to validate.</param>
+    </member>
+    <member name="M:System.Web.WebPages.RequestFieldValidatorBase.IsValid(System.Web.HttpContextBase,System.String)">
+      <summary>Returns a value that indicates whether the specified value is valid.</summary>
+      <returns>true if the value is valid; otherwise, false.</returns>
+      <param name="httpContext">The current context.</param>
+      <param name="value">The value to validate.</param>
+    </member>
+    <member name="M:System.Web.WebPages.RequestFieldValidatorBase.Validate(System.ComponentModel.DataAnnotations.ValidationContext)">
+      <summary>Performs the validation test.</summary>
+      <returns>The result of the validation test.</returns>
+      <param name="validationContext">The context.</param>
+    </member>
+    <member name="T:System.Web.WebPages.ResponseExtensions">
+      <summary>Defines extension methods for the <see cref="T:System.Web.HttpResponseBase" /> base class.</summary>
+    </member>
+    <member name="M:System.Web.WebPages.ResponseExtensions.OutputCache(System.Web.HttpResponseBase,System.Int32,System.Boolean,System.Collections.Generic.IEnumerable{System.String},System.Collections.Generic.IEnumerable{System.String},System.Collections.Generic.IEnumerable{System.String},System.Web.HttpCacheability)">
+      <summary>Configures the cache policy of an HTTP response instance.</summary>
+      <param name="response">The HTTP response instance.</param>
+      <param name="numberOfSeconds">The length of time, in seconds, before items expire from the cache.</param>
+      <param name="sliding">true to indicate that items expire from the cache on a sliding basis; false to indicate that items expire when they reach the predefined expiration time.</param>
+      <param name="varyByParams">The list of all parameters that can be received by a GET or POST operation that affect caching.</param>
+      <param name="varyByHeaders">The list of all HTTP headers that affect caching.</param>
+      <param name="varyByContentEncodings">The list of all Content-Encoding headers that affect caching.</param>
+      <param name="cacheability">One of the enumeration values that specifies how items are cached.</param>
+    </member>
+    <member name="M:System.Web.WebPages.ResponseExtensions.SetStatus(System.Web.HttpResponseBase,System.Int32)">
+      <summary>Sets the HTTP status code of an HTTP response using the specified integer value.</summary>
+      <param name="response">The HTTP response instance.</param>
+      <param name="httpStatusCode">The HTTP status code.</param>
+    </member>
+    <member name="M:System.Web.WebPages.ResponseExtensions.SetStatus(System.Web.HttpResponseBase,System.Net.HttpStatusCode)">
+      <summary>Sets the HTTP status code of an HTTP response using the specified HTTP status code enumeration value.</summary>
+      <param name="response">The HTTP response instance.</param>
+      <param name="httpStatusCode">The HTTP status code</param>
+    </member>
+    <member name="M:System.Web.WebPages.ResponseExtensions.WriteBinary(System.Web.HttpResponseBase,System.Byte[])">
+      <summary>Writes a sequence of bytes that represent binary content of an unspecified type to the output stream of an HTTP response.</summary>
+      <param name="response">The HTTP response instance.</param>
+      <param name="data">An array that contains the bytes to write.</param>
+    </member>
+    <member name="M:System.Web.WebPages.ResponseExtensions.WriteBinary(System.Web.HttpResponseBase,System.Byte[],System.String)">
+      <summary>Writes a sequence of bytes that represent binary content of the specified MIME type to the output stream of an HTTP response.</summary>
+      <param name="response">The receiving HTTP response instance.</param>
+      <param name="data">An array that contains the bytes to write.</param>
+      <param name="mimeType">The MIME type of the binary content.</param>
+    </member>
+    <member name="T:System.Web.WebPages.SectionWriter">
+      <summary>Provides a delegate that represents one or more methods that are called when a content section is written.</summary>
+    </member>
+    <member name="T:System.Web.WebPages.StartPage">
+      <summary>Provides methods and properties that are used to render start pages that use the Razor view engine.</summary>
+    </member>
+    <member name="M:System.Web.WebPages.StartPage.#ctor">
+      <summary>Initializes a new instance of the <see cref="T:System.Web.WebPages.StartPage" /> class.</summary>
+    </member>
+    <member name="P:System.Web.WebPages.StartPage.ChildPage">
+      <summary>Gets or sets the child page of the current start page.</summary>
+      <returns>The child page of the current start page.</returns>
+    </member>
+    <member name="P:System.Web.WebPages.StartPage.Context">
+      <summary>Gets or sets the context of the <see cref="P:System.Web.WebPages.StartPage.ChildPage" /> page.</summary>
+      <returns>The context of the <see cref="P:System.Web.WebPages.StartPage.ChildPage" /> page.</returns>
+    </member>
+    <member name="M:System.Web.WebPages.StartPage.ExecutePageHierarchy">
+      <summary>Calls the methods that are used to execute the developer-written code in the _PageStart start page and in the <see cref="P:System.Web.WebPages.StartPage.ChildPage" /> page.</summary>
+    </member>
+    <member name="M:System.Web.WebPages.StartPage.GetOutputWriter">
+      <summary>Returns the text writer instance that is used to render the page.</summary>
+      <returns>The text writer.</returns>
+    </member>
+    <member name="M:System.Web.WebPages.StartPage.GetStartPage(System.Web.WebPages.WebPageRenderingBase,System.String,System.Collections.Generic.IEnumerable{System.String})">
+      <summary>Returns the initialization page for the specified page.</summary>
+      <returns>The _AppStart page if the _AppStart page exists. If the _AppStart page cannot be found, returns the _PageStart page if a _PageStart page exists. If the _AppStart and _PageStart pages cannot be found, returns <paramref name="page" />.</returns>
+      <param name="page">The page.</param>
+      <param name="fileName">The file name of the page.</param>
+      <param name="supportedExtensions">The collection of file-name extensions that can contain ASP.NET Razor syntax, such as "cshtml" and "vbhtml".</param>
+      <exception cref="T:System.ArgumentNullException">Either <paramref name="page" /> or <paramref name="fileName" /> are null.</exception>
+      <exception cref="T:System.ArgumentException">
+        <paramref name="supportedExtensions" /> is null or empty.</exception>
+    </member>
+    <member name="P:System.Web.WebPages.StartPage.Layout">
+      <summary>Gets or sets the path of the layout page for the <see cref="P:System.Web.WebPages.StartPage.ChildPage" /> page.</summary>
+      <returns>The path of the layout page for the <see cref="P:System.Web.WebPages.StartPage.ChildPage" /> page.</returns>
+    </member>
+    <member name="P:System.Web.WebPages.StartPage.Page">
+      <summary>Gets property-like access to <see cref="P:System.Web.WebPages.StartPage.ChildPage" /> page data that is shared between pages, layout pages, and partial pages.</summary>
+      <returns>An object that contains <see cref="P:System.Web.WebPages.StartPage.ChildPage" /> page data.</returns>
+    </member>
+    <member name="P:System.Web.WebPages.StartPage.PageData">
+      <summary>Gets array-like access to <see cref="P:System.Web.WebPages.StartPage.ChildPage" /> page data that is shared between pages, layout pages, and partial pages.</summary>
+      <returns>An object that provides array-like access to <see cref="P:System.Web.WebPages.StartPage.ChildPage" /> page data.</returns>
+    </member>
+    <member name="M:System.Web.WebPages.StartPage.RenderPage(System.String,System.Object[])">
+      <summary>Renders the <see cref="P:System.Web.WebPages.StartPage.ChildPage" /> page.</summary>
+      <returns>The HTML markup that represents the web page.</returns>
+      <param name="path">The path of the page to render.</param>
+      <param name="data">Additional data that is used to render the page.</param>
+    </member>
+    <member name="M:System.Web.WebPages.StartPage.RunPage">
+      <summary>Executes the developer-written code in the <see cref="P:System.Web.WebPages.StartPage.ChildPage" /> page.</summary>
+    </member>
+    <member name="M:System.Web.WebPages.StartPage.Write(System.Object)">
+      <summary>Writes the string representation of the specified object as an HTML-encoded string.</summary>
+      <param name="value">The object to encode and write.</param>
+    </member>
+    <member name="M:System.Web.WebPages.StartPage.Write(System.Web.WebPages.HelperResult)">
+      <summary>Writes the string representation of the specified <see cref="T:System.Web.WebPages.HelperResult" /> object as an HTML-encoded string.</summary>
+      <param name="result">The helper result to encode and write.</param>
+    </member>
+    <member name="M:System.Web.WebPages.StartPage.WriteLiteral(System.Object)">
+      <summary>Writes the string representation of the specified object without HTML encoding.</summary>
+      <param name="value">The object to write.</param>
+    </member>
+    <member name="T:System.Web.WebPages.StringExtensions">
+      <summary>Provides utility methods for converting string values to other data types.</summary>
+    </member>
+    <member name="M:System.Web.WebPages.StringExtensions.As``1(System.String)">
+      <summary>Converts a string to a strongly typed value of the specified data type.</summary>
+      <returns>The converted value.</returns>
+      <param name="value">The value to convert.</param>
+      <typeparam name="TValue">The data type to convert to.</typeparam>
+    </member>
+    <member name="M:System.Web.WebPages.StringExtensions.As``1(System.String,``0)">
+      <summary>Converts a string to the specified data type and specifies a default value.</summary>
+      <returns>The converted value.</returns>
+      <param name="value">The value to convert.</param>
+      <param name="defaultValue">The value to return if <paramref name="value" /> is null.</param>
+      <typeparam name="TValue">The data type to convert to.</typeparam>
+    </member>
+    <member name="M:System.Web.WebPages.StringExtensions.AsBool(System.String)">
+      <summary>Converts a string to a Boolean (true/false) value.</summary>
+      <returns>The converted value.</returns>
+      <param name="value">The value to convert.</param>
+    </member>
+    <member name="M:System.Web.WebPages.StringExtensions.AsBool(System.String,System.Boolean)">
+      <summary>Converts a string to a Boolean (true/false) value and specifies a default value.</summary>
+      <returns>The converted value.</returns>
+      <param name="value">The value to convert.</param>
+      <param name="defaultValue">The value to return if <paramref name="value" /> is null or is an invalid value.</param>
+    </member>
+    <member name="M:System.Web.WebPages.StringExtensions.AsDateTime(System.String)">
+      <summary>Converts a string to a <see cref="T:System.DateTime" /> value.</summary>
+      <returns>The converted value.</returns>
+      <param name="value">The value to convert.</param>
+    </member>
+    <member name="M:System.Web.WebPages.StringExtensions.AsDateTime(System.String,System.DateTime)">
+      <summary>Converts a string to a <see cref="T:System.DateTime" /> value and specifies a default value.</summary>
+      <returns>The converted value.</returns>
+      <param name="value">The value to convert.</param>
+      <param name="defaultValue">The value to return if <paramref name="value" /> is null or is an invalid value. The default is the minimum time value on the system.</param>
+    </member>
+    <member name="M:System.Web.WebPages.StringExtensions.AsDecimal(System.String)">
+      <summary>Converts a string to a <see cref="T:System.Decimal" /> number.</summary>
+      <returns>The converted value.</returns>
+      <param name="value">The value to convert.</param>
+    </member>
+    <member name="M:System.Web.WebPages.StringExtensions.AsDecimal(System.String,System.Decimal)">
+      <summary>Converts a string to a <see cref="T:System.Decimal" /> number and specifies a default value.</summary>
+      <returns>The converted value.</returns>
+      <param name="value">The value to convert.</param>
+      <param name="defaultValue">The value to return if <paramref name="value" /> is null or invalid.</param>
+    </member>
+    <member name="M:System.Web.WebPages.StringExtensions.AsFloat(System.String)">
+      <summary>Converts a string to a <see cref="T:System.Single" /> number.</summary>
+      <returns>The converted value.</returns>
+      <param name="value">The value to convert.</param>
+    </member>
+    <member name="M:System.Web.WebPages.StringExtensions.AsFloat(System.String,System.Single)">
+      <summary>Converts a string to a <see cref="T:System.Single" /> number and specifies a default value.</summary>
+      <returns>The converted value.</returns>
+      <param name="value">The value to convert.</param>
+      <param name="defaultValue">The value to return if <paramref name="value" /> is null.</param>
+    </member>
+    <member name="M:System.Web.WebPages.StringExtensions.AsInt(System.String)">
+      <summary>Converts a string to an integer.</summary>
+      <returns>The converted value.</returns>
+      <param name="value">The value to convert.</param>
+    </member>
+    <member name="M:System.Web.WebPages.StringExtensions.AsInt(System.String,System.Int32)">
+      <summary>Converts a string to an integer and specifies a default value.</summary>
+      <returns>The converted value.</returns>
+      <param name="value">The value to convert.</param>
+      <param name="defaultValue">The value to return if <paramref name="value" /> is null or is an invalid value.</param>
+    </member>
+    <member name="M:System.Web.WebPages.StringExtensions.Is``1(System.String)">
+      <summary>Checks whether a string can be converted to the specified data type.</summary>
+      <returns>true if <paramref name="value" /> can be converted to the specified type; otherwise, false.</returns>
+      <param name="value">The value to test.</param>
+      <typeparam name="TValue">The data type to convert to.</typeparam>
+    </member>
+    <member name="M:System.Web.WebPages.StringExtensions.IsBool(System.String)">
+      <summary>Checks whether a string can be converted to the Boolean (true/false) type.</summary>
+      <returns>true if <paramref name="value" /> can be converted to the specified type; otherwise, false.</returns>
+      <param name="value">The string value to test.</param>
+    </member>
+    <member name="M:System.Web.WebPages.StringExtensions.IsDateTime(System.String)">
+      <summary>Checks whether a string can be converted to the <see cref="T:System.DateTime" /> type.</summary>
+      <returns>true if <paramref name="value" /> can be converted to the specified type; otherwise, false.</returns>
+      <param name="value">The string value to test.</param>
+    </member>
+    <member name="M:System.Web.WebPages.StringExtensions.IsDecimal(System.String)">
+      <summary>Checks whether a string can be converted to the <see cref="T:System.Decimal" /> type.</summary>
+      <returns>true if <paramref name="value" /> can be converted to the specified type; otherwise, false.</returns>
+      <param name="value">The string value to test.</param>
+    </member>
+    <member name="M:System.Web.WebPages.StringExtensions.IsEmpty(System.String)">
+      <summary>Checks whether a string value is null or empty.</summary>
+      <returns>true if <paramref name="value" /> is null or is a zero-length string (""); otherwise, false.</returns>
+      <param name="value">The string value to test.</param>
+    </member>
+    <member name="M:System.Web.WebPages.StringExtensions.IsFloat(System.String)">
+      <summary>Checks whether a string can be converted to the <see cref="T:System.Single" /> type.</summary>
+      <returns>true if <paramref name="value" /> can be converted to the specified type; otherwise, false.</returns>
+      <param name="value">The string value to test.</param>
+    </member>
+    <member name="M:System.Web.WebPages.StringExtensions.IsInt(System.String)">
+      <summary>Checks whether a string can be converted to an integer.</summary>
+      <returns>true if <paramref name="value" /> can be converted to the specified type; otherwise, false.</returns>
+      <param name="value">The string value to test.</param>
+    </member>
+    <member name="T:System.Web.WebPages.TemplateFileInfo">
+      <summary>Contains methods and properties that describe a file information template.</summary>
+    </member>
+    <member name="M:System.Web.WebPages.TemplateFileInfo.#ctor(System.String)">
+      <summary>Initializes a new instance of the <see cref="T:System.Web.WebPages.TemplateFileInfo" /> class by using the specified virtual path.</summary>
+      <param name="virtualPath">The virtual path.</param>
+    </member>
+    <member name="P:System.Web.WebPages.TemplateFileInfo.VirtualPath">
+      <summary>Gets the virtual path of the web page.</summary>
+      <returns>The virtual path.</returns>
+    </member>
+    <member name="T:System.Web.WebPages.TemplateStack">
+      <summary>Represents a last-in-first-out (LIFO) collection of <see cref="T:System.Web.WebPages.ITemplateFile" /> template files.</summary>
+    </member>
+    <member name="M:System.Web.WebPages.TemplateStack.GetCurrentTemplate(System.Web.HttpContextBase)">
+      <summary>Returns the current template file from the specified HTTP context.</summary>
+      <returns>The template file, removed from the top of the stack.</returns>
+      <param name="httpContext"> The HTTP context that contains the stack that stores the template files.</param>
+    </member>
+    <member name="M:System.Web.WebPages.TemplateStack.Pop(System.Web.HttpContextBase)">
+      <summary>Removes and returns the template file that is at the top of the stack in the specified HTTP context.</summary>
+      <returns>The template file, removed from the top of the stack.</returns>
+      <param name="httpContext">The HTTP context that contains the stack that stores the template files.</param>
+      <exception cref="T:System.ArgumentNullException">
+        <paramref name="httpContext" /> is null.</exception>
+    </member>
+    <member name="M:System.Web.WebPages.TemplateStack.Push(System.Web.HttpContextBase,System.Web.WebPages.ITemplateFile)">
+      <summary>Inserts a template file at the top of the stack in the specified HTTP context.</summary>
+      <param name="httpContext">The HTTP context that contains the stack that stores the template files.</param>
+      <param name="templateFile">The template file to push onto the specified stack.</param>
+      <exception cref="T:System.ArgumentNullException">
+        <paramref name="httpContext" /> or <paramref name="templateFile" /> are null.</exception>
+    </member>
+    <member name="T:System.Web.WebPages.ValidationHelper">
+      <summary>Implements validation for user input.</summary>
+    </member>
+    <member name="M:System.Web.WebPages.ValidationHelper.Add(System.Collections.Generic.IEnumerable{System.String},System.Web.WebPages.IValidator[])">
+      <summary>Registers a list of user input elements for validation.</summary>
+      <param name="fields">The names (value of the name attribute) of the user input elements to validate.</param>
+      <param name="validators">The type of validation to register for each user input element specified in <paramref name="fields" />.</param>
+    </member>
+    <member name="M:System.Web.WebPages.ValidationHelper.Add(System.String,System.Web.WebPages.IValidator[])">
+      <summary>Registers a user input element for validation.</summary>
+      <param name="field">The name (value of the name attribute) of the user input element to validate.</param>
+      <param name="validators">A list of one or more types of validation to register.</param>
+    </member>
+    <member name="M:System.Web.WebPages.ValidationHelper.AddFormError(System.String)">
+      <summary>Adds an error message.</summary>
+      <param name="errorMessage">The error message.</param>
+    </member>
+    <member name="M:System.Web.WebPages.ValidationHelper.ClassFor(System.String)">
+      <summary>Renders an attribute that references the CSS style definition to use when validation messages for the user input element are rendered.</summary>
+      <returns>The attribute.</returns>
+      <param name="field">The name (value of the name attribute) of the user input element to validate.</param>
+    </member>
+    <member name="M:System.Web.WebPages.ValidationHelper.For(System.String)">
+      <summary>Renders attributes that enable client-side validation for an individual user input element.</summary>
+      <returns>The attributes to render.</returns>
+      <param name="field">The name (value of the name attribute) of the user input element to validate.</param>
+    </member>
+    <member name="P:System.Web.WebPages.ValidationHelper.FormField">
+      <summary>Gets the name of the current form. This type/member supports the .NET Framework infrastructure and is not intended to be used directly from your code. </summary>
+      <returns>The name.</returns>
+    </member>
+    <member name="M:System.Web.WebPages.ValidationHelper.GetErrors(System.String[])">
+      <summary>Returns a list of current validation errors, and optionally lets you specify a list of fields to check.</summary>
+      <returns>The list of errors.</returns>
+      <param name="fields">Optional. The names (value of the name attribute) of the user input elements to get error information for. You can specify any number of element names, separated by commas. If you do not specify a list of fields, the method returns errors for all fields.</param>
+    </member>
+    <member name="P:System.Web.WebPages.ValidationHelper.InvalidCssClass">
+      <summary>Gets the name of the class that is used to specify the appearance of error-message display when errors have occurred. This type/member supports the .NET Framework infrastructure and is not intended to be used directly from your code. </summary>
+      <returns>The name.</returns>
+    </member>
+    <member name="M:System.Web.WebPages.ValidationHelper.IsValid(System.String[])">
+      <summary>Determines whether the contents of the user input fields pass validation checks, and optionally lets you specify a list of fields to check.</summary>
+      <returns>true if all specified field or fields pass validation checks; false if any field contains a validation error.</returns>
+      <param name="fields">Optional. The names (value of the name attribute) of the user input elements to check for validation errors. You can specify any number of element names, separated by commas. If you do not specify a list of fields, the method checks all elements that are registered for validation.</param>
+    </member>
+    <member name="M:System.Web.WebPages.ValidationHelper.RequireField(System.String)">
+      <summary>Registers the specified field as one that requires user entry.</summary>
+      <param name="field">The name (value of the name attribute) of the user input element to validate.</param>
+    </member>
+    <member name="M:System.Web.WebPages.ValidationHelper.RequireField(System.String,System.String)">
+      <summary>Registers the specified field as one that requires user entry and registers the specified string as the error message to display if no value is supplied.</summary>
+      <param name="field">The name (value of the name attribute) of the user input element to validate.</param>
+      <param name="errorMessage">The error message.</param>
+    </member>
+    <member name="M:System.Web.WebPages.ValidationHelper.RequireFields(System.String[])">
+      <summary>Registers the specified fields as ones that require user entry.</summary>
+      <param name="fields">The names (value of the name attribute) of the user input elements to validate. You can specify any number of element names, separated by commas.</param>
+    </member>
+    <member name="M:System.Web.WebPages.ValidationHelper.Validate(System.String[])">
+      <summary>Performs validation on elements registered for validation, and optionally lets you specify a list of fields to check.</summary>
+      <returns>The list of errors for the specified fields, if any validation errors occurred.</returns>
+      <param name="fields">Optional. The names (value of the name attribute) of the user input elements to validate. You can specify any number of element names, separated by commas. If you do not specify a list, the method validates all registered elements.</param>
+    </member>
+    <member name="P:System.Web.WebPages.ValidationHelper.ValidCssClass">
+      <summary>Gets the name of the class that is used to specify the appearance of error-message display when errors have occurred. This type/member supports the .NET Framework infrastructure and is not intended to be used directly from your code. </summary>
+      <returns>The name.</returns>
+    </member>
+    <member name="T:System.Web.WebPages.Validator">
+      <summary>Defines validation tests that can be registered using the <see cref="M:System.Web.WebPages.ValidationHelper.Add(System.String,System.Web.WebPages.IValidator[])" /> method.</summary>
+    </member>
+    <member name="M:System.Web.WebPages.Validator.#ctor">
+      <summary>Initializes a new instance of the <see cref="T:System.Web.WebPages.Validator" /> class.</summary>
+    </member>
+    <member name="M:System.Web.WebPages.Validator.DateTime(System.String)">
+      <summary>Defines a validation test that tests whether a value can be treated as a date/time value.</summary>
+      <returns>The validation test.</returns>
+      <param name="errorMessage">The error message to display if validation fails.</param>
+    </member>
+    <member name="M:System.Web.WebPages.Validator.Decimal(System.String)">
+      <summary>Defines a validation test that tests whether a value can be treated as a decimal number.</summary>
+      <returns>The validation test.</returns>
+      <param name="errorMessage">The error message to display if validation fails.</param>
+    </member>
+    <member name="M:System.Web.WebPages.Validator.EqualsTo(System.String,System.String)">
+      <summary>Defines a validation test that test user input against the value of another field.</summary>
+      <returns>The validation test.</returns>
+      <param name="otherFieldName">The other field to compare.</param>
+      <param name="errorMessage">The error message to display if validation fails.</param>
+    </member>
+    <member name="M:System.Web.WebPages.Validator.Float(System.String)">
+      <summary>Defines a validation test that tests whether a value can be treated as a floating-point number.</summary>
+      <returns>The validation test.</returns>
+      <param name="errorMessage">The error message to display if validation fails.</param>
+    </member>
+    <member name="M:System.Web.WebPages.Validator.Integer(System.String)">
+      <summary>Defines a validation test that tests whether a value can be treated as an integer.</summary>
+      <returns>The validation test.</returns>
+      <param name="errorMessage">The error message to display if validation fails.</param>
+    </member>
+    <member name="M:System.Web.WebPages.Validator.Range(System.Double,System.Double,System.String)">
+      <summary>Defines a validation test that tests whether a decimal number falls within a specific range.</summary>
+      <returns>The validation test.</returns>
+      <param name="minValue">The minimum value. The default is 0.</param>
+      <param name="maxValue">The maximum value.</param>
+      <param name="errorMessage">The error message to display if validation fails.</param>
+    </member>
+    <member name="M:System.Web.WebPages.Validator.Range(System.Int32,System.Int32,System.String)">
+      <summary>Defines a validation test that tests whether an integer value falls within a specific range.</summary>
+      <returns>The validation test.</returns>
+      <param name="minValue">The minimum value. The default is 0.</param>
+      <param name="maxValue">The maximum value. </param>
+      <param name="errorMessage">The error message to display if validation fails.</param>
+    </member>
+    <member name="M:System.Web.WebPages.Validator.Regex(System.String,System.String)">
+      <summary>Defines a validation test that tests a value against a pattern specified as a regular expression.</summary>
+      <returns>The validation test.</returns>
+      <param name="pattern">The regular expression to use to test the user input.</param>
+      <param name="errorMessage">The error message to display if validation fails.</param>
+    </member>
+    <member name="M:System.Web.WebPages.Validator.Required(System.String)">
+      <summary>Defines a validation test that tests whether a value has been provided.</summary>
+      <returns>The validation test.</returns>
+      <param name="errorMessage">The error message to display if validation fails.</param>
+    </member>
+    <member name="M:System.Web.WebPages.Validator.StringLength(System.Int32,System.Int32,System.String)">
+      <summary>Defines a validation test that tests the length of a string.</summary>
+      <returns>The validation test.</returns>
+      <param name="maxLength">The maximum length of the string.</param>
+      <param name="minLength">The minimum length of the string. The default is 0.</param>
+      <param name="errorMessage">The error message to display if validation fails.</param>
+    </member>
+    <member name="M:System.Web.WebPages.Validator.Url(System.String)">
+      <summary>Defines a validation test that tests whether a value is a well-formed URL.</summary>
+      <returns>The validation test.</returns>
+      <param name="errorMessage">The error message to display if validation fails.</param>
+    </member>
+    <member name="T:System.Web.WebPages.VirtualPathFactoryManager">
+      <summary>This type/member supports the .NET Framework infrastructure and is not intended to be used directly from your code.</summary>
+    </member>
+    <member name="M:System.Web.WebPages.VirtualPathFactoryManager.CreateInstance(System.String)">
+      <summary>This type/member supports the .NET Framework infrastructure and is not intended to be used directly from your code.</summary>
+    </member>
+    <member name="M:System.Web.WebPages.VirtualPathFactoryManager.Exists(System.String)">
+      <summary>This type/member supports the .NET Framework infrastructure and is not intended to be used directly from your code.</summary>
+    </member>
+    <member name="M:System.Web.WebPages.VirtualPathFactoryManager.RegisterVirtualPathFactory(System.Web.WebPages.IVirtualPathFactory)">
+      <summary>This type/member supports the .NET Framework infrastructure and is not intended to be used directly from your code.</summary>
+    </member>
+    <member name="T:System.Web.WebPages.WebPage">
+      <summary>Represents an ASP.NET Razor page.</summary>
+    </member>
+    <member name="M:System.Web.WebPages.WebPage.#ctor">
+      <summary>Called from a derived class to create a new instance that is based on the <see cref="T:System.Web.WebPages.WebPage" /> class.</summary>
+    </member>
+    <member name="P:System.Web.WebPages.WebPage.Context">
+      <summary>Gets or sets the <see cref="T:System.Web.HttpContextBase" /> object that is associated with a page.</summary>
+      <returns>The current context data.</returns>
+    </member>
+    <member name="M:System.Web.WebPages.WebPage.ExecutePageHierarchy">
+      <summary>Executes the code in a set of dependent pages.</summary>
+    </member>
+    <member name="P:System.Web.WebPages.WebPage.Html">
+      <summary>Gets the <see cref="T:System.Web.WebPages.Html.HtmlHelper" /> object that is associated with a page.</summary>
+      <returns>An object that can render HTML form controls in a page.</returns>
+    </member>
+    <member name="M:System.Web.WebPages.WebPage.InitializePage">
+      <summary>Initializes an object that inherits from the <see cref="T:System.Web.WebPages.WebPage" /> class.</summary>
+    </member>
+    <member name="P:System.Web.WebPages.WebPage.Model">
+      <summary>Gets the model that is associated with a page.</summary>
+      <returns>An object that represents a model that is associated with the view data for a page.</returns>
+    </member>
+    <member name="P:System.Web.WebPages.WebPage.ModelState">
+      <summary>Gets the state of the model that is associated with a page.</summary>
+      <returns>The state of the model.</returns>
+    </member>
+    <member name="M:System.Web.WebPages.WebPage.RegisterPageExecutor(System.Web.WebPages.IWebPageRequestExecutor)">
+      <summary>Adds a class to a list of classes that handle page execution and that implement custom features for pages.</summary>
+      <param name="executor">The class to add.</param>
+    </member>
+    <member name="M:System.Web.WebPages.WebPage.RenderPage(System.String,System.Object[])">
+      <summary>Renders a content page.</summary>
+      <returns>An object that can write the output of the page.</returns>
+      <param name="path">The path of the page to render.</param>
+      <param name="data">Data to pass to the page.</param>
+    </member>
+    <member name="P:System.Web.WebPages.WebPage.Validation">
+      <summary>Gets the validation helper for the current page context.</summary>
+      <returns>The validation helper.</returns>
+    </member>
+    <member name="T:System.Web.WebPages.WebPageBase">
+      <summary>Serves as the base class for classes that represent an ASP.NET Razor page.</summary>
+    </member>
+    <member name="M:System.Web.WebPages.WebPageBase.#ctor">
+      <summary>Initializes the <see cref="T:System.Web.WebPages.WebPageBase" /> class for use by an inherited class instance. This constructor can only be called by an inherited class.</summary>
+    </member>
+    <member name="M:System.Web.WebPages.WebPageBase.ConfigurePage(System.Web.WebPages.WebPageBase)">
+      <summary>When overridden in a derived class, configures the current web page based on the configuration of the parent web page.</summary>
+      <param name="parentPage">The parent page from which to read configuration information.</param>
+    </member>
+    <member name="M:System.Web.WebPages.WebPageBase.CreateInstanceFromVirtualPath(System.String)">
+      <summary>Creates a new instance of the <see cref="T:System.Web.WebPages.WebPageBase" /> class by using the specified virtual path.</summary>
+      <returns>The new <see cref="T:System.Web.WebPages.WebPageBase" /> object.</returns>
+      <param name="virtualPath">The virtual path to use to create the instance.</param>
+    </member>
+    <member name="M:System.Web.WebPages.WebPageBase.CreatePageFromVirtualPath(System.String,System.Web.HttpContextBase,System.Func{System.String,System.Boolean},System.Web.WebPages.DisplayModeProvider,System.Web.WebPages.IDisplayMode)">
+      <summary>Attempts to create a WebPageBase instance from a virtualPath and wraps complex compiler exceptions with simpler messages</summary>
+    </member>
+    <member name="M:System.Web.WebPages.WebPageBase.DefineSection(System.String,System.Web.WebPages.SectionWriter)">
+      <summary>Called by content pages to create named content sections.</summary>
+      <param name="name">The name of the section to create.</param>
+      <param name="action">The type of action to take with the new section.</param>
+    </member>
+    <member name="M:System.Web.WebPages.WebPageBase.ExecutePageHierarchy">
+      <summary>Executes the code in a set of dependent web pages.</summary>
+    </member>
+    <member name="M:System.Web.WebPages.WebPageBase.ExecutePageHierarchy(System.Web.WebPages.WebPageContext,System.IO.TextWriter)">
+      <summary>Executes the code in a set of dependent web pages by using the specified parameters.</summary>
+      <param name="pageContext">The context data for the page.</param>
+      <param name="writer">The writer to use to write the executed HTML.</param>
+    </member>
+    <member name="M:System.Web.WebPages.WebPageBase.ExecutePageHierarchy(System.Web.WebPages.WebPageContext,System.IO.TextWriter,System.Web.WebPages.WebPageRenderingBase)">
+      <summary>Executes the code in a set of dependent web pages by using the specified context, writer, and start page.</summary>
+      <param name="pageContext">The context data for the page.</param>
+      <param name="writer">The writer to use to write the executed HTML.</param>
+      <param name="startPage">The page to start execution in the page hierarchy.</param>
+    </member>
+    <member name="M:System.Web.WebPages.WebPageBase.GetOutputWriter">
+      <summary>Returns the text writer instance that is used to render the page.</summary>
+      <returns>The text writer.</returns>
+    </member>
+    <member name="M:System.Web.WebPages.WebPageBase.InitializePage">
+      <summary>Initializes the current page.</summary>
+    </member>
+    <member name="M:System.Web.WebPages.WebPageBase.IsSectionDefined(System.String)">
+      <summary>Returns a value that indicates whether the specified section is defined in the page.</summary>
+      <returns>true if the specified section is defined in the page; otherwise, false.</returns>
+      <param name="name">The name of the section to search for.</param>
+    </member>
+    <member name="P:System.Web.WebPages.WebPageBase.Layout">
+      <summary>Gets or sets the path of a layout page.</summary>
+      <returns>The path of the layout page.</returns>
+    </member>
+    <member name="P:System.Web.WebPages.WebPageBase.Output">
+      <summary>Gets the current <see cref="T:System.IO.TextWriter" /> object for the page.</summary>
+      <returns>The <see cref="T:System.IO.TextWriter" /> object.</returns>
+    </member>
+    <member name="P:System.Web.WebPages.WebPageBase.OutputStack">
+      <summary>Gets the stack of <see cref="T:System.IO.TextWriter" /> objects for the current page context.</summary>
+      <returns>The <see cref="T:System.IO.TextWriter" /> objects.</returns>
+    </member>
+    <member name="P:System.Web.WebPages.WebPageBase.Page">
+      <summary>Provides property-like access to page data that is shared between pages, layout pages, and partial pages.</summary>
+      <returns>An object that contains page data.</returns>
+    </member>
+    <member name="P:System.Web.WebPages.WebPageBase.PageData">
+      <summary>Provides array-like access to page data that is shared between pages, layout pages, and partial pages.</summary>
+      <returns>A dictionary that contains page data.</returns>
+    </member>
+    <member name="M:System.Web.WebPages.WebPageBase.PopContext">
+      <summary>Returns and removes the context from the top of the <see cref="P:System.Web.WebPages.WebPageBase.OutputStack" /> instance.</summary>
+    </member>
+    <member name="M:System.Web.WebPages.WebPageBase.PushContext(System.Web.WebPages.WebPageContext,System.IO.TextWriter)">
+      <summary>Inserts the specified context at the top of the <see cref="P:System.Web.WebPages.WebPageBase.OutputStack" /> instance.</summary>
+      <param name="pageContext">The page context to push onto the <see cref="P:System.Web.WebPages.WebPageBase.OutputStack" /> instance.</param>
+      <param name="writer">The writer for the page context.</param>
+    </member>
+    <member name="M:System.Web.WebPages.WebPageBase.RenderBody">
+      <summary>In layout pages, renders the portion of a content page that is not within a named section.</summary>
+      <returns>The HTML content to render.</returns>
+    </member>
+    <member name="M:System.Web.WebPages.WebPageBase.RenderPage(System.String,System.Object[])">
+      <summary>Renders the content of one page within another page.</summary>
+      <returns>The HTML content to render.</returns>
+      <param name="path">The path of the page to render.</param>
+      <param name="data">(Optional) An array of data to pass to the page being rendered. In the rendered page, these parameters can be accessed by using the <see cref="P:System.Web.WebPages.WebPageBase.PageData" /> property.</param>
+    </member>
+    <member name="M:System.Web.WebPages.WebPageBase.RenderSection(System.String)">
+      <summary>In layout pages, renders the content of a named section.</summary>
+      <returns>The HTML content to render.</returns>
+      <param name="name">The section to render.</param>
+      <exception cref="T:System.Web.HttpException">The <paramref name="name" /> section was already rendered.-or-The <paramref name="name" /> section was marked as required but was not found.</exception>
+    </member>
+    <member name="M:System.Web.WebPages.WebPageBase.RenderSection(System.String,System.Boolean)">
+      <summary>In layout pages, renders the content of a named section and specifies whether the section is required.</summary>
+      <returns>The HTML content to render.</returns>
+      <param name="name">The section to render.</param>
+      <param name="required">true to specify that the section is required; otherwise, false.</param>
+    </member>
+    <member name="M:System.Web.WebPages.WebPageBase.Write(System.Object)">
+      <summary>Writes the specified object as an HTML-encoded string.</summary>
+      <param name="value">The object to encode and write.</param>
+    </member>
+    <member name="M:System.Web.WebPages.WebPageBase.Write(System.Web.WebPages.HelperResult)">
+      <summary>Writes the specified <see cref="T:System.Web.WebPages.HelperResult" /> object as an HTML-encoded string.</summary>
+      <param name="result">The helper result to encode and write.</param>
+    </member>
+    <member name="M:System.Web.WebPages.WebPageBase.WriteLiteral(System.Object)">
+      <summary>Writes the specified object without HTML-encoding it first.</summary>
+      <param name="value">The object to write.</param>
+    </member>
+    <member name="T:System.Web.WebPages.WebPageContext">
+      <summary>Contains data that is used by a <see cref="T:System.Web.WebPages.WebPage" /> object to reference details about the web application, the current HTTP request, the current execution context, and page-rendering data.</summary>
+    </member>
+    <member name="M:System.Web.WebPages.WebPageContext.#ctor">
+      <summary>Initializes a new instance of the <see cref="T:System.Web.WebPages.WebPageContext" /> class.</summary>
+    </member>
+    <member name="M:System.Web.WebPages.WebPageContext.#ctor(System.Web.HttpContextBase,System.Web.WebPages.WebPageRenderingBase,System.Object)">
+      <summary>Initializes a new instance of the class by using the specified context, page, and model.</summary>
+      <param name="context">The HTTP request context data to associate with the page context.</param>
+      <param name="page">The page data to share between pages, layout pages, and partial pages.</param>
+      <param name="model">The model to associate with the view data.</param>
+    </member>
+    <member name="P:System.Web.WebPages.WebPageContext.Current">
+      <summary>Gets a reference to the current <see cref="T:System.Web.WebPages.WebPageContext" /> object that is associated with a page.</summary>
+      <returns>The current page context object.</returns>
+    </member>
+    <member name="P:System.Web.WebPages.WebPageContext.Model">
+      <summary>Gets the model that is associated with a page.</summary>
+      <returns>An object that represents a model that is associated with the view data for a page.</returns>
+    </member>
+    <member name="P:System.Web.WebPages.WebPageContext.Page">
+      <summary>Gets the <see cref="T:System.Web.WebPages.WebPageRenderingBase" /> object that is associated with a page.</summary>
+      <returns>The object that renders the page.</returns>
+    </member>
+    <member name="P:System.Web.WebPages.WebPageContext.PageData">
+      <summary>Gets the page data that is shared between pages, layout pages, and partial pages.</summary>
+      <returns>A dictionary that contains page data.</returns>
+    </member>
+    <member name="T:System.Web.WebPages.WebPageExecutingBase">
+      <summary>Provides objects and methods that are used to execute and render ASP.NET pages that include Razor syntax.</summary>
+    </member>
+    <member name="M:System.Web.WebPages.WebPageExecutingBase.#ctor">
+      <summary>Initializes a new instance of the <see cref="T:System.Web.WebPages.WebPageExecutingBase" /> class. This constructor can only be called by an inherited class.</summary>
+    </member>
+    <member name="P:System.Web.WebPages.WebPageExecutingBase.App"></member>
+    <member name="P:System.Web.WebPages.WebPageExecutingBase.AppState"></member>
+    <member name="M:System.Web.WebPages.WebPageExecutingBase.BeginContext(System.Int32,System.Int32,System.Boolean)"></member>
+    <member name="M:System.Web.WebPages.WebPageExecutingBase.BeginContext(System.IO.TextWriter,System.Int32,System.Int32,System.Boolean)"></member>
+    <member name="M:System.Web.WebPages.WebPageExecutingBase.BeginContext(System.IO.TextWriter,System.String,System.Int32,System.Int32,System.Boolean)"></member>
+    <member name="M:System.Web.WebPages.WebPageExecutingBase.BeginContext(System.String,System.Int32,System.Int32,System.Boolean)"></member>
+    <member name="P:System.Web.WebPages.WebPageExecutingBase.Context">
+      <summary>When overridden in a derived class, gets or sets the <see cref="T:System.Web.HttpContextBase" /> object that is associated with a page.</summary>
+      <returns>The current context data.</returns>
+    </member>
+    <member name="M:System.Web.WebPages.WebPageExecutingBase.EndContext(System.Int32,System.Int32,System.Boolean)"></member>
+    <member name="M:System.Web.WebPages.WebPageExecutingBase.EndContext(System.IO.TextWriter,System.Int32,System.Int32,System.Boolean)"></member>
+    <member name="M:System.Web.WebPages.WebPageExecutingBase.EndContext(System.IO.TextWriter,System.String,System.Int32,System.Int32,System.Boolean)"></member>
+    <member name="M:System.Web.WebPages.WebPageExecutingBase.EndContext(System.String,System.Int32,System.Int32,System.Boolean)"></member>
+    <member name="M:System.Web.WebPages.WebPageExecutingBase.Execute"></member>
+    <member name="M:System.Web.WebPages.WebPageExecutingBase.GetOutputWriter">
+      <summary>Returns the text writer instance that is used to render the page.</summary>
+      <returns>The text writer.</returns>
+    </member>
+    <member name="M:System.Web.WebPages.WebPageExecutingBase.Href(System.String,System.Object[])"></member>
+    <member name="M:System.Web.WebPages.WebPageExecutingBase.NormalizeLayoutPagePath(System.String)"></member>
+    <member name="M:System.Web.WebPages.WebPageExecutingBase.NormalizePath(System.String)"></member>
+    <member name="P:System.Web.WebPages.WebPageExecutingBase.VirtualPath"></member>
+    <member name="P:System.Web.WebPages.WebPageExecutingBase.VirtualPathFactory"></member>
+    <member name="M:System.Web.WebPages.WebPageExecutingBase.Write(System.Object)">
+      <summary>Writes the string representation of the specified object as an HTML-encoded string.</summary>
+      <param name="value">The object to encode and write.</param>
+    </member>
+    <member name="M:System.Web.WebPages.WebPageExecutingBase.Write(System.Web.WebPages.HelperResult)">
+      <summary>Writes the specified <see cref="T:System.Web.WebPages.HelperResult" /> object as an HTML-encoded string.</summary>
+      <param name="result">The helper result to encode and write.</param>
+    </member>
+    <member name="M:System.Web.WebPages.WebPageExecutingBase.WriteAttribute(System.String,System.Web.WebPages.Instrumentation.PositionTagged{System.String},System.Web.WebPages.Instrumentation.PositionTagged{System.String},System.Web.WebPages.AttributeValue[])"></member>
+    <member name="M:System.Web.WebPages.WebPageExecutingBase.WriteAttributeTo(System.IO.TextWriter,System.String,System.Web.WebPages.Instrumentation.PositionTagged{System.String},System.Web.WebPages.Instrumentation.PositionTagged{System.String},System.Web.WebPages.AttributeValue[])"></member>
+    <member name="M:System.Web.WebPages.WebPageExecutingBase.WriteAttributeTo(System.String,System.IO.TextWriter,System.String,System.Web.WebPages.Instrumentation.PositionTagged{System.String},System.Web.WebPages.Instrumentation.PositionTagged{System.String},System.Web.WebPages.AttributeValue[])"></member>
+    <member name="M:System.Web.WebPages.WebPageExecutingBase.WriteLiteral(System.Object)">
+      <summary>Writes the specified object without HTML encoding.</summary>
+      <param name="value">The object to write.</param>
+    </member>
+    <member name="M:System.Web.WebPages.WebPageExecutingBase.WriteLiteralTo(System.IO.TextWriter,System.Object)">
+      <summary>Writes the specified object to the specified <see cref="T:System.IO.TextWriter" /> instance without HTML encoding.</summary>
+      <param name="writer">The text writer.</param>
+      <param name="content">The object to write.</param>
+    </member>
+    <member name="M:System.Web.WebPages.WebPageExecutingBase.WriteTo(System.IO.TextWriter,System.Object)">
+      <summary>Writes the specified object as an HTML-encoded string to the specified text writer.</summary>
+      <param name="writer">The text writer.</param>
+      <param name="content">The object to encode and write.</param>
+    </member>
+    <member name="M:System.Web.WebPages.WebPageExecutingBase.WriteTo(System.IO.TextWriter,System.Web.WebPages.HelperResult)">
+      <summary>Writes the specified <see cref="T:System.Web.WebPages.HelperResult" /> object as an HTML-encoded string to the specified text writer.</summary>
+      <param name="writer">The text writer.</param>
+      <param name="content">The helper result to encode and write.</param>
+    </member>
+    <member name="T:System.Web.WebPages.WebPageHttpHandler">
+      <summary>Provides methods and properties that are used to process specific URL extensions.</summary>
+    </member>
+    <member name="M:System.Web.WebPages.WebPageHttpHandler.#ctor(System.Web.WebPages.WebPage)">
+      <summary>Initializes a new instance of the <see cref="T:System.Web.WebPages.WebPageHttpHandler" /> class by using the specified web page.</summary>
+      <param name="webPage">The web page to process.</param>
+      <exception cref="T:System.ArgumentNullException">
+        <paramref name="webPage" /> is null.</exception>
+    </member>
+    <member name="M:System.Web.WebPages.WebPageHttpHandler.CreateFromVirtualPath(System.String)">
+      <summary>Creates a new <see cref="T:System.Web.IHttpHandler" /> handler object from the specified virtual path.</summary>
+      <returns>A <see cref="T:System.Web.IHttpHandler" /> object for the specified virtual path.</returns>
+      <param name="virtualPath">The virtual path to use to create the handler.</param>
+    </member>
+    <member name="P:System.Web.WebPages.WebPageHttpHandler.DisableWebPagesResponseHeader">
+      <summary>Gets or sets a value that indicates whether web page response headers are disabled.</summary>
+      <returns>true if web page response headers are disabled; otherwise, false.</returns>
+    </member>
+    <member name="M:System.Web.WebPages.WebPageHttpHandler.GetRegisteredExtensions">
+      <summary>Returns a list of file name extensions that the current <see cref="T:System.Web.WebPages.WebPageHttpHandler" /> instance can process.</summary>
+      <returns>A read-only list of file name extensions that are processed by the current <see cref="T:System.Web.WebPages.WebPageHttpHandler" /> instance.</returns>
+    </member>
+    <member name="P:System.Web.WebPages.WebPageHttpHandler.IsReusable">
+      <summary>Gets a value that indicates whether another request can use the <see cref="T:System.Web.WebPages.WebPageHttpHandler" /> instance.</summary>
+      <returns>true if the <see cref="T:System.Web.WebPages.WebPageHttpHandler" /> instance is reusable; otherwise, false.</returns>
+    </member>
+    <member name="M:System.Web.WebPages.WebPageHttpHandler.ProcessRequest(System.Web.HttpContext)">
+      <summary>Processes the web page by using the specified context.</summary>
+      <param name="context">The context to use when processing the web page.</param>
+    </member>
+    <member name="M:System.Web.WebPages.WebPageHttpHandler.RegisterExtension(System.String)">
+      <summary>Adds a file name extension to the list of extensions that are processed by the current <see cref="T:System.Web.WebPages.WebPageHttpHandler" /> instance.</summary>
+      <param name="extension">The extension to add, without a leading period.</param>
+    </member>
+    <member name="F:System.Web.WebPages.WebPageHttpHandler.WebPagesVersionHeaderName">
+      <summary>The HTML tag name (X-AspNetWebPages-Version) for the version of the ASP.NET Web Pages specification that is used by this web page.</summary>
+    </member>
+    <member name="T:System.Web.WebPages.WebPageRenderingBase">
+      <summary>Provides methods and properties that are used to render pages that use the Razor view engine.</summary>
+    </member>
+    <member name="M:System.Web.WebPages.WebPageRenderingBase.#ctor">
+      <summary>Initializes a new instance of the <see cref="T:System.Web.WebPages.WebPageRenderingBase" /> class.</summary>
+    </member>
+    <member name="P:System.Web.WebPages.WebPageRenderingBase.Cache"></member>
+    <member name="P:System.Web.WebPages.WebPageRenderingBase.Culture"></member>
+    <member name="P:System.Web.WebPages.WebPageRenderingBase.DisplayMode"></member>
+    <member name="M:System.Web.WebPages.WebPageRenderingBase.ExecutePageHierarchy">
+      <summary>When overridden in a derived class, calls the methods that are used to initialize the page.</summary>
+    </member>
+    <member name="P:System.Web.WebPages.WebPageRenderingBase.IsAjax"></member>
+    <member name="P:System.Web.WebPages.WebPageRenderingBase.IsPost"></member>
+    <member name="P:System.Web.WebPages.WebPageRenderingBase.Layout">
+      <summary>When overridden in a derived class, gets or sets the path of a layout page.</summary>
+      <returns>The path of a layout page.</returns>
+    </member>
+    <member name="P:System.Web.WebPages.WebPageRenderingBase.Page">
+      <summary>When overridden in a derived class, provides property-like access to page data that is shared between pages, layout pages, and partial pages.</summary>
+      <returns>An object that contains page data.</returns>
+    </member>
+    <member name="P:System.Web.WebPages.WebPageRenderingBase.PageContext"></member>
+    <member name="P:System.Web.WebPages.WebPageRenderingBase.PageData">
+      <summary>When overridden in a derived class, provides array-like access to page data that is shared between pages, layout pages, and partial pages.</summary>
+      <returns>An object that provides array-like access to page data.</returns>
+    </member>
+    <member name="P:System.Web.WebPages.WebPageRenderingBase.Profile"></member>
+    <member name="M:System.Web.WebPages.WebPageRenderingBase.RenderPage(System.String,System.Object[])">
+      <summary>When overridden in a derived class, renders a web page.</summary>
+      <returns>The markup that represents the web page.</returns>
+      <param name="path">The path of the page to render.</param>
+      <param name="data">Additional data that is used to render the page.</param>
+    </member>
+    <member name="P:System.Web.WebPages.WebPageRenderingBase.Request"></member>
+    <member name="P:System.Web.WebPages.WebPageRenderingBase.Response"></member>
+    <member name="P:System.Web.WebPages.WebPageRenderingBase.Server"></member>
+    <member name="P:System.Web.WebPages.WebPageRenderingBase.Session"></member>
+    <member name="P:System.Web.WebPages.WebPageRenderingBase.TemplateInfo"></member>
+    <member name="P:System.Web.WebPages.WebPageRenderingBase.UICulture"></member>
+    <member name="P:System.Web.WebPages.WebPageRenderingBase.UrlData"></member>
+    <member name="P:System.Web.WebPages.WebPageRenderingBase.User"></member>
+    <member name="T:System.Web.WebPages.Html.HtmlHelper">
+      <summary>Provides support for rendering HTML form controls and performing form validation in a web page.</summary>
+    </member>
+    <member name="M:System.Web.WebPages.Html.HtmlHelper.AnonymousObjectToHtmlAttributes(System.Object)">
+      <summary>Creates a dictionary of HTML attributes from the input object, translating underscores to dashes.</summary>
+      <returns>A dictionary that represents HTML attributes.</returns>
+      <param name="htmlAttributes">Anonymous object describing HTML attributes.</param>
+    </member>
+    <member name="M:System.Web.WebPages.Html.HtmlHelper.AttributeEncode(System.Object)">
+      <summary>Returns an HTML-encoded string that represents the specified object by using a minimal encoding that is suitable only for HTML attributes that are enclosed in quotation marks.</summary>
+      <returns>An HTML-encoded string that represents the object.</returns>
+      <param name="value">The object to encode.</param>
+    </member>
+    <member name="M:System.Web.WebPages.Html.HtmlHelper.AttributeEncode(System.String)">
+      <summary>Returns an HTML-encoded string that represents the specified string by using a minimal encoding that is suitable only for HTML attributes that are enclosed in quotation marks.</summary>
+      <returns>An HTML-encoded string that represents the original string.</returns>
+      <param name="value">The string to encode.</param>
+    </member>
+    <member name="M:System.Web.WebPages.Html.HtmlHelper.CheckBox(System.String)">
+      <summary>Returns an HTML check box control that has the specified name.</summary>
+      <returns>The HTML markup that represents the check box control.</returns>
+      <param name="name">The value to assign to the name attribute of the HTML control element.</param>
+      <exception cref="T:Sytem.ArgumentException">
+        <paramref name="name" /> is null or empty.</exception>
+    </member>
+    <member name="M:System.Web.WebPages.Html.HtmlHelper.CheckBox(System.String,System.Boolean)">
+      <summary>Returns an HTML check box control that has the specified name and default checked status.</summary>
+      <returns>The HTML markup that represents the check box control.</returns>
+      <param name="name">The value to assign to the name attribute of the HTML control element.</param>
+      <param name="isChecked">true to indicate that the checked attribute is set to checked; otherwise, false.</param>
+      <exception cref="T:Sytem.ArgumentException">
+        <paramref name="name" /> is null or empty.</exception>
+    </member>
+    <member name="M:System.Web.WebPages.Html.HtmlHelper.CheckBox(System.String,System.Boolean,System.Collections.Generic.IDictionary{System.String,System.Object})">
+      <summary>Returns an HTML check box control that has the specified name, default checked status, and custom attributes defined by an attribute dictionary.</summary>
+      <returns>The HTML markup that represents the check box control.</returns>
+      <param name="name">The value to assign to the name attribute of the HTML control element.</param>
+      <param name="isChecked">true to indicate that the checked attribute is set to checked; otherwise, false.</param>
+      <param name="htmlAttributes">The names and values of custom attributes for the element.</param>
+      <exception cref="T:Sytem.ArgumentException">
+        <paramref name="name" /> is null or empty.</exception>
+    </member>
+    <member name="M:System.Web.WebPages.Html.HtmlHelper.CheckBox(System.String,System.Boolean,System.Object)">
+      <summary>Returns an HTML check box control that has the specified name, default checked status, and custom attributes defined by an attribute object.</summary>
+      <returns>The HTML markup that represents the check box control.</returns>
+      <param name="name">The value to assign to the name attribute of the HTML control element.</param>
+      <param name="isChecked">true to indicate that the checked attribute is set to checked; otherwise, false.</param>
+      <param name="htmlAttributes">An object that contains custom attributes for the element. The attribute names and values are retrieved through reflection by examining the properties of the object.</param>
+      <exception cref="T:Sytem.ArgumentException">
+        <paramref name="name" /> is null or empty.</exception>
+    </member>
+    <member name="M:System.Web.WebPages.Html.HtmlHelper.CheckBox(System.String,System.Collections.Generic.IDictionary{System.String,System.Object})">
+      <summary>Returns an HTML check box control that has the specified name and custom attributes defined by an attribute dictionary.</summary>
+      <returns>The HTML markup that represents the check box control.</returns>
+      <param name="name"> The value to assign to the name attribute of the HTML control element.</param>
+      <param name="htmlAttributes">The names and values of custom attributes for the element.</param>
+      <exception cref="T:Sytem.ArgumentException">
+        <paramref name="name" /> is null or empty.</exception>
+    </member>
+    <member name="M:System.Web.WebPages.Html.HtmlHelper.CheckBox(System.String,System.Object)">
+      <summary>Returns an HTML check box control that has the specified name and custom attributes defined by an attribute object.</summary>
+      <returns>The HTML markup that represents the check box control.</returns>
+      <param name="name">The value to assign to the name attribute of the HTML control element.</param>
+      <param name="htmlAttributes">An object that contains custom attributes for the element. The attribute names and values are retrieved through reflection by examining the properties of the object.</param>
+      <exception cref="T:Sytem.ArgumentException">
+        <paramref name="name" /> is null or empty.</exception>
+    </member>
+    <member name="M:System.Web.WebPages.Html.HtmlHelper.DropDownList(System.String,System.Collections.Generic.IEnumerable{System.Web.WebPages.Html.SelectListItem})">
+      <summary>Returns an HTML drop-down list control that has the specified name and that contains the specified list items.</summary>
+      <returns>The HTML markup that represents the drop-down list control.</returns>
+      <param name="name">The value to assign to the name attribute of the HTML select element.</param>
+      <param name="selectList">A list of <see cref="T:System.Web.WebPages.Html.SelectListItem" /> instances that are used to populate the list.</param>
+      <exception cref="T:Sytem.ArgumentException">
+        <paramref name="name" /> is null or empty.</exception>
+    </member>
+    <member name="M:System.Web.WebPages.Html.HtmlHelper.DropDownList(System.String,System.Collections.Generic.IEnumerable{System.Web.WebPages.Html.SelectListItem},System.Collections.Generic.IDictionary{System.String,System.Object})">
+      <summary>Returns an HTML drop-down list control that has the specified name and custom attributes defined by an attribute dictionary, and that contains the specified list items.</summary>
+      <returns>The HTML markup that represents the drop-down list control.</returns>
+      <param name="name">The value to assign to the name attribute of the HTML select element.</param>
+      <param name="selectList">A list of <see cref="T:System.Web.WebPages.Html.SelectListItem" /> instances that are used to populate the list.</param>
+      <param name="htmlAttributes">The names and values of custom attributes for the element.</param>
+      <exception cref="T:Sytem.ArgumentException">
+        <paramref name="name" /> is null or empty.</exception>
+    </member>
+    <member name="M:System.Web.WebPages.Html.HtmlHelper.DropDownList(System.String,System.Collections.Generic.IEnumerable{System.Web.WebPages.Html.SelectListItem},System.Object)">
+      <summary>Returns an HTML drop-down list control that has the specified name and custom attributes defined by an attribute object, and that contains the specified list items.</summary>
+      <returns>The HTML markup that represents the drop-down list control.</returns>
+      <param name="name">The value to assign to the name attribute of the HTML select element.</param>
+      <param name="selectList">A list of <see cref="T:System.Web.WebPages.Html.SelectListItem" /> instances that are used to populate the list.</param>
+      <param name="htmlAttributes">An object that contains custom attributes for the element. The attribute names and values are retrieved through reflection by examining the properties of the object.</param>
+      <exception cref="T:Sytem.ArgumentException">
+        <paramref name="name" /> is null or empty.</exception>
+    </member>
+    <member name="M:System.Web.WebPages.Html.HtmlHelper.DropDownList(System.String,System.String,System.Collections.Generic.IEnumerable{System.Web.WebPages.Html.SelectListItem})">
+      <summary>Returns an HTML drop-down list control that has the specified name, and that contains the specified list items and default item.</summary>
+      <returns>The HTML markup that represents the drop-down list control.</returns>
+      <param name="name">The value to assign to the name attribute of the HTML select element.</param>
+      <param name="defaultOption">The text to display for the default option in the list.</param>
+      <param name="selectList">A list of <see cref="T:System.Web.WebPages.Html.SelectListItem" /> instances that are used to populate the list.</param>
+      <exception cref="T:Sytem.ArgumentException">
+        <paramref name="name" /> is null or empty.</exception>
+    </member>
+    <member name="M:System.Web.WebPages.Html.HtmlHelper.DropDownList(System.String,System.String,System.Collections.Generic.IEnumerable{System.Web.WebPages.Html.SelectListItem},System.Collections.Generic.IDictionary{System.String,System.Object})">
+      <summary>Returns an HTML drop-down list control that has the specified name and custom attributes defined by an attribute dictionary, and that contains the specified list items and default item.</summary>
+      <returns>The HTML markup that represents the drop-down list control.</returns>
+      <param name="name">The value to assign to the name attribute of the HTML select element.</param>
+      <param name="defaultOption">The text to display for the default option in the list.</param>
+      <param name="selectList">A list of <see cref="T:System.Web.WebPages.Html.SelectListItem" /> instances that are used to populate the list.</param>
+      <param name="htmlAttributes">The names and values of custom attributes for the element.</param>
+      <exception cref="T:Sytem.ArgumentException">
+        <paramref name="name" /> is null or empty.</exception>
+    </member>
+    <member name="M:System.Web.WebPages.Html.HtmlHelper.DropDownList(System.String,System.String,System.Collections.Generic.IEnumerable{System.Web.WebPages.Html.SelectListItem},System.Object)">
+      <summary>Returns an HTML drop-down list control that has the specified name and custom attributes defined by an attribute object, and that contains the specified list items and default item.</summary>
+      <returns>The HTML markup that represents the drop-down list control.</returns>
+      <param name="name">The value to assign to the name attribute of the HTML select element.</param>
+      <param name="defaultOption">The text to display for the default option in the list.</param>
+      <param name="selectList">A list of <see cref="T:System.Web.WebPages.Html.SelectListItem" /> instances that are used to populate the list.</param>
+      <param name="htmlAttributes">An object that contains custom attributes for the element. The attribute names and values are retrieved through reflection by examining the properties of the object.</param>
+      <exception cref="T:Sytem.ArgumentException">
+        <paramref name="name" /> is null or empty.</exception>
+    </member>
+    <member name="M:System.Web.WebPages.Html.HtmlHelper.DropDownList(System.String,System.String,System.Collections.Generic.IEnumerable{System.Web.WebPages.Html.SelectListItem},System.Object,System.Collections.Generic.IDictionary{System.String,System.Object})">
+      <summary>Returns an HTML drop-down list control that has the specified name, custom attributes defined by an attribute dictionary, and default selection, and that contains the specified list items and default item.</summary>
+      <returns>The HTML markup that represents the drop-down list control.</returns>
+      <param name="name">The value to assign to the name attribute of the HTML select element.</param>
+      <param name="defaultOption">The text to display for the default option in the list.</param>
+      <param name="selectList">A list of <see cref="T:System.Web.WebPages.Html.SelectListItem" /> instances that are used to populate the list.</param>
+      <param name="selectedValue">The value that specifies the item in the list that is selected by default. The selected item is the first item in the list whose value matches the parameter (or whose text matches, if there is no value.) </param>
+      <param name="htmlAttributes">The names and values of custom attributes for the element.</param>
+      <exception cref="T:Sytem.ArgumentException">
+        <paramref name="name" /> is null or empty.</exception>
+    </member>
+    <member name="M:System.Web.WebPages.Html.HtmlHelper.DropDownList(System.String,System.String,System.Collections.Generic.IEnumerable{System.Web.WebPages.Html.SelectListItem},System.Object,System.Object)">
+      <summary>Returns an HTML drop-down list control that has the specified name, custom attributes defined by an attribute object, and default selection, and that contains the specified list items and default item.</summary>
+      <returns>The HTML markup that represents the drop-down list control.</returns>
+      <param name="name">The value to assign to the name attribute of the HTML select element.</param>
+      <param name="defaultOption">The text to display for the default option in the list.</param>
+      <param name="selectList">A list of <see cref="T:System.Web.WebPages.Html.SelectListItem" /> instances that are used to populate the list.</param>
+      <param name="selectedValue">The value that specifies the item in the list that is selected by default. The item that is selected is the first item in the list that has a matching value, or that matches the items displayed text if the item has no value.</param>
+      <param name="htmlAttributes">An object that contains custom attributes for the element. The attribute names and values are retrieved through reflection by examining the properties of the object.</param>
+      <exception cref="T:Sytem.ArgumentException">
+        <paramref name="name" /> is null or empty.</exception>
+    </member>
+    <member name="M:System.Web.WebPages.Html.HtmlHelper.Encode(System.Object)">
+      <summary>Returns an HTML-encoded string that represents the specified object by using a full encoding that is suitable for arbitrary HTML.</summary>
+      <returns>An HTML-encoded string that represents the object.</returns>
+      <param name="value">The object to encode.</param>
+    </member>
+    <member name="M:System.Web.WebPages.Html.HtmlHelper.Encode(System.String)">
+      <summary>Returns an HTML-encoded string that represents the specified string by using a full encoding that is suitable for arbitrary HTML.</summary>
+      <returns>An HTML-encoded string that represents the original string.</returns>
+      <param name="value">The string to encode.</param>
+    </member>
+    <member name="M:System.Web.WebPages.Html.HtmlHelper.Hidden(System.String)">
+      <summary>Returns an HTML hidden control that has the specified name.</summary>
+      <returns>The HTML markup that represents the hidden control.</returns>
+      <param name="name">The value to assign to the name attribute of the HTML control element.</param>
+      <exception cref="T:Sytem.ArgumentException">
+        <paramref name="name" /> is null or empty.</exception>
+    </member>
+    <member name="M:System.Web.WebPages.Html.HtmlHelper.Hidden(System.String,System.Object)">
+      <summary>Returns an HTML hidden control that has the specified name and value.</summary>
+      <returns>The HTML markup that represents the hidden control.</returns>
+      <param name="name">The value to assign to the name attribute of the HTML control element.</param>
+      <param name="value">The value to assign to the value attribute of the element.</param>
+      <exception cref="T:Sytem.ArgumentException">
+        <paramref name="name" /> is null or empty.</exception>
+    </member>
+    <member name="M:System.Web.WebPages.Html.HtmlHelper.Hidden(System.String,System.Object,System.Collections.Generic.IDictionary{System.String,System.Object})">
+      <summary>Returns an HTML hidden control that has the specified name, value, and custom attributes defined by an attribute dictionary.</summary>
+      <returns>The HTML markup that represents the hidden control.</returns>
+      <param name="name">The value to assign to the name attribute of the HTML control element.</param>
+      <param name="value">The value to assign to the value attribute of the element.</param>
+      <param name="htmlAttributes">The names and values of custom attributes for the element.</param>
+      <exception cref="T:Sytem.ArgumentException">
+        <paramref name="name" /> is null or empty.</exception>
+    </member>
+    <member name="M:System.Web.WebPages.Html.HtmlHelper.Hidden(System.String,System.Object,System.Object)">
+      <summary>Returns an HTML hidden control that has the specified name, value, and custom attributes defined by an attribute object.</summary>
+      <returns>The HTML markup that represents the hidden control.</returns>
+      <param name="name"> The value to assign to the name attribute of the HTML control element.</param>
+      <param name="value">The value to assign to the value attribute of the element.</param>
+      <param name="htmlAttributes">An object that contains custom attributes for the element. The attribute names and values are retrieved through reflection by examining the properties of the object.</param>
+      <exception cref="T:Sytem.ArgumentException">
+        <paramref name="name" /> is null or empty.</exception>
+    </member>
+    <member name="P:System.Web.WebPages.Html.HtmlHelper.IdAttributeDotReplacement">
+      <summary>Gets or sets the character that is used to replace the dot (.) in the id attribute of rendered form controls.</summary>
+      <returns>The character that is used to replace the dot in the id attribute of rendered form controls. The default is an underscore (_).</returns>
+    </member>
+    <member name="M:System.Web.WebPages.Html.HtmlHelper.Label(System.String)">
+      <summary>Returns an HTML label that displays the specified text.</summary>
+      <returns>The HTML markup that represents the label.</returns>
+      <param name="labelText">The text to display.</param>
+      <exception cref="T:Sytem.ArgumentException">
+        <paramref name="labelText" /> is null or empty.</exception>
+    </member>
+    <member name="M:System.Web.WebPages.Html.HtmlHelper.Label(System.String,System.Object)">
+      <summary>Returns an HTML label that displays the specified text and that has the specified custom attributes.</summary>
+      <returns>The HTML markup that represents the label.</returns>
+      <param name="labelText">The text to display.</param>
+      <param name="attributes">An object that contains custom attributes for the element. The attribute names and values are retrieved through reflection by examining the properties of the object.</param>
+      <exception cref="T:Sytem.ArgumentException">
+        <paramref name="labelText" /> is null or empty.</exception>
+    </member>
+    <member name="M:System.Web.WebPages.Html.HtmlHelper.Label(System.String,System.String)">
+      <summary>Returns an HTML label that displays the specified text and that has the specified for attribute.</summary>
+      <returns>The HTML markup that represents the label.</returns>
+      <param name="labelText">The text to display.</param>
+      <param name="labelFor">The value to assign to the for attribute of the HTML control element.</param>
+      <exception cref="T:Sytem.ArgumentException">
+        <paramref name="labelText" /> is null or empty.</exception>
+    </member>
+    <member name="M:System.Web.WebPages.Html.HtmlHelper.Label(System.String,System.String,System.Collections.Generic.IDictionary{System.String,System.Object})">
+      <summary>Returns an HTML label that displays the specified text, and that has the specified for attribute and custom attributes defined by an attribute dictionary.</summary>
+      <returns>The HTML markup that represents the label.</returns>
+      <param name="labelText">The text to display.</param>
+      <param name="labelFor">The value to assign to the for attribute of the HTML control element.</param>
+      <param name="attributes">The names and values of custom attributes for the element.</param>
+      <exception cref="T:Sytem.ArgumentException">
+        <paramref name="labelText" /> is null or empty.</exception>
+    </member>
+    <member name="M:System.Web.WebPages.Html.HtmlHelper.Label(System.String,System.String,System.Object)">
+      <summary>Returns an HTML label that displays the specified text, and that has the specified for attribute and custom attributes defined by an attribute object.</summary>
+      <returns>The HTML markup that represents the label.</returns>
+      <param name="labelText">The text to display.</param>
+      <param name="labelFor">The value to assign to the for attribute of the HTML control element.</param>
+      <param name="attributes">An object that contains custom attributes for the element. The attribute names and values are retrieved through reflection by examining the properties of the object.</param>
+      <exception cref="T:Sytem.ArgumentException">
+        <paramref name="labelText" /> is null or empty.</exception>
+    </member>
+    <member name="M:System.Web.WebPages.Html.HtmlHelper.ListBox(System.String,System.Collections.Generic.IEnumerable{System.Web.WebPages.Html.SelectListItem})">
+      <summary>Returns an HTML list box control that has the specified name and that contains the specified list items.</summary>
+      <returns>The HTML markup that represents the list box control.</returns>
+      <param name="name">The value to assign to the name attribute of the HTML select element.</param>
+      <param name="selectList">A list of <see cref="T:System.Web.WebPages.Html.SelectListItem" /> instances that are used to populate the list.</param>
+      <exception cref="T:Sytem.ArgumentException">
+        <paramref name="name" /> is null or empty.</exception>
+    </member>
+    <member name="M:System.Web.WebPages.Html.HtmlHelper.ListBox(System.String,System.Collections.Generic.IEnumerable{System.Web.WebPages.Html.SelectListItem},System.Collections.Generic.IDictionary{System.String,System.Object})">
+      <summary>Returns an HTML list box control that has the specified name and custom attributes defined by an attribute dictionary, and that contains the specified list items.</summary>
+      <returns>The HTML markup that represents the list box control.</returns>
+      <param name="name">The value to assign to the name attribute of the HTML select element.</param>
+      <param name="selectList">A list of <see cref="T:System.Web.WebPages.Html.SelectListItem" /> instances that are used to populate the list.</param>
+      <param name="htmlAttributes">The names and values of custom attributes for the element.</param>
+      <exception cref="T:Sytem.ArgumentException">
+        <paramref name="name" /> is null or empty.</exception>
+    </member>
+    <member name="M:System.Web.WebPages.Html.HtmlHelper.ListBox(System.String,System.Collections.Generic.IEnumerable{System.Web.WebPages.Html.SelectListItem},System.Object)">
+      <summary>Returns an HTML list box control that has the specified name and custom attributes defined by an attribute object, and that contains the specified list items.</summary>
+      <returns>The HTML markup that represents the list box control.</returns>
+      <param name="name">The value to assign to the name attribute of the HTML select element.</param>
+      <param name="selectList">A list of <see cref="T:System.Web.WebPages.Html.SelectListItem" /> instances that are used to populate the list.</param>
+      <param name="htmlAttributes">An object that contains custom attributes for the element. The attribute names and values are retrieved through reflection by examining the properties of the object.</param>
+      <exception cref="T:Sytem.ArgumentException">
+        <paramref name="name" /> is null or empty.</exception>
+    </member>
+    <member name="M:System.Web.WebPages.Html.HtmlHelper.ListBox(System.String,System.Collections.Generic.IEnumerable{System.Web.WebPages.Html.SelectListItem},System.Object,System.Int32,System.Boolean)">
+      <summary>Returns an HTML list box control that has the specified name, size, list items, and default selections, and that specifies whether multiple selections are enabled.</summary>
+      <returns>The HTML markup that represents the list box control.</returns>
+      <param name="name">The value to assign to the name attribute of the HTML select element.</param>
+      <param name="selectList">A list of <see cref="T:System.Web.WebPages.Html.SelectListItem" /> instances that are used to populate the list.</param>
+      <param name="selectedValues">An object that specifies the items in the list that are selected by default. The selections are retrieved through reflection by examining the properties of the object.</param>
+      <param name="size">The value to assign to the size attribute of the element.</param>
+      <param name="allowMultiple">true to indicate that the multiple selections are enabled; otherwise, false.</param>
+      <exception cref="T:Sytem.ArgumentException">
+        <paramref name="name" /> is null or empty.</exception>
+    </member>
+    <member name="M:System.Web.WebPages.Html.HtmlHelper.ListBox(System.String,System.String,System.Collections.Generic.IEnumerable{System.Web.WebPages.Html.SelectListItem})">
+      <summary>Returns an HTML list box control that has the specified name, and that contains the specified list items and default item.</summary>
+      <returns>The HTML markup that represents the list box control.</returns>
+      <param name="name">The value to assign to the name attribute of the HTML select element.</param>
+      <param name="defaultOption">The text to display for the default option in the list.</param>
+      <param name="selectList">A list of <see cref="T:System.Web.WebPages.Html.SelectListItem" /> instances that are used to populate the list box.</param>
+      <exception cref="T:Sytem.ArgumentException">
+        <paramref name="name" /> is null or empty.</exception>
+    </member>
+    <member name="M:System.Web.WebPages.Html.HtmlHelper.ListBox(System.String,System.String,System.Collections.Generic.IEnumerable{System.Web.WebPages.Html.SelectListItem},System.Collections.Generic.IDictionary{System.String,System.Object})">
+      <summary>Returns an HTML list box control that has the specified name and custom attributes defined by an attribute dictionary, and that contains the specified list items and default item.</summary>
+      <returns>The HTML markup that represents the list box control.</returns>
+      <param name="name">The value to assign to the name attribute of the HTML select element.</param>
+      <param name="defaultOption">The text to display for the default option in the list.</param>
+      <param name="selectList">A list of <see cref="T:System.Web.WebPages.Html.SelectListItem" /> instances that are used to populate the list.</param>
+      <param name="htmlAttributes">The names and values of custom attributes for the element.</param>
+      <exception cref="T:Sytem.ArgumentException">
+        <paramref name="name" /> is null or empty.</exception>
+    </member>
+    <member name="M:System.Web.WebPages.Html.HtmlHelper.ListBox(System.String,System.String,System.Collections.Generic.IEnumerable{System.Web.WebPages.Html.SelectListItem},System.Object)">
+      <summary>Returns an HTML list box control that has the specified name and custom attributes defined by an attribute object, and that contains the specified list items and default item.</summary>
+      <returns>The HTML markup that represents the list box control.</returns>
+      <param name="name">The value to assign to the name attribute of the HTML select element.</param>
+      <param name="defaultOption">The text to display for the default option in the list.</param>
+      <param name="selectList">A list of <see cref="T:System.Web.WebPages.Html.SelectListItem" /> instances that are used to populate the list box.</param>
+      <param name="htmlAttributes">An object that contains custom attributes for the element. The attribute names and values are retrieved through reflection by examining the properties of the object.</param>
+      <exception cref="T:Sytem.ArgumentException">
+        <paramref name="name" /> is null or empty.</exception>
+    </member>
+    <member name="M:System.Web.WebPages.Html.HtmlHelper.ListBox(System.String,System.String,System.Collections.Generic.IEnumerable{System.Web.WebPages.Html.SelectListItem},System.Object,System.Collections.Generic.IDictionary{System.String,System.Object})">
+      <summary>Returns an HTML list box control that has the specified name and custom attributes defined by an attribute dictionary, and that contains the specified list items, default item, and selections.</summary>
+      <returns>The HTML markup that represents the list box control.</returns>
+      <param name="name">The value to assign to the name attribute of the HTML select element.</param>
+      <param name="defaultOption">The text to display for the default option in the list.</param>
+      <param name="selectList">A list of <see cref="T:System.Web.WebPages.Html.SelectListItem" /> instances that are used to populate the list.</param>
+      <param name="selectedValues">An object that specifies the items in the list that are selected by default. The selections are retrieved through reflection by examining the properties of the object. </param>
+      <param name="htmlAttributes">The names and values of custom attributes for the element.</param>
+      <exception cref="T:Sytem.ArgumentException">
+        <paramref name="name" /> is null or empty.</exception>
+    </member>
+    <member name="M:System.Web.WebPages.Html.HtmlHelper.ListBox(System.String,System.String,System.Collections.Generic.IEnumerable{System.Web.WebPages.Html.SelectListItem},System.Object,System.Int32,System.Boolean)">
+      <summary>Returns an HTML list box control that has the specified name, size, items, default item, and selections, and that specifies whether multiple selections are enabled.</summary>
+      <returns>The HTML markup that represents the list box control.</returns>
+      <param name="name">The value to assign to the name attribute of the HTML select element.</param>
+      <param name="defaultOption">The text to display for the default option in the list.</param>
+      <param name="selectList">A list of <see cref="T:System.Web.WebPages.Html.SelectListItem" /> instances that are used to populate the list.</param>
+      <param name="selectedValues">An object that specifies the items in the list that are selected by default. The selections are retrieved through reflection by examining the properties of the object.</param>
+      <param name="size">The value to assign to the size attribute of the element.</param>
+      <param name="allowMultiple">true to indicate that multiple selections are enabled; otherwise, false.</param>
+      <exception cref="T:Sytem.ArgumentException">
+        <paramref name="name" /> is null or empty.</exception>
+    </member>
+    <member name="M:System.Web.WebPages.Html.HtmlHelper.ListBox(System.String,System.String,System.Collections.Generic.IEnumerable{System.Web.WebPages.Html.SelectListItem},System.Object,System.Int32,System.Boolean,System.Collections.Generic.IDictionary{System.String,System.Object})">
+      <summary>Returns an HTML list box control that has the specified name, size, custom attributes defined by an attribute dictionary, items, default item, and selections, and that specifies whether multiple selections are enabled.</summary>
+      <returns>The HTML markup that represents the list box control.</returns>
+      <param name="name">The value to assign to the name attribute of the HTML select element.</param>
+      <param name="defaultOption">The text to display for the default option in the list.</param>
+      <param name="selectList">A list of <see cref="T:System.Web.WebPages.Html.SelectListItem" /> instances that are used to populate the list.</param>
+      <param name="selectedValues">An object that specifies the items in the list that are selected by default. The selections are retrieved through reflection by examining the properties of the object.</param>
+      <param name="size">The value to assign to the size attribute of the element.</param>
+      <param name="allowMultiple">true to indicate that multiple selections are enabled; otherwise, false.</param>
+      <param name="htmlAttributes">The names and values of custom attributes for the element.</param>
+      <exception cref="T:Sytem.ArgumentException">
+        <paramref name="name" /> is null or empty.</exception>
+    </member>
+    <member name="M:System.Web.WebPages.Html.HtmlHelper.ListBox(System.String,System.String,System.Collections.Generic.IEnumerable{System.Web.WebPages.Html.SelectListItem},System.Object,System.Int32,System.Boolean,System.Object)">
+      <summary>Returns an HTML list box control that has the specified name, size, custom attributes defined by an attribute object, items, default item, and selections, and that specifies whether multiple selections are enabled.</summary>
+      <returns>The HTML markup that represents the list box control.</returns>
+      <param name="name">The value to assign to the name attribute of the HTML select element.</param>
+      <param name="defaultOption">The text to display for the default option in the list.</param>
+      <param name="selectList">A list of <see cref="T:System.Web.WebPages.Html.SelectListItem" /> instances that are used to populate the list.</param>
+      <param name="selectedValues">An object that specifies the items in the list that are selected by default. The selections are retrieved through reflection by examining the properties of the object.</param>
+      <param name="size">The value to assign to the size attribute of the element.</param>
+      <param name="allowMultiple">true to indicate that multiple selections are enabled; otherwise, false.</param>
+      <param name="htmlAttributes">An object that contains custom attributes for the element. The attribute names and values are retrieved through reflection by examining the properties of the object.</param>
+      <exception cref="T:Sytem.ArgumentException">
+        <paramref name="name" /> is null or empty.</exception>
+    </member>
+    <member name="M:System.Web.WebPages.Html.HtmlHelper.ListBox(System.String,System.String,System.Collections.Generic.IEnumerable{System.Web.WebPages.Html.SelectListItem},System.Object,System.Object)">
+      <summary>Returns an HTML list box control that has the specified name, items, default item, and custom attributes defined by an attribute object, and selections.</summary>
+      <returns>The HTML markup that represents the list box control.</returns>
+      <param name="name">The value to assign to the name attribute of the HTML select element.</param>
+      <param name="defaultOption">The text to display for the default option in the list.</param>
+      <param name="selectList">A list of <see cref="T:System.Web.WebPages.Html.SelectListItem" /> instances that are used to populate the list.</param>
+      <param name="selectedValues">An object that specifies the items in the list that are selected by default. The selections are retrieved through reflection by examining the properties of the object.</param>
+      <param name="htmlAttributes">An object that contains custom attributes for the element. The attribute names and values are retrieved through reflection by examining the properties of the object.</param>
+      <exception cref="T:Sytem.ArgumentException">
+        <paramref name="name" /> is null or empty.</exception>
+    </member>
+    <member name="M:System.Web.WebPages.Html.HtmlHelper.ObjectToDictionary(System.Object)">
+      <summary>Creates a dictionary from an object, by adding each public instance property as a key with its associated value to the dictionary. It will expose public properties from derived types as well. This is typically used with objects of an anonymous type.</summary>
+      <returns>The created dictionary of property names and property values.</returns>
+      <param name="value">The object to be converted.</param>
+    </member>
+    <member name="M:System.Web.WebPages.Html.HtmlHelper.Password(System.String)">
+      <summary>Returns an HTML password control that has the specified name.</summary>
+      <returns>The HTML markup that represents the password control.</returns>
+      <param name="name">The value to assign to the name attribute of the HTML control element.</param>
+      <exception cref="T:Sytem.ArgumentException">
+        <paramref name="name" /> is null or empty.</exception>
+    </member>
+    <member name="M:System.Web.WebPages.Html.HtmlHelper.Password(System.String,System.Object)">
+      <summary>Returns an HTML password control that has the specified name and value.</summary>
+      <returns>The HTML markup that represents the password control.</returns>
+      <param name="name">The value to assign to the name attribute of the HTML control element.</param>
+      <param name="value">The value to assign to the value attribute of the element.</param>
+      <exception cref="T:Sytem.ArgumentException">
+        <paramref name="name" /> is null or empty.</exception>
+    </member>
+    <member name="M:System.Web.WebPages.Html.HtmlHelper.Password(System.String,System.Object,System.Collections.Generic.IDictionary{System.String,System.Object})">
+      <summary>Returns an HTML password control that has the specified name, value, and custom attributes defined by an attribute dictionary.</summary>
+      <returns>The HTML markup that represents the password control.</returns>
+      <param name="name">The value to assign to the name attribute of the HTML control element.</param>
+      <param name="value">The value to assign to the value attribute of the element.</param>
+      <param name="htmlAttributes">The names and values of custom attributes for the element.</param>
+      <exception cref="T:Sytem.ArgumentException">
+        <paramref name="name" /> is null or empty.</exception>
+    </member>
+    <member name="M:System.Web.WebPages.Html.HtmlHelper.Password(System.String,System.Object,System.Object)">
+      <summary>Returns an HTML password control that has the specified name, value, and custom attributes defined by an attribute object.</summary>
+      <returns>The HTML markup that represents the password control.</returns>
+      <param name="name">The value to assign to the name attribute of the HTML control element.</param>
+      <param name="value">The value to assign to the value attribute of the element.</param>
+      <param name="htmlAttributes"> An object that contains custom attributes for the element. The attribute names and values are retrieved through reflection by examining the properties of the object.</param>
+      <exception cref="T:Sytem.ArgumentException">
+        <paramref name="name" /> is null or empty.</exception>
+    </member>
+    <member name="M:System.Web.WebPages.Html.HtmlHelper.RadioButton(System.String,System.Object)">
+      <summary>Returns an HTML radio button control that has the specified name and value.</summary>
+      <returns>The HTML markup that represents the radio button control.</returns>
+      <param name="name">The value to assign to the name attribute of the HTML control element. The name attribute defines the group that the radio button belongs to.</param>
+      <param name="value">The value to assign to the value attribute of the element.</param>
+      <exception cref="T:Sytem.ArgumentException">
+        <paramref name="name" /> is null or empty.</exception>
+    </member>
+    <member name="M:System.Web.WebPages.Html.HtmlHelper.RadioButton(System.String,System.Object,System.Boolean)">
+      <summary>Returns an HTML radio button control that has the specified name, value, and default selected status.</summary>
+      <returns>The HTML markup that represents the radio button control.</returns>
+      <param name="name">The value to assign to the name attribute of the HTML control element. The name attribute defines the group that the radio button belongs to.</param>
+      <param name="value">The value to assign to the value attribute of the element.</param>
+      <param name="isChecked">true to indicate that the control is selected; otherwise, false.</param>
+      <exception cref="T:Sytem.ArgumentException">
+        <paramref name="name" /> is null or empty.</exception>
+    </member>
+    <member name="M:System.Web.WebPages.Html.HtmlHelper.RadioButton(System.String,System.Object,System.Boolean,System.Collections.Generic.IDictionary{System.String,System.Object})">
+      <summary>Returns an HTML radio button control that has the specified name, value, default selected status, and custom attributes defined by an attribute dictionary.</summary>
+      <returns>The HTML markup that represents the radio button control.</returns>
+      <param name="name">The value to assign to the name attribute of the HTML control element. The name attribute defines the group that the radio button belongs to.</param>
+      <param name="value">The value to assign to the value attribute of the element.</param>
+      <param name="isChecked">true to indicate that the control is selected; otherwise, false.</param>
+      <param name="htmlAttributes">The names and values of custom attributes for the element.</param>
+      <exception cref="T:Sytem.ArgumentException">
+        <paramref name="name" /> is null or empty.</exception>
+    </member>
+    <member name="M:System.Web.WebPages.Html.HtmlHelper.RadioButton(System.String,System.Object,System.Boolean,System.Object)">
+      <summary>Returns an HTML radio button control that has the specified name, value, default selected status, and custom attributes defined by an attribute object.</summary>
+      <returns>The HTML markup that represents the radio button control.</returns>
+      <param name="name">The value to assign to the name attribute of the HTML control element. The name attribute defines the group that the radio button belongs to.</param>
+      <param name="value">The value to assign to the value attribute of the element.</param>
+      <param name="isChecked">true to indicate that the control is selected; otherwise, false.</param>
+      <param name="htmlAttributes">An object that contains custom attributes for the element. The attribute names and values are retrieved through reflection by examining the properties of the object.</param>
+      <exception cref="T:Sytem.ArgumentException">
+        <paramref name="name" /> is null or empty.</exception>
+    </member>
+    <member name="M:System.Web.WebPages.Html.HtmlHelper.RadioButton(System.String,System.Object,System.Collections.Generic.IDictionary{System.String,System.Object})">
+      <summary>Returns an HTML radio button control that has the specified name, value, and custom attributes defined by an attribute dictionary.</summary>
+      <returns>The HTML markup that represents the radio button control.</returns>
+      <param name="name">The value to assign to the name attribute of the HTML control element. The name attribute defines the group that the radio button belongs to.</param>
+      <param name="value">The value to assign to the value attribute of the element.</param>
+      <param name="htmlAttributes">The names and values of custom attributes for the element.</param>
+      <exception cref="T:Sytem.ArgumentException">
+        <paramref name="name" /> is null or empty.</exception>
+    </member>
+    <member name="M:System.Web.WebPages.Html.HtmlHelper.RadioButton(System.String,System.Object,System.Object)">
+      <summary>Returns an HTML radio button control that has the specified name, value, and custom attributes defined by an attribute object.</summary>
+      <returns>The HTML markup that represents the radio button control.</returns>
+      <param name="name">The value to assign to the name attribute of the HTML control element. The name attribute defines the group that the radio button belongs to.</param>
+      <param name="value">The value to assign to the value attribute of the element.</param>
+      <param name="htmlAttributes">An object that contains custom attributes for the element. The attribute names and values are retrieved through reflection by examining the properties of the object.</param>
+      <exception cref="T:Sytem.ArgumentException">
+        <paramref name="name" /> is null or empty.</exception>
+    </member>
+    <member name="M:System.Web.WebPages.Html.HtmlHelper.Raw(System.Object)">
+      <summary>Wraps HTML markup in an <see cref="T:System.Web.HtmlString" /> instance so that it is interpreted as HTML markup.</summary>
+      <returns>The unencoded HTML.</returns>
+      <param name="value">The object to render HTML for.</param>
+    </member>
+    <member name="M:System.Web.WebPages.Html.HtmlHelper.Raw(System.String)">
+      <summary>Wraps HTML markup in an <see cref="T:System.Web.HtmlString" /> instance so that it is interpreted as HTML markup.</summary>
+      <returns>The unencoded HTML.</returns>
+      <param name="value">The string to interpret as HTML markup instead of being HTML-encoded.</param>
+    </member>
+    <member name="M:System.Web.WebPages.Html.HtmlHelper.TextArea(System.String)">
+      <summary>Returns an HTML multi-line text input (text area) control that has the specified name.</summary>
+      <returns>The HTML markup that represents the text area control.</returns>
+      <param name="name">The value to assign to the name attribute of the HTML textarea element.</param>
+      <exception cref="T:Sytem.ArgumentException">
+        <paramref name="name" /> is null or empty.</exception>
+    </member>
+    <member name="M:System.Web.WebPages.Html.HtmlHelper.TextArea(System.String,System.Collections.Generic.IDictionary{System.String,System.Object})">
+      <summary>Returns an HTML multi-line text input (text area) control that has the specified name and custom attributes defined by an attribute dictionary.</summary>
+      <returns>The HTML markup that represents the text area control.</returns>
+      <param name="name">The value to assign to the name attribute of the HTML textarea element.</param>
+      <param name="htmlAttributes">The names and values of custom attributes for the element.</param>
+      <exception cref="T:Sytem.ArgumentException">
+        <paramref name="name" /> is null or empty.</exception>
+    </member>
+    <member name="M:System.Web.WebPages.Html.HtmlHelper.TextArea(System.String,System.Object)">
+      <summary>Returns an HTML multi-line text input (text area) control that has the specified name and custom attributes defined by an attribute object.</summary>
+      <returns>The HTML markup that represents the text area control.</returns>
+      <param name="name">The value to assign to the name attribute of the HTML textarea element.</param>
+      <param name="htmlAttributes">An object that contains custom attributes for the element. The attribute names and values are retrieved through reflection by examining the properties of the object.</param>
+      <exception cref="T:Sytem.ArgumentException">
+        <paramref name="name" /> is null or empty.</exception>
+    </member>
+    <member name="M:System.Web.WebPages.Html.HtmlHelper.TextArea(System.String,System.String)">
+      <summary>Returns an HTML multi-line text input (text area) control that has the specified name and value.</summary>
+      <returns>The HTML markup that represents the text area control.</returns>
+      <param name="name">The value to assign to the name attribute of the HTML textrarea element.</param>
+      <param name="value">The text to display.</param>
+      <exception cref="T:Sytem.ArgumentException">
+        <paramref name="name" /> is null or empty.</exception>
+    </member>
+    <member name="M:System.Web.WebPages.Html.HtmlHelper.TextArea(System.String,System.String,System.Collections.Generic.IDictionary{System.String,System.Object})">
+      <summary>Returns an HTML multi-line text input (text area) control that has the specified name, value, and custom attributes defined by an attribute dictionary.</summary>
+      <returns>The HTML markup that represents the text area control.</returns>
+      <param name="name">The value to assign to the name attribute of the HTML textarea element.</param>
+      <param name="value">The text to display.</param>
+      <param name="htmlAttributes">The names and values of custom attributes for the element.</param>
+      <exception cref="T:Sytem.ArgumentException">
+        <paramref name="name" /> is null or empty.</exception>
+    </member>
+    <member name="M:System.Web.WebPages.Html.HtmlHelper.TextArea(System.String,System.String,System.Int32,System.Int32,System.Collections.Generic.IDictionary{System.String,System.Object})">
+      <summary>Returns an HTML multi-line text input (text area) control that has the specified name, value, row attribute, col attribute, and custom attributes defined by an attribute dictionary.</summary>
+      <returns>The HTML markup that represents the text area control.</returns>
+      <param name="name">The value to assign to the name attribute of the HTML textarea element.</param>
+      <param name="value">The text to display.</param>
+      <param name="rows">The value to assign to the rows attribute of the element.</param>
+      <param name="columns">The value to assign to the cols attribute of the element.</param>
+      <param name="htmlAttributes">The names and values of custom attributes for the element.</param>
+      <exception cref="T:Sytem.ArgumentException">
+        <paramref name="name" /> is null or empty.</exception>
+    </member>
+    <member name="M:System.Web.WebPages.Html.HtmlHelper.TextArea(System.String,System.String,System.Int32,System.Int32,System.Object)">
+      <summary>Returns an HTML multi-line text input (text area) control that has the specified name, value, row attribute, col attribute, and custom attributes defined by an attribute object.</summary>
+      <returns>The HTML markup that represents the text area control.</returns>
+      <param name="name">The value to assign to the name attribute of the HTML textarea element.</param>
+      <param name="value">The text to display.</param>
+      <param name="rows">The value to assign to the rows attribute of the element.</param>
+      <param name="columns">The value to assign to the cols attribute of the element.</param>
+      <param name="htmlAttributes">An object that contains custom attributes for the element. The attribute names and values are retrieved through reflection by examining the properties of the object.</param>
+      <exception cref="T:Sytem.ArgumentException">
+        <paramref name="name" /> is null or empty.</exception>
+    </member>
+    <member name="M:System.Web.WebPages.Html.HtmlHelper.TextArea(System.String,System.String,System.Object)">
+      <summary>Returns an HTML multi-line text input (text area) control that has the specified name, value, and custom attributes defined by an attribute object.</summary>
+      <returns>The HTML markup that represents the text area control.</returns>
+      <param name="name">The value to assign to the name attribute of the HTML textarea element.</param>
+      <param name="value">The text to display.</param>
+      <param name="htmlAttributes">An object that contains custom attributes for the element. The attribute names and values are retrieved through reflection by examining the properties of the object.</param>
+      <exception cref="T:Sytem.ArgumentException">
+        <paramref name="name" /> is null or empty.</exception>
+    </member>
+    <member name="M:System.Web.WebPages.Html.HtmlHelper.TextBox(System.String)">
+      <summary>Returns an HTML text control that has the specified name.</summary>
+      <returns>The HTML markup that represents the text control.</returns>
+      <param name="name">The value to assign to the name attribute of the HTML control element.</param>
+      <exception cref="T:Sytem.ArgumentException">
+        <paramref name="name" /> is null or empty.</exception>
+    </member>
+    <member name="M:System.Web.WebPages.Html.HtmlHelper.TextBox(System.String,System.Object)">
+      <summary>Returns an HTML text control that has the specified name and value.</summary>
+      <returns>The HTML markup that represents the text control.</returns>
+      <param name="name">The value to assign to the name attribute of the HTML control element.</param>
+      <param name="value">The value to assign to the value attribute of the element.</param>
+      <exception cref="T:Sytem.ArgumentException">
+        <paramref name="name" /> is null or empty.</exception>
+    </member>
+    <member name="M:System.Web.WebPages.Html.HtmlHelper.TextBox(System.String,System.Object,System.Collections.Generic.IDictionary{System.String,System.Object})">
+      <summary>Returns an HTML text control that has the specified name, value, and custom attributes defined by an attribute dictionary.</summary>
+      <returns>The HTML markup that represents the text control.</returns>
+      <param name="name">The value to assign to the name attribute of the HTML control element.</param>
+      <param name="value">The value to assign to the value attribute of the element.</param>
+      <param name="htmlAttributes">The names and values of custom attributes for the element.</param>
+      <exception cref="T:Sytem.ArgumentException">
+        <paramref name="name" /> is null or empty.</exception>
+    </member>
+    <member name="M:System.Web.WebPages.Html.HtmlHelper.TextBox(System.String,System.Object,System.Object)">
+      <summary>Returns an HTML text control that has the specified name, value, and custom attributes defined by an attribute object.</summary>
+      <returns>The HTML markup that represents the text control.</returns>
+      <param name="name">The value to assign to the name attribute of the HTML control element.</param>
+      <param name="value">The value to assign to the value attribute of the element.</param>
+      <param name="htmlAttributes">An object that contains custom attributes for the element. The attribute names and values are retrieved through reflection by examining the properties of the object.</param>
+      <exception cref="T:Sytem.ArgumentException">
+        <paramref name="name" /> is null or empty.</exception>
+    </member>
+    <member name="P:System.Web.WebPages.Html.HtmlHelper.UnobtrusiveJavaScriptEnabled">
+      <summary>Gets or sets a value that indicates whether the page uses unobtrusive JavaScript for Ajax functionality.</summary>
+      <returns>true if the page uses unobtrusive JavaScript; otherwise, false.</returns>
+    </member>
+    <member name="P:System.Web.WebPages.Html.HtmlHelper.ValidationInputCssClassName">
+      <summary>Gets or sets the name of the CSS class that defines the appearance of input elements when validation fails.</summary>
+      <returns>The name of the CSS class. The default is field-validation-error.</returns>
+    </member>
+    <member name="P:System.Web.WebPages.Html.HtmlHelper.ValidationInputValidCssClassName">
+      <summary>Gets or sets the name of the CSS class that defines the appearance of input elements when validation passes.</summary>
+      <returns>The name of the CSS class. The default is input-validation-valid.</returns>
+    </member>
+    <member name="M:System.Web.WebPages.Html.HtmlHelper.ValidationMessage(System.String)">
+      <summary>Returns an HTML span element that contains the first validation error message for the specified form field.</summary>
+      <returns>If the value in the specified field is valid, null; otherwise, the HTML markup that represents the validation error message that is associated with the specified field.</returns>
+      <param name="name">The name of the form field that was validated.</param>
+      <exception cref="T:Sytem.ArgumentException">
+        <paramref name="name" /> is null or empty.</exception>
+    </member>
+    <member name="M:System.Web.WebPages.Html.HtmlHelper.ValidationMessage(System.String,System.Collections.Generic.IDictionary{System.String,System.Object})">
+      <summary>Returns an HTML span element that has the specified custom attributes defined by an attribute dictionary, and that contains the first validation error message for the specified form field.</summary>
+      <returns>If the value in the specified field is valid, null; otherwise, the HTML markup that represents the validation error message that is associated with the specified field.</returns>
+      <param name="name">The name of the form field that was validated.</param>
+      <param name="htmlAttributes">The names and values of custom attributes for the element.</param>
+      <exception cref="T:Sytem.ArgumentException">
+        <paramref name="name" /> is null or empty.</exception>
+    </member>
+    <member name="M:System.Web.WebPages.Html.HtmlHelper.ValidationMessage(System.String,System.Object)">
+      <summary>Returns an HTML span element that has the specified custom attributes defined by an attribute object, and that contains the first validation error message for the specified form field.</summary>
+      <returns>If the value in the specified field is valid, null; otherwise, the HTML markup that represents the validation error message that is associated with the specified field.</returns>
+      <param name="name">The name of the form field that was validated.</param>
+      <param name="htmlAttributes">An object that contains custom attributes for the element. The attribute names and values are retrieved through reflection by examining the properties of the object.</param>
+      <exception cref="T:Sytem.ArgumentException">
+        <paramref name="name" /> is null or empty.</exception>
+    </member>
+    <member name="M:System.Web.WebPages.Html.HtmlHelper.ValidationMessage(System.String,System.String)">
+      <summary>Returns an HTML span element that contains a validation error message for the specified form field.</summary>
+      <returns>If the value in the specified field is valid, null; otherwise, the HTML markup that represents the validation error message that is associated with the specified field.</returns>
+      <param name="name">The name of the form field that was validated.</param>
+      <param name="message">The validation error message to display. If null, the first validation error message that is associated with the specified form field is displayed.</param>
+      <exception cref="T:Sytem.ArgumentException">
+        <paramref name="name" /> is null or empty.</exception>
+    </member>
+    <member name="M:System.Web.WebPages.Html.HtmlHelper.ValidationMessage(System.String,System.String,System.Collections.Generic.IDictionary{System.String,System.Object})">
+      <summary>Returns an HTML span element that has the specified custom attributes defined by an attribute dictionary, and that contains a validation error message for the specified form field.</summary>
+      <returns>If the specified field is valid, null; otherwise, the HTML markup that represents a validation error message that is associated with the specified field.</returns>
+      <param name="name">The name of the form field that was validated.</param>
+      <param name="message">The validation error message to display. If null, the first validation error message that is associated with the specified form field is displayed.</param>
+      <param name="htmlAttributes"> The names and values of custom attributes for the element.</param>
+      <exception cref="T:Sytem.ArgumentException">
+        <paramref name="name" /> is null or empty.</exception>
+    </member>
+    <member name="M:System.Web.WebPages.Html.HtmlHelper.ValidationMessage(System.String,System.String,System.Object)">
+      <summary>Returns an HTML span element that has the specified custom attributes defined by an attribute object, and that contains a validation error message for the specified form field.</summary>
+      <returns>If the specified field is valid, null; otherwise, the HTML markup that represents a validation error message that is associated with the specified field.</returns>
+      <param name="name">The name of the form field that was validated.</param>
+      <param name="message">The validation error message to display. If null, the first validation error message that is associated with the specified form field is displayed.</param>
+      <param name="htmlAttributes">An object that contains custom attributes for the element. The attribute names and values are retrieved through reflection by examining the properties of the object.</param>
+      <exception cref="T:Sytem.ArgumentException">
+        <paramref name="name" /> is null or empty.</exception>
+    </member>
+    <member name="P:System.Web.WebPages.Html.HtmlHelper.ValidationMessageCssClassName">
+      <summary>Gets or sets the name of the CSS class that defines the appearance of validation error messages when validation fails.</summary>
+      <returns>The name of the CSS class.  The default is field-validation-error.</returns>
+    </member>
+    <member name="P:System.Web.WebPages.Html.HtmlHelper.ValidationMessageValidCssClassName">
+      <summary>Gets or sets the name of the CSS class that defines the appearance of validation error messages when validation passes.</summary>
+      <returns>The name of the CSS class. The default is field-validation-valid.</returns>
+    </member>
+    <member name="M:System.Web.WebPages.Html.HtmlHelper.ValidationSummary">
+      <summary>Returns an HTML div element that contains an unordered list of all validation error messages from the model-state dictionary.</summary>
+      <returns>The HTML markup that represents the validation error messages.</returns>
+    </member>
+    <member name="M:System.Web.WebPages.Html.HtmlHelper.ValidationSummary(System.Boolean)">
+      <summary>Returns an HTML div element that contains an unordered list of validation error message from the model-state dictionary, optionally excluding field-level errors.</summary>
+      <returns>The HTML markup that represents the validation error messages.</returns>
+      <param name="excludeFieldErrors">true to exclude field-level validation error messages from the list; false to include both model-level and field-level validation error messages.</param>
+    </member>
+    <member name="M:System.Web.WebPages.Html.HtmlHelper.ValidationSummary(System.Collections.Generic.IDictionary{System.String,System.Object})">
+      <summary>Returns an HTML div element that has the specified custom attributes defined by an attribute dictionary, and that contains an unordered list of all validation error messages that are in the model-state dictionary.</summary>
+      <returns>The HTML markup that represents the validation error messages.</returns>
+      <param name="htmlAttributes">The names and values of custom attributes for the element.</param>
+    </member>
+    <member name="M:System.Web.WebPages.Html.HtmlHelper.ValidationSummary(System.Object)">
+      <summary>Returns an HTML div element that has the specified custom attributes defined by an attribute object, and that contains an unordered list of all validation error messages that are in the model-state dictionary.</summary>
+      <returns>The HTML markup that represents the validation error messages.</returns>
+      <param name="htmlAttributes">An object that contains custom attributes for the element. The attribute names and values are retrieved through reflection by examining the properties of the object.</param>
+    </member>
+    <member name="M:System.Web.WebPages.Html.HtmlHelper.ValidationSummary(System.String)">
+      <summary>Returns an HTML div element that contains a summary message and an unordered list of all validation error messages that are in the model-state dictionary.</summary>
+      <returns>The HTML markup that represents the validation error messages.</returns>
+      <param name="message">The message that comes before the list of validation error messages.</param>
+    </member>
+    <member name="M:System.Web.WebPages.Html.HtmlHelper.ValidationSummary(System.String,System.Boolean,System.Collections.Generic.IDictionary{System.String,System.Object})">
+      <summary>Returns an HTML div element that has the specified custom attributes defined by an attribute dictionary, and that contains a summary message and an unordered list of validation error message from the model-state dictionary, optionally excluding field-level errors.</summary>
+      <returns>The HTML markup that represents the validation error messages.</returns>
+      <param name="message">The summary message that comes before the list of validation error messages.</param>
+      <param name="excludeFieldErrors">true to exclude field-level validation error messages from the results; false to include both model-level and field-level validation error messages.</param>
+      <param name="htmlAttributes">The names and values of custom attributes for the element.</param>
+    </member>
+    <member name="M:System.Web.WebPages.Html.HtmlHelper.ValidationSummary(System.String,System.Boolean,System.Object)">
+      <summary>Returns an HTML div element that has the specified custom attributes defined by an attribute object, and that contains a summary message and an unordered list of validation error message from the model-state dictionary, optionally excluding field-level errors.</summary>
+      <returns>The HTML markup that represents the validation error messages.</returns>
+      <param name="message">The summary message that comes before the list of validation error messages.</param>
+      <param name="excludeFieldErrors">true to exclude field-level validation error messages from the results; false to include and field-level validation error messages.</param>
+      <param name="htmlAttributes">An object that contains custom attributes for the element. The attribute names and values are retrieved through reflection by examining the properties of the object.</param>
+    </member>
+    <member name="M:System.Web.WebPages.Html.HtmlHelper.ValidationSummary(System.String,System.Collections.Generic.IDictionary{System.String,System.Object})">
+      <summary>Returns an HTML div element that has the specified custom attributes defined by an attribute dictionary, and that contains a summary message and an unordered list of all validation error message from the model-state dictionary.</summary>
+      <returns>The HTML markup that represents the validation error messages.</returns>
+      <param name="message">The message that comes before the list of validation error messages.</param>
+      <param name="htmlAttributes">The names and values of custom attributes for the element.</param>
+    </member>
+    <member name="M:System.Web.WebPages.Html.HtmlHelper.ValidationSummary(System.String,System.Object)">
+      <summary>Returns an HTML div element that has the specified custom attributes defined by an attribute object, and that contains a summary message and an unordered list of all validation error message from the model-state dictionary.</summary>
+      <returns>The HTML markup that represents the validation error messages.</returns>
+      <param name="message">The summary message that comes before the list of validation error messages.</param>
+      <param name="htmlAttributes">An object that contains custom attributes for the element. The attribute names and values are retrieved through reflection by examining the properties of the object.</param>
+    </member>
+    <member name="P:System.Web.WebPages.Html.HtmlHelper.ValidationSummaryClass">
+      <summary>Gets or sets the name of the CSS class that defines the appearance of a validation summary when validation fails.</summary>
+      <returns>The name of the CSS class. The default is validation-summary-errors.</returns>
+    </member>
+    <member name="P:System.Web.WebPages.Html.HtmlHelper.ValidationSummaryValidClass">
+      <summary>Gets or sets the name of the CSS class that defines the appearance of a validation summary when validation passes.</summary>
+      <returns>The name of the CSS class. The default is validation-summary-valid.</returns>
+    </member>
+    <member name="T:System.Web.WebPages.Html.ModelState">
+      <summary>Encapsulates the state of model binding to a property of an action-method argument, or to the argument itself.</summary>
+    </member>
+    <member name="M:System.Web.WebPages.Html.ModelState.#ctor">
+      <summary>Initializes a new instance of the <see cref="T:System.Web.WebPages.Html.ModelState" /> class.</summary>
+    </member>
+    <member name="P:System.Web.WebPages.Html.ModelState.Errors">
+      <summary>Returns a list of strings that contains any errors that occurred during model binding.</summary>
+      <returns>The errors that occurred during model binding.</returns>
+    </member>
+    <member name="P:System.Web.WebPages.Html.ModelState.Value">
+      <summary>Returns an object that encapsulates the value that was bound during model binding.</summary>
+      <returns>The value that was bound.</returns>
+    </member>
+    <member name="T:System.Web.WebPages.Html.ModelStateDictionary">
+      <summary>Represents the result of binding a posted form to an action method, which includes information such as validation status and validation error messages.</summary>
+    </member>
+    <member name="M:System.Web.WebPages.Html.ModelStateDictionary.#ctor">
+      <summary>Initializes a new instance of the <see cref="T:System.Web.WebPages.Html.ModelStateDictionary" /> class.</summary>
+    </member>
+    <member name="M:System.Web.WebPages.Html.ModelStateDictionary.#ctor(System.Web.WebPages.Html.ModelStateDictionary)">
+      <summary>Initializes a new instance of the <see cref="T:System.Web.WebPages.Html.ModelStateDictionary" /> class by using values that are copied from the specified model-state dictionary.</summary>
+      <param name="dictionary">The model-state dictionary that values are copied from.</param>
+    </member>
+    <member name="M:System.Web.WebPages.Html.ModelStateDictionary.Add(System.Collections.Generic.KeyValuePair{System.String,System.Web.WebPages.Html.ModelState})">
+      <summary>Adds the specified item to the model-state dictionary.</summary>
+      <param name="item">The item to add to the model-state dictionary.</param>
+    </member>
+    <member name="M:System.Web.WebPages.Html.ModelStateDictionary.Add(System.String,System.Web.WebPages.Html.ModelState)">
+      <summary>Adds an item that has the specified key and value to the model-state dictionary.</summary>
+      <param name="key">The key.</param>
+      <param name="value">The value.</param>
+    </member>
+    <member name="M:System.Web.WebPages.Html.ModelStateDictionary.AddError(System.String,System.String)">
+      <summary>Adds an error message to the model state that is associated with the specified key.</summary>
+      <param name="key">The key that is associated with the model state that the error message is added to.</param>
+      <param name="errorMessage">The error message.</param>
+    </member>
+    <member name="M:System.Web.WebPages.Html.ModelStateDictionary.AddFormError(System.String)">
+      <summary>Adds an error message to the model state that is associated with the entire form.</summary>
+      <param name="errorMessage">The error message.</param>
+    </member>
+    <member name="M:System.Web.WebPages.Html.ModelStateDictionary.Clear">
+      <summary>Removes all items from the model-state dictionary.</summary>
+    </member>
+    <member name="M:System.Web.WebPages.Html.ModelStateDictionary.Contains(System.Collections.Generic.KeyValuePair{System.String,System.Web.WebPages.Html.ModelState})">
+      <summary>Determines whether the model-state dictionary contains the specified item.</summary>
+      <returns>true if the model-state dictionary contains the specified item; otherwise, false.</returns>
+      <param name="item">The item to look for.</param>
+    </member>
+    <member name="M:System.Web.WebPages.Html.ModelStateDictionary.ContainsKey(System.String)">
+      <summary>Determines whether the model-state dictionary contains the specified key.</summary>
+      <returns>true if the model-state dictionary contains the specified key; otherwise, false.</returns>
+      <param name="key">The key to look for.</param>
+    </member>
+    <member name="M:System.Web.WebPages.Html.ModelStateDictionary.CopyTo(System.Collections.Generic.KeyValuePair{System.String,System.Web.WebPages.Html.ModelState}[],System.Int32)">
+      <summary>Copies the elements of the model-state dictionary to an array, starting at the specified index.</summary>
+      <param name="array">The one-dimensional <see cref="T:System.Array" /> instance where the elements will be copied to.</param>
+      <param name="arrayIndex">The index in <paramref name="array" /> at which copying begins.</param>
+    </member>
+    <member name="P:System.Web.WebPages.Html.ModelStateDictionary.Count">
+      <summary>Gets the number of model states that the model-state dictionary contains.</summary>
+      <returns>The number of model states in the model-state dictionary.</returns>
+    </member>
+    <member name="M:System.Web.WebPages.Html.ModelStateDictionary.GetEnumerator">
+      <summary>Returns an enumerator that can be used to iterate through the collection.</summary>
+      <returns>An enumerator that can be used to iterate through the collection.</returns>
+    </member>
+    <member name="P:System.Web.WebPages.Html.ModelStateDictionary.IsReadOnly">
+      <summary>Gets a value that indicates whether the model-state dictionary is read-only.</summary>
+      <returns>true if the model-state dictionary is read-only; otherwise, false.</returns>
+    </member>
+    <member name="P:System.Web.WebPages.Html.ModelStateDictionary.IsValid">
+      <summary>Gets a value that indicates whether any error messages are associated with any model state in the model-state dictionary.</summary>
+      <returns>true if any error messages are associated with any model state in the dictionary; otherwise, false.</returns>
+    </member>
+    <member name="M:System.Web.WebPages.Html.ModelStateDictionary.IsValidField(System.String)">
+      <summary>Determines whether any error messages are associated with the specified key.</summary>
+      <returns>true if no error messages are associated with the specified key, or the specified key does not exist; otherwise, false.</returns>
+      <param name="key">The key.</param>
+      <exception cref="T:System.ArgumentNullException">
+        <paramref name="key" /> is null.</exception>
+    </member>
+    <member name="P:System.Web.WebPages.Html.ModelStateDictionary.Item(System.String)">
+      <summary>Gets or sets the model state that is associated with the specified key in the model-state dictionary.</summary>
+      <returns>The model state that is associated with the specified key in the dictionary.</returns>
+    </member>
+    <member name="P:System.Web.WebPages.Html.ModelStateDictionary.Keys">
+      <summary>Gets a list that contains the keys in the model-state dictionary.</summary>
+      <returns>The list of keys in the dictionary.</returns>
+    </member>
+    <member name="M:System.Web.WebPages.Html.ModelStateDictionary.Merge(System.Web.WebPages.Html.ModelStateDictionary)">
+      <summary>Copies the values from the specified model-state dictionary into this <see cref="T:System.Web.WebPages.Html.ModelStateDictionary" /> instance, overwriting existing values when the keys are the same.</summary>
+      <param name="dictionary">The model-state dictionary that values are copied from.</param>
+    </member>
+    <member name="M:System.Web.WebPages.Html.ModelStateDictionary.Remove(System.Collections.Generic.KeyValuePair{System.String,System.Web.WebPages.Html.ModelState})">
+      <summary>Removes the first occurrence of the specified item from the model-state dictionary.</summary>
+      <returns>true if the item was successfully removed from the model-state dictionary; false if the item was not removed or if the item does not exist in the model-state dictionary.</returns>
+      <param name="item">The item to remove.</param>
+    </member>
+    <member name="M:System.Web.WebPages.Html.ModelStateDictionary.Remove(System.String)">
+      <summary>Removes the item that has the specified key from the model-state dictionary.</summary>
+      <returns>true if the item was successfully removed from the model-state dictionary; false if the item was not removed or does not exist in the model-state dictionary.</returns>
+      <param name="key">The key of the element to remove.</param>
+    </member>
+    <member name="M:System.Web.WebPages.Html.ModelStateDictionary.SetModelValue(System.String,System.Object)">
+      <summary>Sets the value of the model state that is associated with the specified key.</summary>
+      <param name="key">The key to set the value of.</param>
+      <param name="value">The value to set the key to.</param>
+    </member>
+    <member name="M:System.Web.WebPages.Html.ModelStateDictionary.System#Collections#IEnumerable#GetEnumerator">
+      <summary>Returns an enumerator that can be used to iterate through the model-state dictionary.</summary>
+      <returns>An enumerator that can be used to iterate through the model-state dictionary.</returns>
+    </member>
+    <member name="M:System.Web.WebPages.Html.ModelStateDictionary.TryGetValue(System.String,System.Web.WebPages.Html.ModelState@)">
+      <summary>Gets the model-state value that is associated with the specified key.</summary>
+      <returns>true if the model-state dictionary contains an element that has the specified key; otherwise, false.</returns>
+      <param name="key">The key to get the value of.</param>
+      <param name="value">When this method returns, if the key is found, contains the model-state value that is associated with the specified key; otherwise, contains the default value for the <see cref="T:System.Web.WebPages.Html.ModelState" /> type. This parameter is passed uninitialized.</param>
+    </member>
+    <member name="P:System.Web.WebPages.Html.ModelStateDictionary.Values">
+      <summary>Gets a list that contains the values in the model-state dictionary.</summary>
+      <returns>The list of values in the dictionary.</returns>
+    </member>
+    <member name="T:System.Web.WebPages.Html.SelectListItem">
+      <summary>Represents an item in an HTML select list.</summary>
+    </member>
+    <member name="M:System.Web.WebPages.Html.SelectListItem.#ctor">
+      <summary>Initializes a new instance of the <see cref="T:System.Web.WebPages.Html.SelectListItem" /> class using the default settings.</summary>
+    </member>
+    <member name="M:System.Web.WebPages.Html.SelectListItem.#ctor(System.Web.WebPages.Html.SelectListItem)">
+      <summary>Initializes a new instance of the <see cref="T:System.Web.WebPages.Html.SelectListItem" /> class by copying the specified select list item.</summary>
+      <param name="item">The select list item to copy.</param>
+    </member>
+    <member name="P:System.Web.WebPages.Html.SelectListItem.Selected">
+      <summary>Gets or sets a value that indicates whether the <see cref="T:System.Web.WebPages.Html.SelectListItem" /> instance is selected.</summary>
+      <returns>true if the select list item is selected; otherwise, false.</returns>
+    </member>
+    <member name="P:System.Web.WebPages.Html.SelectListItem.Text">
+      <summary>Gets or sets the text that is used to display the <see cref="T:System.Web.WebPages.Html.SelectListItem" /> instance on a web page.</summary>
+      <returns>The text that is used to display the select list item.</returns>
+    </member>
+    <member name="P:System.Web.WebPages.Html.SelectListItem.Value">
+      <summary>Gets or sets the value of the HTML value attribute of the HTML option element that is associated with the <see cref="T:System.Web.WebPages.Html.SelectListItem" /> instance.</summary>
+      <returns>The value of the HTML value attribute that is associated with the select list item.</returns>
+    </member>
+    <member name="T:System.Web.WebPages.Instrumentation.InstrumentationService">
+      <summary>This type/member supports the .NET Framework infrastructure and is not intended to be used directly from your code.Represents a web pages instrumentation service.</summary>
+    </member>
+    <member name="M:System.Web.WebPages.Instrumentation.InstrumentationService.#ctor">
+      <summary>This type/member supports the .NET Framework infrastructure and is not intended to be used directly from your code.Initializes a new instance of the <see cref="T:System.Web.WebPages.Instrumentation.InstrumentationService" /> class.</summary>
+    </member>
+    <member name="M:System.Web.WebPages.Instrumentation.InstrumentationService.BeginContext(System.Web.HttpContextBase,System.String,System.IO.TextWriter,System.Int32,System.Int32,System.Boolean)">
+      <summary>This type/member supports the .NET Framework infrastructure and is not intended to be used directly from your code.Called before it renders the output for the specified context.</summary>
+      <param name="context">The context.</param>
+      <param name="virtualPath">The virtual path.</param>
+      <param name="writer">The writer.</param>
+      <param name="startPosition">The start position.</param>
+      <param name="length">The length of the context.</param>
+      <param name="isLiteral">Determines whether the context is literal.</param>
+    </member>
+    <member name="M:System.Web.WebPages.Instrumentation.InstrumentationService.EndContext(System.Web.HttpContextBase,System.String,System.IO.TextWriter,System.Int32,System.Int32,System.Boolean)">
+      <summary>This type/member supports the .NET Framework infrastructure and is not intended to be used directly from your code.Called after it renders the output for the specified context.</summary>
+      <param name="context">The context.</param>
+      <param name="virtualPath">The virtual path.</param>
+      <param name="writer">The writer.</param>
+      <param name="startPosition">The start position.</param>
+      <param name="length">The length of the context.</param>
+      <param name="isLiteral">Determines whether the context is literal.</param>
+    </member>
+    <member name="P:System.Web.WebPages.Instrumentation.InstrumentationService.IsAvailable">
+      <summary>This type/member supports the .NET Framework infrastructure and is not intended to be used directly from your code.Gets a value indicating whether the service is available.</summary>
+      <returns>true if the service is available; otherwise, false.</returns>
+    </member>
+    <member name="T:System.Web.WebPages.Instrumentation.PositionTagged`1">
+      <summary>This type/member supports the .NET Framework infrastructure and is not intended to be used directly from your code.Represents a position tagged.</summary>
+      <typeparam name="T">The type of the position..</typeparam>
+    </member>
+    <member name="M:System.Web.WebPages.Instrumentation.PositionTagged`1.#ctor(`0,System.Int32)">
+      <summary>This type/member supports the .NET Framework infrastructure and is not intended to be used directly from your code.Initializes a new instance of the <see cref="T:System.Web.WebPages.Instrumentation.PositionTagged`1" /> class.</summary>
+      <param name="value">The value of this current instance.</param>
+      <param name="offset">The offset.</param>
+    </member>
+    <member name="M:System.Web.WebPages.Instrumentation.PositionTagged`1.Equals(System.Object)">
+      <summary>This type/member supports the .NET Framework infrastructure and is not intended to be used directly from your code.Determines whether the specified object is equal to the current object.</summary>
+      <returns>true if the specified object is equal to the current object; otherwise, false.</returns>
+      <param name="obj">The object to compare to.</param>
+    </member>
+    <member name="M:System.Web.WebPages.Instrumentation.PositionTagged`1.GetHashCode">
+      <summary>This type/member supports the .NET Framework infrastructure and is not intended to be used directly from your code.Gets the hash code of the current instance.</summary>
+      <returns>The hash code of the current instance.</returns>
+    </member>
+    <member name="M:System.Web.WebPages.Instrumentation.PositionTagged`1.op_Equality(System.Web.WebPages.Instrumentation.PositionTagged{`0},System.Web.WebPages.Instrumentation.PositionTagged{`0})">
+      <summary>This type/member supports the .NET Framework infrastructure and is not intended to be used directly from your code.Determines whether the two objects are equal.</summary>
+      <returns>true if the two objects are equal; otherwise, false.</returns>
+      <param name="left">The first object.</param>
+      <param name="right">The second object.</param>
+    </member>
+    <member name="M:System.Web.WebPages.Instrumentation.PositionTagged`1.op_Implicit(System.Tuple{`0,System.Int32})~System.Web.WebPages.Instrumentation.PositionTagged{`0}">
+      <summary>Converts the specified object to a <see cref="T:System.Web.WebPages.Instrumentation.PositionTagged`1" /> object.</summary>
+      <returns>The <see cref="T:System.Web.WebPages.Instrumentation.PositionTagged`1" /> that represents the converted <paramref name="value" />.</returns>
+      <param name="value">The object to convert.</param>
+    </member>
+    <member name="M:System.Web.WebPages.Instrumentation.PositionTagged`1.op_Implicit(System.Web.WebPages.Instrumentation.PositionTagged{`0})~`0">
+      <summary>This type/member supports the .NET Framework infrastructure and is not intended to be used directly from your code.Converts the <paramref name="value" /> to a <see cref="T:System.Web.WebPages.Instrumentation.PositionTagged`1" /> object.</summary>
+      <returns>The <see cref="T:System.Web.WebPages.Instrumentation.PositionTagged`1" /> that represents the converted <paramref name="value" />.</returns>
+      <param name="value">The object to convert.</param>
+    </member>
+    <member name="M:System.Web.WebPages.Instrumentation.PositionTagged`1.op_Inequality(System.Web.WebPages.Instrumentation.PositionTagged{`0},System.Web.WebPages.Instrumentation.PositionTagged{`0})">
+      <summary>This type/member supports the .NET Framework infrastructure and is not intended to be used directly from your code.Determines whether the two objects are not equal.</summary>
+      <returns>true if the two objects are not equal; otherwise, false.</returns>
+      <param name="left">The first object.</param>
+      <param name="right">The second object.</param>
+    </member>
+    <member name="P:System.Web.WebPages.Instrumentation.PositionTagged`1.Position">
+      <summary>This type/member supports the .NET Framework infrastructure and is not intended to be used directly from your code.Gets or sets the position associated with the <see cref="T:System.Web.WebPages.Instrumentation.PositionTagged`1" />.</summary>
+      <returns>The position associated with the <see cref="T:System.Web.WebPages.Instrumentation.PositionTagged`1" />.</returns>
+    </member>
+    <member name="M:System.Web.WebPages.Instrumentation.PositionTagged`1.ToString">
+      <summary>This type/member supports the .NET Framework infrastructure and is not intended to be used directly from your code.Returns a string representation of the <see cref="T:System.Web.WebPages.Instrumentation.PositionTagged`1" /> object.</summary>
+      <returns>A string that represents the <see cref="T:System.Web.WebPages.Instrumentation.PositionTagged`1" /> object.</returns>
+    </member>
+    <member name="P:System.Web.WebPages.Instrumentation.PositionTagged`1.Value">
+      <summary>This type/member supports the .NET Framework infrastructure and is not intended to be used directly from your code.Gets or sets the value of the current instance.</summary>
+      <returns>The value of the current instance.</returns>
+    </member>
+    <member name="T:System.Web.WebPages.Scope.AspNetRequestScopeStorageProvider">
+      <summary>Defines an ASP.NET request scope storage provider.</summary>
+    </member>
+    <member name="M:System.Web.WebPages.Scope.AspNetRequestScopeStorageProvider.#ctor">
+      <summary>Initializes a new instance of the <see cref="T:System.Web.WebPages.Scope.AspNetRequestScopeStorageProvider" /> class.</summary>
+    </member>
+    <member name="P:System.Web.WebPages.Scope.AspNetRequestScopeStorageProvider.ApplicationScope">
+      <summary>Gets the dictionary to store data in the application scope.</summary>
+      <returns>The dictionary that stores application scope data.</returns>
+    </member>
+    <member name="P:System.Web.WebPages.Scope.AspNetRequestScopeStorageProvider.CurrentScope">
+      <summary>Gets or sets the dictionary to store data in the current scope.</summary>
+      <returns>The dictionary that stores current scope data.</returns>
+    </member>
+    <member name="P:System.Web.WebPages.Scope.AspNetRequestScopeStorageProvider.GlobalScope">
+      <summary>Gets the dictionary to store data in the global scope.</summary>
+      <returns>The dictionary that stores global scope data.</returns>
+    </member>
+    <member name="P:System.Web.WebPages.Scope.AspNetRequestScopeStorageProvider.RequestScope">
+      <summary>Gets the dictionary to store data in the request scope.</summary>
+      <returns>The dictionary that stores request scope data.</returns>
+    </member>
+    <member name="T:System.Web.WebPages.Scope.IScopeStorageProvider">
+      <summary>Defines a dictionary that provides scoped access to data.</summary>
+    </member>
+    <member name="P:System.Web.WebPages.Scope.IScopeStorageProvider.CurrentScope">
+      <summary>Gets and sets the dictionary that is used to store data in the current scope.</summary>
+      <returns>The dictionary that stores current scope data.</returns>
+    </member>
+    <member name="P:System.Web.WebPages.Scope.IScopeStorageProvider.GlobalScope">
+      <summary>Gets the dictionary that is used to store data in the global scope.</summary>
+      <returns>The dictionary that stores global scope data.</returns>
+    </member>
+    <member name="T:System.Web.WebPages.Scope.ScopeStorage">
+      <summary>Defines a class that is used to contain storage for a transient scope.</summary>
+    </member>
+    <member name="M:System.Web.WebPages.Scope.ScopeStorage.CreateTransientScope">
+      <summary>Returns a dictionary that is used to store data in a transient scope, based on the scope in the <see cref="P:System.Web.WebPages.Scope.ScopeStorage.CurrentScope" /> property.</summary>
+      <returns>The dictionary that stores transient scope data.</returns>
+    </member>
+    <member name="M:System.Web.WebPages.Scope.ScopeStorage.CreateTransientScope(System.Collections.Generic.IDictionary{System.Object,System.Object})">
+      <summary>Returns a dictionary that is used to store data in a transient scope.</summary>
+      <returns>The dictionary that stores transient scope data.</returns>
+      <param name="context">The context.</param>
+    </member>
+    <member name="P:System.Web.WebPages.Scope.ScopeStorage.CurrentProvider">
+      <summary>Gets or sets the current scope provider.</summary>
+      <returns>The current scope provider.</returns>
+    </member>
+    <member name="P:System.Web.WebPages.Scope.ScopeStorage.CurrentScope">
+      <summary>Gets the dictionary that is used to store data in the current scope.</summary>
+      <returns>The dictionary that stores current scope data.</returns>
+    </member>
+    <member name="P:System.Web.WebPages.Scope.ScopeStorage.GlobalScope">
+      <summary>Gets the dictionary that is used to store data in the global scope.</summary>
+      <returns>The dictionary that stores global scope data.</returns>
+    </member>
+    <member name="T:System.Web.WebPages.Scope.ScopeStorageDictionary">
+      <summary>Represents a collection of keys and values that are used to store data at different scope levels (local, global, and so on).</summary>
+    </member>
+    <member name="M:System.Web.WebPages.Scope.ScopeStorageDictionary.#ctor">
+      <summary>Initializes a new instance of the <see cref="T:System.Web.WebPages.Scope.ScopeStorageDictionary" /> class.</summary>
+    </member>
+    <member name="M:System.Web.WebPages.Scope.ScopeStorageDictionary.#ctor(System.Collections.Generic.IDictionary{System.Object,System.Object})">
+      <summary>Initializes a new instance of the <see cref="T:System.Web.WebPages.Scope.ScopeStorageDictionary" /> class using the specified base scope.</summary>
+      <param name="baseScope">The base scope.</param>
+    </member>
+    <member name="M:System.Web.WebPages.Scope.ScopeStorageDictionary.Add(System.Collections.Generic.KeyValuePair{System.Object,System.Object})">
+      <summary>Adds a key/value pair to the <see cref="T:System.Web.WebPages.Scope.ScopeStorageDictionary" /> object using the specified generic collection.</summary>
+      <param name="item">The key/value pair.</param>
+    </member>
+    <member name="M:System.Web.WebPages.Scope.ScopeStorageDictionary.Add(System.Object,System.Object)">
+      <summary>Adds the specified key and specified value to the <see cref="T:System.Web.WebPages.Scope.ScopeStorageDictionary" /> object.</summary>
+      <param name="key">The key.</param>
+      <param name="value">The value.</param>
+    </member>
+    <member name="P:System.Web.WebPages.Scope.ScopeStorageDictionary.BackingStore">
+      <summary>Gets the dictionary that stores the <see cref="T:System.Web.WebPages.Scope.ScopeStorageDictionary" /> object data.</summary>
+    </member>
+    <member name="P:System.Web.WebPages.Scope.ScopeStorageDictionary.BaseScope">
+      <summary>Gets the base scope for the <see cref="T:System.Web.WebPages.Scope.ScopeStorageDictionary" /> object.</summary>
+      <returns>The base scope for the <see cref="T:System.Web.WebPages.Scope.ScopeStorageDictionary" /> object.</returns>
+    </member>
+    <member name="M:System.Web.WebPages.Scope.ScopeStorageDictionary.Clear">
+      <summary>Removes all keys and values from the concatenated <see cref="T:System.Web.WebPages.Scope.ScopeStorageDictionary" /> and <see cref="P:System.Web.WebPages.Scope.ScopeStorageDictionary.BaseScope" /> objects.</summary>
+    </member>
+    <member name="M:System.Web.WebPages.Scope.ScopeStorageDictionary.Contains(System.Collections.Generic.KeyValuePair{System.Object,System.Object})">
+      <summary>Returns a value that indicates whether the specified key/value pair exists in either the <see cref="T:System.Web.WebPages.Scope.ScopeStorageDictionary" /> object or in the <see cref="P:System.Web.WebPages.Scope.ScopeStorageDictionary.BaseScope" /> object.</summary>
+      <returns>true if the <see cref="T:System.Web.WebPages.Scope.ScopeStorageDictionary" /> object or the <see cref="P:System.Web.WebPages.Scope.ScopeStorageDictionary.BaseScope" /> object contains an element that has the specified key/value pair; otherwise, false.</returns>
+      <param name="item">The key/value pair.</param>
+    </member>
+    <member name="M:System.Web.WebPages.Scope.ScopeStorageDictionary.ContainsKey(System.Object)">
+      <summary>Returns a value that indicates whether the specified key exists in the <see cref="T:System.Web.WebPages.Scope.ScopeStorageDictionary" /> object or in the <see cref="P:System.Web.WebPages.Scope.ScopeStorageDictionary.BaseScope" /> object.</summary>
+      <returns>true if the <see cref="T:System.Web.WebPages.Scope.ScopeStorageDictionary" /> object or the <see cref="P:System.Web.WebPages.Scope.ScopeStorageDictionary.BaseScope" /> object contains an element that has the specified key; otherwise, false.</returns>
+      <param name="key">The key.</param>
+    </member>
+    <member name="M:System.Web.WebPages.Scope.ScopeStorageDictionary.CopyTo(System.Collections.Generic.KeyValuePair{System.Object,System.Object}[],System.Int32)">
+      <summary>Copies all of the elements in the <see cref="T:System.Web.WebPages.Scope.ScopeStorageDictionary" /> object and the <see cref="P:System.Web.WebPages.Scope.ScopeStorageDictionary.BaseScope" /> object to an <see cref="T:System.Array" /> object, starting at the specified index.</summary>
+      <param name="array">The array.</param>
+      <param name="arrayIndex">The zero-based index in <paramref name="array" />.</param>
+    </member>
+    <member name="P:System.Web.WebPages.Scope.ScopeStorageDictionary.Count">
+      <summary>Gets the number of key/value pairs that are in the concatenated <see cref="T:System.Web.WebPages.Scope.ScopeStorageDictionary" /> and <see cref="P:System.Web.WebPages.Scope.ScopeStorageDictionary.BaseScope" /> objects.</summary>
+      <returns>The number of key/value pairs.</returns>
+    </member>
+    <member name="M:System.Web.WebPages.Scope.ScopeStorageDictionary.GetEnumerator">
+      <summary>Returns an enumerator that can be used to iterate through concatenated <see cref="T:System.Web.WebPages.Scope.ScopeStorageDictionary" /> and <see cref="P:System.Web.WebPages.Scope.ScopeStorageDictionary.BaseScope" /> objects.</summary>
+      <returns>An <see cref="T:System.Collections.Generic.IEnumerator" /> object.</returns>
+    </member>
+    <member name="M:System.Web.WebPages.Scope.ScopeStorageDictionary.GetItems">
+      <summary>Returns an enumerator that can be used to iterate through the distinct elements of concatenated <see cref="T:System.Web.WebPages.Scope.ScopeStorageDictionary" /> and <see cref="P:System.Web.WebPages.Scope.ScopeStorageDictionary.BaseScope" /> objects.</summary>
+      <returns>An enumerator that contains distinct elements from the concatenated dictionary objects.</returns>
+    </member>
+    <member name="P:System.Web.WebPages.Scope.ScopeStorageDictionary.IsReadOnly">
+      <summary>Gets a value that indicates whether the <see cref="T:System.Web.WebPages.Scope.ScopeStorageDictionary" /> object is read-only.</summary>
+      <returns>true if the <see cref="T:System.Web.WebPages.Scope.ScopeStorageDictionary" /> object is read-only; otherwise, false.</returns>
+    </member>
+    <member name="P:System.Web.WebPages.Scope.ScopeStorageDictionary.Item(System.Object)">
+      <summary>Gets or sets the element that is associated with the specified key.</summary>
+      <returns>The element that has the specified key.</returns>
+    </member>
+    <member name="P:System.Web.WebPages.Scope.ScopeStorageDictionary.Keys">
+      <summary>Gets a <see cref="T:System.Collections.Generic.List`1" /> object that contains the keys from the concatenated <see cref="T:System.Web.WebPages.Scope.ScopeStorageDictionary" /> and <see cref="P:System.Web.WebPages.Scope.ScopeStorageDictionary.BaseScope" /> objects.</summary>
+      <returns>An object that contains that contains the keys.</returns>
+    </member>
+    <member name="M:System.Web.WebPages.Scope.ScopeStorageDictionary.Remove(System.Collections.Generic.KeyValuePair{System.Object,System.Object})">
+      <summary>Removes the specified key/value pair from the concatenated <see cref="T:System.Web.WebPages.Scope.ScopeStorageDictionary" /> and <see cref="P:System.Web.WebPages.Scope.ScopeStorageDictionary.BaseScope" /> objects.</summary>
+      <returns>true if the key/value pair is removed, or false if <paramref name="item" /> is not found in the concatenated <see cref="T:System.Web.WebPages.Scope.ScopeStorageDictionary" /> and <see cref="P:System.Web.WebPages.Scope.ScopeStorageDictionary.BaseScope" /> objects.</returns>
+      <param name="item">The key/value pair.</param>
+    </member>
+    <member name="M:System.Web.WebPages.Scope.ScopeStorageDictionary.Remove(System.Object)">
+      <summary>Removes the value that has the specified key from the concatenated <see cref="T:System.Web.WebPages.Scope.ScopeStorageDictionary" /> and <see cref="P:System.Web.WebPages.Scope.ScopeStorageDictionary.BaseScope" /> objects.</summary>
+      <returns>true if the key/value pair is removed, or false if <paramref name="key" /> is not found in the concatenated <see cref="T:System.Web.WebPages.Scope.ScopeStorageDictionary" /> and <see cref="P:System.Web.WebPages.Scope.ScopeStorageDictionary.BaseScope" /> objects.</returns>
+      <param name="key">The key.</param>
+    </member>
+    <member name="M:System.Web.WebPages.Scope.ScopeStorageDictionary.SetValue(System.Object,System.Object)">
+      <summary>Sets a value using the specified key in the concatenated <see cref="T:System.Web.WebPages.Scope.ScopeStorageDictionary" /> and <see cref="P:System.Web.WebPages.Scope.ScopeStorageDictionary.BaseScope" /> objects.</summary>
+      <param name="key">The key.</param>
+      <param name="value">The value.</param>
+    </member>
+    <member name="M:System.Web.WebPages.Scope.ScopeStorageDictionary.System#Collections#IEnumerable#GetEnumerator">
+      <summary>Returns an enumerator for the concatenated <see cref="T:System.Web.WebPages.Scope.ScopeStorageDictionary" /> and <see cref="P:System.Web.WebPages.Scope.ScopeStorageDictionary.BaseScope" /> objects.</summary>
+      <returns>The enumerator.</returns>
+    </member>
+    <member name="M:System.Web.WebPages.Scope.ScopeStorageDictionary.TryGetValue(System.Object,System.Object@)">
+      <summary>Gets the value that is associated with the specified key from the concatenated <see cref="T:System.Web.WebPages.Scope.ScopeStorageDictionary" /> and <see cref="P:System.Web.WebPages.Scope.ScopeStorageDictionary.BaseScope" /> objects.</summary>
+      <returns>true if the concatenated <see cref="T:System.Web.WebPages.Scope.ScopeStorageDictionary" /> and <see cref="P:System.Web.WebPages.Scope.ScopeStorageDictionary.BaseScope" /> objects contain an element that has the specified key; otherwise, false.</returns>
+      <param name="key">The key.</param>
+      <param name="value">When this method returns, if the key is found, contains the value that is associated with the specified key; otherwise, the default value for the type of the <paramref name="value" /> parameter. This parameter is passed uninitialized.</param>
+    </member>
+    <member name="P:System.Web.WebPages.Scope.ScopeStorageDictionary.Values">
+      <summary>Gets a <see cref="T:System.Collections.Generic.List`1" /> object that contains the values from the concatenated <see cref="T:System.Web.WebPages.Scope.ScopeStorageDictionary" /> and <see cref="P:System.Web.WebPages.Scope.ScopeStorageDictionary.BaseScope" /> objects.</summary>
+      <returns>The object that contains the values.</returns>
+    </member>
+    <member name="T:System.Web.WebPages.Scope.StaticScopeStorageProvider">
+      <summary>Provides scoped access to static data.</summary>
+    </member>
+    <member name="M:System.Web.WebPages.Scope.StaticScopeStorageProvider.#ctor">
+      <summary>Initializes a new instance of the <see cref="T:System.Web.WebPages.Scope.StaticScopeStorageProvider" /> class.</summary>
+    </member>
+    <member name="P:System.Web.WebPages.Scope.StaticScopeStorageProvider.CurrentScope">
+      <summary>Gets or sets a dictionary that stores current data under a static context.</summary>
+      <returns>The dictionary that provides current scoped data.</returns>
+    </member>
+    <member name="P:System.Web.WebPages.Scope.StaticScopeStorageProvider.GlobalScope">
+      <summary>Gets a dictionary that stores global data under a static context.</summary>
+      <returns>The dictionary that provides global scoped data.</returns>
+    </member>
+  </members>
+</doc>
\ No newline at end of file
