Copying //UE4/Dev-Core to //UE4/Main

diff --git a/Engine/Source/Programs/CrashReporter/CrashReportWebSite/packages/Microsoft.Owin.3.0.0/lib/net45/Microsoft.Owin.XML b/Engine/Source/Programs/CrashReporter/CrashReportWebSite/packages/Microsoft.Owin.3.0.0/lib/net45/Microsoft.Owin.XML
new file mode 100644
--- /dev/null
+++ b/Engine/Source/Programs/CrashReporter/CrashReportWebSite/packages/Microsoft.Owin.3.0.0/lib/net45/Microsoft.Owin.XML
@@ -0,0 +1,3224 @@
+<?xml version="1.0"?>
+<doc>
+    <assembly>
+        <name>Microsoft.Owin</name>
+    </assembly>
+    <members>
+        <member name="T:Owin.AppBuilderUseExtensions">
+            <summary>
+            Extension methods for <see cref="T:Owin.IAppBuilder"/>.
+            </summary>
+        </member>
+        <member name="M:Owin.AppBuilderUseExtensions.Use``1(Owin.IAppBuilder,System.Object[])">
+            <summary>
+            Inserts a middleware into the OWIN pipeline.
+            </summary>
+            <typeparam name="T">The middleware type</typeparam>
+            <param name="app"></param>
+            <param name="args">Any additional arguments for the middleware constructor</param>
+            <returns></returns>
+        </member>
+        <member name="M:Owin.AppBuilderUseExtensions.Run(Owin.IAppBuilder,System.Func{Microsoft.Owin.IOwinContext,System.Threading.Tasks.Task})">
+            <summary>
+            Inserts into the OWIN pipeline a middleware which does not have a next middleware reference.
+            </summary>
+            <param name="app"></param>
+            <param name="handler">An app that handles all requests</param>
+        </member>
+        <member name="M:Owin.AppBuilderUseExtensions.Use(Owin.IAppBuilder,System.Func{Microsoft.Owin.IOwinContext,System.Func{System.Threading.Tasks.Task},System.Threading.Tasks.Task})">
+            <summary>
+            Inserts a middleware into the OWIN pipeline.
+            </summary>
+            <param name="app"></param>
+            <param name="handler">An app that handles the request or calls the given next Func</param>
+            <returns></returns>
+        </member>
+        <member name="T:Microsoft.Owin.Extensions.UseHandlerMiddleware">
+            <summary>
+            Represents a middleware for executing in-line function middleware.
+            </summary>
+        </member>
+        <member name="M:Microsoft.Owin.Extensions.UseHandlerMiddleware.#ctor(System.Func{System.Collections.Generic.IDictionary{System.String,System.Object},System.Threading.Tasks.Task},System.Func{Microsoft.Owin.IOwinContext,System.Threading.Tasks.Task})">
+            <summary>
+            Initializes a new instance of the <see cref="T:Microsoft.Owin.Extensions.UseHandlerMiddleware" /> class.
+            </summary>
+            <param name="next">The pointer to next middleware.</param>
+            <param name="handler">A function that handles all requests.</param>
+        </member>
+        <member name="M:Microsoft.Owin.Extensions.UseHandlerMiddleware.#ctor(System.Func{System.Collections.Generic.IDictionary{System.String,System.Object},System.Threading.Tasks.Task},System.Func{Microsoft.Owin.IOwinContext,System.Func{System.Threading.Tasks.Task},System.Threading.Tasks.Task})">
+            <summary>
+            Initializes a new instance of the <see cref="T:Microsoft.Owin.Extensions.UseHandlerMiddleware" /> class.
+            </summary>
+            <param name="next">The pointer to next middleware.</param>
+            <param name="handler">A function that handles the request or calls the given next function.</param>
+        </member>
+        <member name="M:Microsoft.Owin.Extensions.UseHandlerMiddleware.Invoke(System.Collections.Generic.IDictionary{System.String,System.Object})">
+            <summary>
+            Invokes the handler for processing the request.
+            </summary>
+            <param name="environment">The OWIN context.</param>
+            <returns>The <see cref="T:System.Threading.Tasks.Task" /> object that represents the request operation.</returns>
+        </member>
+        <member name="T:Microsoft.Owin.FormCollection">
+            <summary>
+            Contains the parsed form values.
+            </summary>
+        </member>
+        <member name="T:Microsoft.Owin.ReadableStringCollection">
+            <summary>
+            Accessors for query, forms, etc.
+            </summary>
+        </member>
+        <member name="T:Microsoft.Owin.IReadableStringCollection">
+            <summary>
+            Accessors for headers, query, forms, etc.
+            </summary>
+        </member>
+        <member name="M:Microsoft.Owin.IReadableStringCollection.Get(System.String)">
+            <summary>
+            Get the associated value from the collection.  Multiple values will be merged.
+            Returns null if the key is not present.
+            </summary>
+            <param name="key"></param>
+            <returns></returns>
+        </member>
+        <member name="M:Microsoft.Owin.IReadableStringCollection.GetValues(System.String)">
+            <summary>
+            Get the associated values from the collection in their original format.
+            Returns null if the key is not present.
+            </summary>
+            <param name="key"></param>
+            <returns></returns>
+        </member>
+        <member name="P:Microsoft.Owin.IReadableStringCollection.Item(System.String)">
+            <summary>
+            Get the associated value from the collection.  Multiple values will be merged.
+            Returns null if the key is not present.
+            </summary>
+            <param name="key"></param>
+            <returns></returns>
+        </member>
+        <member name="M:Microsoft.Owin.ReadableStringCollection.#ctor(System.Collections.Generic.IDictionary{System.String,System.String[]})">
+            <summary>
+            Create a new wrapper
+            </summary>
+            <param name="store"></param>
+        </member>
+        <member name="M:Microsoft.Owin.ReadableStringCollection.Get(System.String)">
+            <summary>
+            Get the associated value from the collection.  Multiple values will be merged.
+            Returns null if the key is not present.
+            </summary>
+            <param name="key"></param>
+            <returns></returns>
+        </member>
+        <member name="M:Microsoft.Owin.ReadableStringCollection.GetValues(System.String)">
+            <summary>
+            Get the associated values from the collection in their original format.
+            Returns null if the key is not present.
+            </summary>
+            <param name="key"></param>
+            <returns></returns>
+        </member>
+        <member name="M:Microsoft.Owin.ReadableStringCollection.GetEnumerator">
+            <summary>
+            
+            </summary>
+            <returns></returns>
+        </member>
+        <member name="M:Microsoft.Owin.ReadableStringCollection.System#Collections#IEnumerable#GetEnumerator">
+            <summary>
+            
+            </summary>
+            <returns></returns>
+        </member>
+        <member name="P:Microsoft.Owin.ReadableStringCollection.Item(System.String)">
+            <summary>
+            Get the associated value from the collection.  Multiple values will be merged.
+            Returns null if the key is not present.
+            </summary>
+            <param name="key"></param>
+            <returns></returns>
+        </member>
+        <member name="T:Microsoft.Owin.IFormCollection">
+            <summary>
+            Contains the parsed form values.
+            </summary>
+        </member>
+        <member name="M:Microsoft.Owin.FormCollection.#ctor(System.Collections.Generic.IDictionary{System.String,System.String[]})">
+            <summary>
+            Initializes a new instance of the <see cref="T:Microsoft.Owin.FormCollection" /> class.
+            </summary>
+            <param name="store">The store for the form.</param>
+        </member>
+        <member name="T:Microsoft.Owin.HeaderDictionary">
+            <summary>
+            Represents a wrapper for owin.RequestHeaders and owin.ResponseHeaders.
+            </summary>
+        </member>
+        <member name="T:Microsoft.Owin.IHeaderDictionary">
+            <summary>
+            Represents a wrapper for owin.RequestHeaders and owin.ResponseHeaders.
+            </summary>
+        </member>
+        <member name="M:Microsoft.Owin.IHeaderDictionary.GetCommaSeparatedValues(System.String)">
+            <summary>
+            Get the associated values from the collection separated into individual values.
+            Quoted values will not be split, and the quotes will be removed.
+            </summary>
+            <param name="key">The header name.</param>
+            <returns>the associated values from the collection separated into individual values, or null if the key is not present.</returns>
+        </member>
+        <member name="M:Microsoft.Owin.IHeaderDictionary.Append(System.String,System.String)">
+            <summary>
+            Add a new value. Appends to the header if already present
+            </summary>
+            <param name="key">The header name.</param>
+            <param name="value">The header value.</param>
+        </member>
+        <member name="M:Microsoft.Owin.IHeaderDictionary.AppendValues(System.String,System.String[])">
+            <summary>
+            Add new values. Each item remains a separate array entry.
+            </summary>
+            <param name="key">The header name.</param>
+            <param name="values">The header values.</param>
+        </member>
+        <member name="M:Microsoft.Owin.IHeaderDictionary.AppendCommaSeparatedValues(System.String,System.String[])">
+            <summary>
+            Quotes any values containing comas, and then coma joins all of the values with any existing values.
+            </summary>
+            <param name="key">The header name.</param>
+            <param name="values">The header values.</param>
+        </member>
+        <member name="M:Microsoft.Owin.IHeaderDictionary.Set(System.String,System.String)">
+            <summary>
+            Sets a specific header value.
+            </summary>
+            <param name="key">The header name.</param>
+            <param name="value">The header value.</param>
+        </member>
+        <member name="M:Microsoft.Owin.IHeaderDictionary.SetValues(System.String,System.String[])">
+            <summary>
+            Sets the specified header values without modification.
+            </summary>
+            <param name="key">The header name.</param>
+            <param name="values">The header values.</param>
+        </member>
+        <member name="M:Microsoft.Owin.IHeaderDictionary.SetCommaSeparatedValues(System.String,System.String[])">
+            <summary>
+            Quotes any values containing comas, and then coma joins all of the values.
+            </summary>
+            <param name="key">The header name.</param>
+            <param name="values">The header values.</param>
+        </member>
+        <member name="P:Microsoft.Owin.IHeaderDictionary.Item(System.String)">
+            <summary>
+            Get or sets the associated value from the collection as a single string.
+            </summary>
+            <param name="key">The header name.</param>
+            <returns>the associated value from the collection as a single string or null if the key is not present.</returns>
+        </member>
+        <member name="M:Microsoft.Owin.HeaderDictionary.#ctor(System.Collections.Generic.IDictionary{System.String,System.String[]})">
+            <summary>
+            Initializes a new instance of the <see cref="T:Microsoft.Owin.HeaderDictionary" /> class.
+            </summary>
+            <param name="store">The underlying data store.</param>
+        </member>
+        <member name="M:Microsoft.Owin.HeaderDictionary.GetEnumerator">
+            <summary>
+            Returns an enumerator that iterates through a collection.
+            </summary>
+            <returns>An <see cref="T:System.Collections.IEnumerator" /> object that can be used to iterate through the collection.</returns>
+        </member>
+        <member name="M:Microsoft.Owin.HeaderDictionary.System#Collections#IEnumerable#GetEnumerator">
+            <summary>
+            Returns an enumerator that iterates through a collection.
+            </summary>
+            <returns>An <see cref="T:System.Collections.IEnumerator" /> object that can be used to iterate through the collection.</returns>
+        </member>
+        <member name="M:Microsoft.Owin.HeaderDictionary.Get(System.String)">
+            <summary>
+            Get the associated value from the collection as a single string.
+            </summary>
+            <param name="key">The header name.</param>
+            <returns>the associated value from the collection as a single string or null if the key is not present.</returns>
+        </member>
+        <member name="M:Microsoft.Owin.HeaderDictionary.GetValues(System.String)">
+            <summary>
+            Get the associated values from the collection without modification.
+            </summary>
+            <param name="key">The header name.</param>
+            <returns>the associated value from the collection without modification, or null if the key is not present.</returns>
+        </member>
+        <member name="M:Microsoft.Owin.HeaderDictionary.GetCommaSeparatedValues(System.String)">
+            <summary>
+            Get the associated values from the collection separated into individual values.
+            Quoted values will not be split, and the quotes will be removed.
+            </summary>
+            <param name="key">The header name.</param>
+            <returns>the associated values from the collection separated into individual values, or null if the key is not present.</returns>
+        </member>
+        <member name="M:Microsoft.Owin.HeaderDictionary.Append(System.String,System.String)">
+            <summary>
+            Add a new value. Appends to the header if already present
+            </summary>
+            <param name="key">The header name.</param>
+            <param name="value">The header value.</param>
+        </member>
+        <member name="M:Microsoft.Owin.HeaderDictionary.AppendValues(System.String,System.String[])">
+            <summary>
+            Add new values. Each item remains a separate array entry.
+            </summary>
+            <param name="key">The header name.</param>
+            <param name="values">The header values.</param>
+        </member>
+        <member name="M:Microsoft.Owin.HeaderDictionary.AppendCommaSeparatedValues(System.String,System.String[])">
+            <summary>
+            Quotes any values containing comas, and then coma joins all of the values with any existing values.
+            </summary>
+            <param name="key">The header name.</param>
+            <param name="values">The header values.</param>
+        </member>
+        <member name="M:Microsoft.Owin.HeaderDictionary.Set(System.String,System.String)">
+            <summary>
+            Sets a specific header value.
+            </summary>
+            <param name="key">The header name.</param>
+            <param name="value">The header value.</param>
+        </member>
+        <member name="M:Microsoft.Owin.HeaderDictionary.SetValues(System.String,System.String[])">
+            <summary>
+            Sets the specified header values without modification.
+            </summary>
+            <param name="key">The header name.</param>
+            <param name="values">The header values.</param>
+        </member>
+        <member name="M:Microsoft.Owin.HeaderDictionary.SetCommaSeparatedValues(System.String,System.String[])">
+            <summary>
+            Quotes any values containing comas, and then coma joins all of the values.
+            </summary>
+            <param name="key">The header name.</param>
+            <param name="values">The header values.</param>
+        </member>
+        <member name="M:Microsoft.Owin.HeaderDictionary.Add(System.String,System.String[])">
+            <summary>
+            Adds the given header and values to the collection.
+            </summary>
+            <param name="key">The header name.</param>
+            <param name="value">The header values.</param>
+        </member>
+        <member name="M:Microsoft.Owin.HeaderDictionary.ContainsKey(System.String)">
+            <summary>
+            Determines whether the <see cref="T:Microsoft.Owin.HeaderDictionary" /> contains a specific key.
+            </summary>
+            <param name="key">The key.</param>
+            <returns>true if the <see cref="T:Microsoft.Owin.HeaderDictionary" /> contains a specific key; otherwise, false.</returns>
+        </member>
+        <member name="M:Microsoft.Owin.HeaderDictionary.Remove(System.String)">
+            <summary>
+            Removes the given header from the collection.
+            </summary>
+            <param name="key">The header name.</param>
+            <returns>true if the specified object was removed from the collection; otherwise, false.</returns>
+        </member>
+        <member name="M:Microsoft.Owin.HeaderDictionary.TryGetValue(System.String,System.String[]@)">
+            <summary>
+            Retrieves a value from the dictionary.
+            </summary>
+            <param name="key">The header name.</param>
+            <param name="value">The value.</param>
+            <returns>true if the <see cref="T:Microsoft.Owin.HeaderDictionary" /> contains the key; otherwise, false.</returns>
+        </member>
+        <member name="M:Microsoft.Owin.HeaderDictionary.Add(System.Collections.Generic.KeyValuePair{System.String,System.String[]})">
+            <summary>
+            Adds a new list of items to the collection.
+            </summary>
+            <param name="item">The item to add.</param>
+        </member>
+        <member name="M:Microsoft.Owin.HeaderDictionary.Clear">
+            <summary>
+            Clears the entire list of objects.
+            </summary>
+        </member>
+        <member name="M:Microsoft.Owin.HeaderDictionary.Contains(System.Collections.Generic.KeyValuePair{System.String,System.String[]})">
+            <summary>
+            Returns a value indicating whether the specified object occurs within this collection.
+            </summary>
+            <param name="item">The item.</param>
+            <returns>true if the specified object occurs within this collection; otherwise, false.</returns>
+        </member>
+        <member name="M:Microsoft.Owin.HeaderDictionary.CopyTo(System.Collections.Generic.KeyValuePair{System.String,System.String[]}[],System.Int32)">
+            <summary>
+            Copies the <see cref="T:Microsoft.Owin.HeaderDictionary" /> elements to a one-dimensional Array instance at the specified index.
+            </summary>
+            <param name="array">The one-dimensional Array that is the destination of the specified objects copied from the <see cref="T:Microsoft.Owin.HeaderDictionary" />.</param>
+            <param name="arrayIndex">The zero-based index in <paramref name="array" /> at which copying begins.</param>
+        </member>
+        <member name="M:Microsoft.Owin.HeaderDictionary.Remove(System.Collections.Generic.KeyValuePair{System.String,System.String[]})">
+            <summary>
+            Removes the given item from the the collection.
+            </summary>
+            <param name="item">The item.</param>
+            <returns>true if the specified object was removed from the collection; otherwise, false.</returns>
+        </member>
+        <member name="P:Microsoft.Owin.HeaderDictionary.Keys">
+            <summary>
+            Gets an <see cref="T:System.Collections.ICollection" /> that contains the keys in the <see cref="T:Microsoft.Owin.HeaderDictionary" />;.
+            </summary>
+            <returns>An <see cref="T:System.Collections.ICollection" /> that contains the keys in the <see cref="T:Microsoft.Owin.HeaderDictionary" />.</returns>
+        </member>
+        <member name="P:Microsoft.Owin.HeaderDictionary.Values">
+            <summary>
+            
+            </summary>
+        </member>
+        <member name="P:Microsoft.Owin.HeaderDictionary.Count">
+            <summary>
+            Gets the number of elements contained in the <see cref="T:Microsoft.Owin.HeaderDictionary" />;.
+            </summary>
+            <returns>The number of elements contained in the <see cref="T:Microsoft.Owin.HeaderDictionary" />.</returns>
+        </member>
+        <member name="P:Microsoft.Owin.HeaderDictionary.IsReadOnly">
+            <summary>
+            Gets a value that indicates whether the <see cref="T:Microsoft.Owin.HeaderDictionary" /> is in read-only mode.
+            </summary>
+            <returns>true if the <see cref="T:Microsoft.Owin.HeaderDictionary" /> is in read-only mode; otherwise, false.</returns>
+        </member>
+        <member name="P:Microsoft.Owin.HeaderDictionary.Item(System.String)">
+            <summary>
+            Get or sets the associated value from the collection as a single string.
+            </summary>
+            <param name="key">The header name.</param>
+            <returns>the associated value from the collection as a single string or null if the key is not present.</returns>
+        </member>
+        <member name="P:Microsoft.Owin.HeaderDictionary.System#Collections#Generic#IDictionary{System#String@System#String[]}#Item(System.String)">
+            <summary>
+            Throws KeyNotFoundException if the key is not present.
+            </summary>
+            <param name="key">The header name.</param>
+            <returns></returns>
+        </member>
+        <member name="T:Microsoft.Owin.HostString">
+            <summary>
+            Represents the host portion of a Uri can be used to construct Uri's properly formatted and encoded for use in
+            HTTP headers.
+            </summary>
+        </member>
+        <member name="M:Microsoft.Owin.HostString.#ctor(System.String)">
+            <summary>
+            Creates a new HostString without modification. The value should be Unicode rather than punycode, and may have a port.
+            IPv4 and IPv6 addresses are also allowed, and also may have ports.
+            </summary>
+            <param name="value"></param>
+        </member>
+        <member name="M:Microsoft.Owin.HostString.ToString">
+            <summary>
+            Returns the value as normalized by ToUriComponent().
+            </summary>
+            <returns></returns>
+        </member>
+        <member name="M:Microsoft.Owin.HostString.ToUriComponent">
+            <summary>
+            Returns the value properly formatted and encoded for use in a URI in a HTTP header.
+            Any Unicode is converted to punycode. IPv6 addresses will have brackets added if they are missing.
+            </summary>
+            <returns></returns>
+        </member>
+        <member name="M:Microsoft.Owin.HostString.FromUriComponent(System.String)">
+            <summary>
+            Creates a new HostString from the given uri component.
+            Any punycode will be converted to Unicode.
+            </summary>
+            <param name="uriComponent"></param>
+            <returns></returns>
+        </member>
+        <member name="M:Microsoft.Owin.HostString.FromUriComponent(System.Uri)">
+            <summary>
+            Creates a new HostString from the host and port of the give Uri instance.
+            Punycode will be converted to Unicode.
+            </summary>
+            <param name="uri"></param>
+            <returns></returns>
+        </member>
+        <member name="M:Microsoft.Owin.HostString.Equals(Microsoft.Owin.HostString)">
+            <summary>
+            Compares the equality of the Value property, ignoring case.
+            </summary>
+            <param name="other"></param>
+            <returns></returns>
+        </member>
+        <member name="M:Microsoft.Owin.HostString.Equals(System.Object)">
+            <summary>
+            Compares against the given object only if it is a HostString.
+            </summary>
+            <param name="obj"></param>
+            <returns></returns>
+        </member>
+        <member name="M:Microsoft.Owin.HostString.GetHashCode">
+            <summary>
+            Gets a hash code for the value.
+            </summary>
+            <returns></returns>
+        </member>
+        <member name="M:Microsoft.Owin.HostString.op_Equality(Microsoft.Owin.HostString,Microsoft.Owin.HostString)">
+            <summary>
+            Compares the two instances for equality.
+            </summary>
+            <param name="left"></param>
+            <param name="right"></param>
+            <returns></returns>
+        </member>
+        <member name="M:Microsoft.Owin.HostString.op_Inequality(Microsoft.Owin.HostString,Microsoft.Owin.HostString)">
+            <summary>
+            Compares the two instances for inequality.
+            </summary>
+            <param name="left"></param>
+            <param name="right"></param>
+            <returns></returns>
+        </member>
+        <member name="P:Microsoft.Owin.HostString.Value">
+            <summary>
+            Returns the original value from the constructor.
+            </summary>
+        </member>
+        <member name="T:Microsoft.Owin.Infrastructure.ChunkingCookieManager">
+            <summary>
+            This handles cookies that are limited by per cookie length. It breaks down long cookies for responses, and reassembles them
+            from requests.
+            </summary>
+        </member>
+        <member name="M:Microsoft.Owin.Infrastructure.ChunkingCookieManager.GetRequestCookie(Microsoft.Owin.IOwinContext,System.String)">
+            <summary>
+            Get the reassembled cookie. Non chunked cookies are returned normally.
+            Cookies with missing chunks just have their "chunks:XX" header returned.
+            </summary>
+            <param name="context"></param>
+            <param name="key"></param>
+            <returns>The reassembled cookie, if any, or null.</returns>
+        </member>
+        <member name="M:Microsoft.Owin.Infrastructure.ChunkingCookieManager.AppendResponseCookie(Microsoft.Owin.IOwinContext,System.String,System.String,Microsoft.Owin.CookieOptions)">
+            <summary>
+            Appends a new response cookie to the Set-Cookie header. If the cookie is larger than the given size limit
+            then it will be broken down into multiple cookies as follows:
+            Set-Cookie: CookieName=chunks:3; path=/
+            Set-Cookie: CookieNameC1=Segment1; path=/
+            Set-Cookie: CookieNameC2=Segment2; path=/
+            Set-Cookie: CookieNameC3=Segment3; path=/
+            </summary>
+            <param name="context"></param>
+            <param name="key"></param>
+            <param name="value"></param>
+            <param name="options"></param>
+        </member>
+        <member name="M:Microsoft.Owin.Infrastructure.ChunkingCookieManager.DeleteCookie(Microsoft.Owin.IOwinContext,System.String,Microsoft.Owin.CookieOptions)">
+            <summary>
+            Deletes the cookie with the given key by setting an expired state. If a matching chunked cookie exists on
+            the request, delete each chunk.
+            </summary>
+            <param name="context"></param>
+            <param name="key"></param>
+            <param name="options"></param>
+        </member>
+        <member name="P:Microsoft.Owin.Infrastructure.ChunkingCookieManager.ChunkSize">
+             <summary>
+             The maximum size of cookie to send back to the client. If a cookie exceeds this size it will be broken down into multiple
+             cookies. Set this value to null to disable this behavior. The default is 4090 characters, which is supported by all
+             common browsers.
+            
+             Note that browsers may also have limits on the total size of all cookies per domain, and on the number of cookies per domain.
+             </summary>
+        </member>
+        <member name="P:Microsoft.Owin.Infrastructure.ChunkingCookieManager.ThrowForPartialCookies">
+            <summary>
+            Throw if not all chunks of a cookie are available on a request for re-assembly.
+            </summary>
+        </member>
+        <member name="T:Microsoft.Owin.IOwinContext">
+            <summary>
+            This wraps OWIN environment dictionary and provides strongly typed accessors.
+            </summary>
+        </member>
+        <member name="M:Microsoft.Owin.IOwinContext.Get``1(System.String)">
+            <summary>
+            Gets a value from the OWIN environment, or returns default(T) if not present.
+            </summary>
+            <typeparam name="T">The type of the value.</typeparam>
+            <param name="key">The key of the value to get.</param>
+            <returns>The value with the specified key or the default(T) if not present.</returns>
+        </member>
+        <member name="M:Microsoft.Owin.IOwinContext.Set``1(System.String,``0)">
+            <summary>
+            Sets the given key and value in the OWIN environment.
+            </summary>
+            <typeparam name="T">The type of the value.</typeparam>
+            <param name="key">The key of the value to set.</param>
+            <param name="value">The value to set.</param>
+            <returns>This instance.</returns>
+        </member>
+        <member name="P:Microsoft.Owin.IOwinContext.Request">
+            <summary>
+            Gets a wrapper exposing request specific properties.
+            </summary>
+            <returns>A wrapper exposing request specific properties.</returns>
+        </member>
+        <member name="P:Microsoft.Owin.IOwinContext.Response">
+            <summary>
+            Gets a wrapper exposing response specific properties.
+            </summary>
+            <returns>A wrapper exposing response specific properties.</returns>
+        </member>
+        <member name="P:Microsoft.Owin.IOwinContext.Authentication">
+            <summary>
+            Gets the Authentication middleware functionality available on the current request.
+            </summary>
+            <returns>The authentication middleware functionality available on the current request.</returns>
+        </member>
+        <member name="P:Microsoft.Owin.IOwinContext.Environment">
+            <summary>
+            Gets the OWIN environment.
+            </summary>
+            <returns>The OWIN environment.</returns>
+        </member>
+        <member name="P:Microsoft.Owin.IOwinContext.TraceOutput">
+            <summary>
+            Gets or sets the host.TraceOutput environment value.
+            </summary>
+            <returns>The host.TraceOutput TextWriter.</returns>
+        </member>
+        <member name="T:Microsoft.Owin.IOwinRequest">
+            <summary>
+            This wraps OWIN environment dictionary and provides strongly typed accessors.
+            </summary>
+        </member>
+        <member name="M:Microsoft.Owin.IOwinRequest.ReadFormAsync">
+            <summary>
+            Asynchronously reads and parses the request body as a form.
+            </summary>
+            <returns>The parsed form data.</returns>
+        </member>
+        <member name="M:Microsoft.Owin.IOwinRequest.Get``1(System.String)">
+            <summary>
+            Gets a value from the OWIN environment, or returns default(T) if not present.
+            </summary>
+            <typeparam name="T">The type of the value.</typeparam>
+            <param name="key">The key of the value to get.</param>
+            <returns>The value with the specified key or the default(T) if not present.</returns>
+        </member>
+        <member name="M:Microsoft.Owin.IOwinRequest.Set``1(System.String,``0)">
+            <summary>
+            Sets the given key and value in the OWIN environment.
+            </summary>
+            <typeparam name="T">The type of the value.</typeparam>
+            <param name="key">The key of the value to set.</param>
+            <param name="value">The value to set.</param>
+            <returns>This instance.</returns>
+        </member>
+        <member name="P:Microsoft.Owin.IOwinRequest.Environment">
+            <summary>
+            Gets the OWIN environment.
+            </summary>
+            <returns>The OWIN environment.</returns>
+        </member>
+        <member name="P:Microsoft.Owin.IOwinRequest.Context">
+            <summary>
+            Gets the request context.
+            </summary>
+            <returns>The request context.</returns>
+        </member>
+        <member name="P:Microsoft.Owin.IOwinRequest.Method">
+            <summary>
+            Gets or set the HTTP method.
+            </summary>
+            <returns>The HTTP method.</returns>
+        </member>
+        <member name="P:Microsoft.Owin.IOwinRequest.Scheme">
+            <summary>
+            Gets or set the HTTP request scheme from owin.RequestScheme.
+            </summary>
+            <returns>The HTTP request scheme from owin.RequestScheme.</returns>
+        </member>
+        <member name="P:Microsoft.Owin.IOwinRequest.IsSecure">
+            <summary>
+            Returns true if the owin.RequestScheme is https.
+            </summary>
+            <returns>true if this request is using https; otherwise, false.</returns>
+        </member>
+        <member name="P:Microsoft.Owin.IOwinRequest.Host">
+            <summary>
+            Gets or set the Host header. May include the port.
+            </summary>
+            <return>The Host header.</return>
+        </member>
+        <member name="P:Microsoft.Owin.IOwinRequest.PathBase">
+            <summary>
+            Gets or set the owin.RequestPathBase.
+            </summary>
+            <returns>The owin.RequestPathBase.</returns>
+        </member>
+        <member name="P:Microsoft.Owin.IOwinRequest.Path">
+            <summary>
+            Gets or set the request path from owin.RequestPath.
+            </summary>
+            <returns>The request path from owin.RequestPath.</returns>
+        </member>
+        <member name="P:Microsoft.Owin.IOwinRequest.QueryString">
+            <summary>
+            Gets or set the query string from owin.RequestQueryString.
+            </summary>
+            <returns>The query string from owin.RequestQueryString.</returns>
+        </member>
+        <member name="P:Microsoft.Owin.IOwinRequest.Query">
+            <summary>
+            Gets the query value collection parsed from owin.RequestQueryString.
+            </summary>
+            <returns>The query value collection parsed from owin.RequestQueryString.</returns>
+        </member>
+        <member name="P:Microsoft.Owin.IOwinRequest.Uri">
+            <summary>
+            Gets the uniform resource identifier (URI) associated with the request.
+            </summary>
+            <returns>The uniform resource identifier (URI) associated with the request.</returns>
+        </member>
+        <member name="P:Microsoft.Owin.IOwinRequest.Protocol">
+            <summary>
+            Gets or set the owin.RequestProtocol.
+            </summary>
+            <returns>The owin.RequestProtocol.</returns>
+        </member>
+        <member name="P:Microsoft.Owin.IOwinRequest.Headers">
+            <summary>
+            Gets the request headers.
+            </summary>
+            <returns>The request headers.</returns>
+        </member>
+        <member name="P:Microsoft.Owin.IOwinRequest.Cookies">
+            <summary>
+            Gets the collection of Cookies for this request.
+            </summary>
+            <returns>The collection of Cookies for this request.</returns>
+        </member>
+        <member name="P:Microsoft.Owin.IOwinRequest.ContentType">
+            <summary>
+            Gets or sets the Content-Type header.
+            </summary>
+            <returns>The Content-Type header.</returns>
+        </member>
+        <member name="P:Microsoft.Owin.IOwinRequest.CacheControl">
+            <summary>
+            Gets or sets the Cache-Control header.
+            </summary>
+            <returns>The Cache-Control header.</returns>
+        </member>
+        <member name="P:Microsoft.Owin.IOwinRequest.MediaType">
+            <summary>
+            Gets or sets the Media-Type header.
+            </summary>
+            <returns>The Media-Type header.</returns>
+        </member>
+        <member name="P:Microsoft.Owin.IOwinRequest.Accept">
+            <summary>
+            Gets or set the Accept header.
+            </summary>
+            <returns>The Accept header.</returns>
+        </member>
+        <member name="P:Microsoft.Owin.IOwinRequest.Body">
+            <summary>
+            Gets or set the owin.RequestBody Stream.
+            </summary>
+            <returns>The owin.RequestBody Stream.</returns>
+        </member>
+        <member name="P:Microsoft.Owin.IOwinRequest.CallCancelled">
+            <summary>
+            Gets or sets the cancellation token for the request.
+            </summary>
+            <returns>The cancellation token for the request.</returns>
+        </member>
+        <member name="P:Microsoft.Owin.IOwinRequest.LocalIpAddress">
+            <summary>
+            Gets or set the server.LocalIpAddress.
+            </summary>
+            <returns>The server.LocalIpAddress.</returns>
+        </member>
+        <member name="P:Microsoft.Owin.IOwinRequest.LocalPort">
+            <summary>
+            Gets or set the server.LocalPort.
+            </summary>
+            <returns>The server.LocalPort.</returns>
+        </member>
+        <member name="P:Microsoft.Owin.IOwinRequest.RemoteIpAddress">
+            <summary>
+            Gets or set the server.RemoteIpAddress.
+            </summary>
+            <returns>The server.RemoteIpAddress.</returns>
+        </member>
+        <member name="P:Microsoft.Owin.IOwinRequest.RemotePort">
+            <summary>
+            Gets or set the server.RemotePort.
+            </summary>
+            <returns>The server.RemotePort.</returns>
+        </member>
+        <member name="P:Microsoft.Owin.IOwinRequest.User">
+            <summary>
+            Gets or set the server.User.
+            </summary>
+            <returns>The server.User.</returns>
+        </member>
+        <member name="T:Microsoft.Owin.IOwinResponse">
+            <summary>
+            This wraps OWIN environment dictionary and provides strongly typed accessors.
+            </summary>
+        </member>
+        <member name="M:Microsoft.Owin.IOwinResponse.OnSendingHeaders(System.Action{System.Object},System.Object)">
+            <summary>
+            Registers for an event that fires when the response headers are sent.
+            </summary>
+            <param name="callback">The callback method.</param>
+            <param name="state">The callback state.</param>
+        </member>
+        <member name="M:Microsoft.Owin.IOwinResponse.Redirect(System.String)">
+            <summary>
+            Sets a 302 response status code and the Location header.
+            </summary>
+            <param name="location">The location where to redirect the client.</param>
+        </member>
+        <member name="M:Microsoft.Owin.IOwinResponse.Write(System.String)">
+            <summary>
+            Writes the given text to the response body stream using UTF-8.
+            </summary>
+            <param name="text">The response data.</param>
+        </member>
+        <member name="M:Microsoft.Owin.IOwinResponse.Write(System.Byte[])">
+            <summary>
+            Writes the given bytes to the response body stream.
+            </summary>
+            <param name="data">The response data.</param>
+        </member>
+        <member name="M:Microsoft.Owin.IOwinResponse.Write(System.Byte[],System.Int32,System.Int32)">
+            <summary>
+            Writes the given bytes to the response body stream.
+            </summary>
+            <param name="data">The response data.</param>
+            <param name="offset">The zero-based byte offset in the <paramref name="data" /> parameter at which to begin copying bytes.</param>
+            <param name="count">The number of bytes to write.</param>
+        </member>
+        <member name="M:Microsoft.Owin.IOwinResponse.WriteAsync(System.String)">
+            <summary>
+            Asynchronously writes the given text to the response body stream using UTF-8.
+            </summary>
+            <param name="text">The response data.</param>
+            <returns>A Task tracking the state of the write operation.</returns>
+        </member>
+        <member name="M:Microsoft.Owin.IOwinResponse.WriteAsync(System.String,System.Threading.CancellationToken)">
+            <summary>
+            Asynchronously writes the given text to the response body stream using UTF-8.
+            </summary>
+            <param name="text">The response data.</param>
+            <param name="token">A token used to indicate cancellation.</param>
+            <returns>A Task tracking the state of the write operation.</returns>
+        </member>
+        <member name="M:Microsoft.Owin.IOwinResponse.WriteAsync(System.Byte[])">
+            <summary>
+            Asynchronously writes the given bytes to the response body stream.
+            </summary>
+            <param name="data">The response data.</param>
+            <returns>A Task tracking the state of the write operation.</returns>
+        </member>
+        <member name="M:Microsoft.Owin.IOwinResponse.WriteAsync(System.Byte[],System.Threading.CancellationToken)">
+            <summary>
+            Asynchronously writes the given bytes to the response body stream.
+            </summary>
+            <param name="data">The response data.</param>
+            <param name="token">A token used to indicate cancellation.</param>
+            <returns>A Task tracking the state of the write operation.</returns>
+        </member>
+        <member name="M:Microsoft.Owin.IOwinResponse.WriteAsync(System.Byte[],System.Int32,System.Int32,System.Threading.CancellationToken)">
+            <summary>
+            Asynchronously writes the given bytes to the response body stream.
+            </summary>
+            <param name="data">The response data.</param>
+            <param name="offset">The zero-based byte offset in the <paramref name="data" /> parameter at which to begin copying bytes.</param>
+            <param name="count">The number of bytes to write.</param>
+            <param name="token">A token used to indicate cancellation.</param>
+            <returns>A Task tracking the state of the write operation.</returns>
+        </member>
+        <member name="M:Microsoft.Owin.IOwinResponse.Get``1(System.String)">
+            <summary>
+            Gets a value from the OWIN environment, or returns default(T) if not present.
+            </summary>
+            <typeparam name="T">The type of the value.</typeparam>
+            <param name="key">The key of the value to get.</param>
+            <returns>The value with the specified key or the default(T) if not present.</returns>
+        </member>
+        <member name="M:Microsoft.Owin.IOwinResponse.Set``1(System.String,``0)">
+            <summary>
+            Sets the given key and value in the OWIN environment.
+            </summary>
+            <typeparam name="T">The type of the value.</typeparam>
+            <param name="key">The key of the value to set.</param>
+            <param name="value">The value to set.</param>
+            <returns>This instance.</returns>
+        </member>
+        <member name="P:Microsoft.Owin.IOwinResponse.Environment">
+            <summary>
+            Gets the OWIN environment.
+            </summary>
+            <returns>The OWIN environment.</returns>
+        </member>
+        <member name="P:Microsoft.Owin.IOwinResponse.Context">
+            <summary>
+            Gets the request context.
+            </summary>
+            <returns>The request context.</returns>
+        </member>
+        <member name="P:Microsoft.Owin.IOwinResponse.StatusCode">
+            <summary>
+            Gets or sets the optional owin.ResponseStatusCode.
+            </summary>
+            <returns>The optional owin.ResponseStatusCode, or 200 if not set.</returns>
+        </member>
+        <member name="P:Microsoft.Owin.IOwinResponse.ReasonPhrase">
+            <summary>
+            Gets or sets the the optional owin.ResponseReasonPhrase.
+            </summary>
+            <returns>The the optional owin.ResponseReasonPhrase.</returns>
+        </member>
+        <member name="P:Microsoft.Owin.IOwinResponse.Protocol">
+            <summary>
+            Gets or sets the owin.ResponseProtocol.
+            </summary>
+            <returns>The owin.ResponseProtocol.</returns>
+        </member>
+        <member name="P:Microsoft.Owin.IOwinResponse.Headers">
+            <summary>
+            Gets the response header collection.
+            </summary>
+            <returns>The response header collection.</returns>
+        </member>
+        <member name="P:Microsoft.Owin.IOwinResponse.Cookies">
+            <summary>
+            Gets a collection used to manipulate the Set-Cookie header.
+            </summary>
+            <returns>A collection used to manipulate the Set-Cookie header.</returns>
+        </member>
+        <member name="P:Microsoft.Owin.IOwinResponse.ContentLength">
+            <summary>
+            Gets or sets the Content-Length header.
+            </summary>
+            <returns>The Content-Length header.</returns>
+        </member>
+        <member name="P:Microsoft.Owin.IOwinResponse.ContentType">
+            <summary>
+            Gets or sets the Content-Type header.
+            </summary>
+            <returns>The Content-Type header.</returns>
+        </member>
+        <member name="P:Microsoft.Owin.IOwinResponse.Expires">
+            <summary>
+            Gets or sets the Expires header.
+            </summary>
+            <returns>The Expires header.</returns>
+        </member>
+        <member name="P:Microsoft.Owin.IOwinResponse.ETag">
+            <summary>
+            Gets or sets the E-Tag header.
+            </summary>
+            <returns>The E-Tag header.</returns>
+        </member>
+        <member name="P:Microsoft.Owin.IOwinResponse.Body">
+            <summary>
+            Gets or sets the owin.ResponseBody Stream.
+            </summary>
+            <returns>The owin.ResponseBody Stream.</returns>
+        </member>
+        <member name="T:Microsoft.Owin.OwinStartupAttribute">
+            <summary>
+            Used to mark which class in an assembly should be used for automatic startup.
+            </summary>
+        </member>
+        <member name="M:Microsoft.Owin.OwinStartupAttribute.#ctor(System.Type)">
+            <summary>
+            Initializes a new instance of the <see cref="T:Microsoft.Owin.OwinStartupAttribute"/> class
+            </summary>
+            <param name="startupType">The startup class</param>
+        </member>
+        <member name="M:Microsoft.Owin.OwinStartupAttribute.#ctor(System.String,System.Type)">
+            <summary>
+            Initializes a new instance of the <see cref="T:Microsoft.Owin.OwinStartupAttribute"/> class
+            </summary>
+            <param name="friendlyName">A non-default configuration, e.g. staging.</param>
+            <param name="startupType">The startup class</param>
+        </member>
+        <member name="M:Microsoft.Owin.OwinStartupAttribute.#ctor(System.Type,System.String)">
+            <summary>
+            Initializes a new instance of the <see cref="T:Microsoft.Owin.OwinStartupAttribute"/> class
+            </summary>
+            <param name="startupType">The startup class</param>
+            <param name="methodName">Specifies which method to call</param>
+        </member>
+        <member name="M:Microsoft.Owin.OwinStartupAttribute.#ctor(System.String,System.Type,System.String)">
+            <summary>
+            Initializes a new instance of the <see cref="T:Microsoft.Owin.OwinStartupAttribute"/> class
+            </summary>
+            <param name="friendlyName">A non-default configuration, e.g. staging.</param>
+            <param name="startupType">The startup class</param>
+            <param name="methodName">Specifies which method to call</param>
+        </member>
+        <member name="P:Microsoft.Owin.OwinStartupAttribute.FriendlyName">
+            <summary>
+            A non-default configuration if any. e.g. Staging.
+            </summary>
+        </member>
+        <member name="P:Microsoft.Owin.OwinStartupAttribute.StartupType">
+            <summary>
+            The startup class
+            </summary>
+        </member>
+        <member name="P:Microsoft.Owin.OwinStartupAttribute.MethodName">
+            <summary>
+            The name of the configuration method
+            </summary>
+        </member>
+        <member name="T:Owin.PipelineStage">
+            <summary>
+            An ordered list of known Asp.Net integrated pipeline stages. More details on the ASP.NET integrated pipeline can be found at http://msdn.microsoft.com/en-us/library/system.web.httpapplication.aspx
+            </summary>
+        </member>
+        <member name="F:Owin.PipelineStage.Authenticate">
+            <summary>
+            Corresponds to the AuthenticateRequest stage of the ASP.NET integrated pipeline.
+            </summary>
+        </member>
+        <member name="F:Owin.PipelineStage.PostAuthenticate">
+            <summary>
+            Corresponds to the PostAuthenticateRequest stage of the ASP.NET integrated pipeline.
+            </summary>
+        </member>
+        <member name="F:Owin.PipelineStage.Authorize">
+            <summary>
+            Corresponds to the AuthorizeRequest stage of the ASP.NET integrated pipeline.
+            </summary>
+        </member>
+        <member name="F:Owin.PipelineStage.PostAuthorize">
+            <summary>
+            Corresponds to the PostAuthorizeRequest stage of the ASP.NET integrated pipeline.
+            </summary>
+        </member>
+        <member name="F:Owin.PipelineStage.ResolveCache">
+            <summary>
+            Corresponds to the ResolveRequestCache stage of the ASP.NET integrated pipeline.
+            </summary>
+        </member>
+        <member name="F:Owin.PipelineStage.PostResolveCache">
+            <summary>
+            Corresponds to the PostResolveRequestCache stage of the ASP.NET integrated pipeline.
+            </summary>
+        </member>
+        <member name="F:Owin.PipelineStage.MapHandler">
+            <summary>
+            Corresponds to the MapRequestHandler stage of the ASP.NET integrated pipeline.
+            </summary>
+        </member>
+        <member name="F:Owin.PipelineStage.PostMapHandler">
+            <summary>
+            Corresponds to the PostMapRequestHandler stage of the ASP.NET integrated pipeline.
+            </summary>
+        </member>
+        <member name="F:Owin.PipelineStage.AcquireState">
+            <summary>
+            Corresponds to the AcquireRequestState stage of the ASP.NET integrated pipeline.
+            </summary>
+        </member>
+        <member name="F:Owin.PipelineStage.PostAcquireState">
+            <summary>
+            Corresponds to the PostAcquireRequestState stage of the ASP.NET integrated pipeline.
+            </summary>
+        </member>
+        <member name="F:Owin.PipelineStage.PreHandlerExecute">
+            <summary>
+            Corresponds to the PreRequestHandlerExecute stage of the ASP.NET integrated pipeline.
+            </summary>
+        </member>
+        <member name="T:Owin.MapExtensions">
+            <summary>
+            Extension methods for the MapMiddleware
+            </summary>
+        </member>
+        <member name="M:Owin.MapExtensions.Map(Owin.IAppBuilder,System.String,System.Action{Owin.IAppBuilder})">
+            <summary>
+            If the request path starts with the given pathMatch, execute the app configured via configuration parameter instead of
+            continuing to the next component in the pipeline.
+            </summary>
+            <param name="app"></param>
+            <param name="pathMatch">The path to match</param>
+            <param name="configuration">The branch to take for positive path matches</param>
+            <returns></returns>
+        </member>
+        <member name="M:Owin.MapExtensions.Map(Owin.IAppBuilder,Microsoft.Owin.PathString,System.Action{Owin.IAppBuilder})">
+            <summary>
+            If the request path starts with the given pathMatch, execute the app configured via configuration parameter instead of
+            continuing to the next component in the pipeline.
+            </summary>
+            <param name="app"></param>
+            <param name="pathMatch">The path to match</param>
+            <param name="configuration">The branch to take for positive path matches</param>
+            <returns></returns>
+        </member>
+        <member name="T:Microsoft.Owin.Mapping.MapMiddleware">
+            <summary>
+            Used to create path based branches in your application pipeline.
+            The owin.RequestPathBase is not included in the evaluation, only owin.RequestPath.
+            Matching paths have the matching piece removed from owin.RequestPath and added to the owin.RequestPathBase.
+            </summary>
+        </member>
+        <member name="M:Microsoft.Owin.Mapping.MapMiddleware.#ctor(System.Func{System.Collections.Generic.IDictionary{System.String,System.Object},System.Threading.Tasks.Task},Microsoft.Owin.Mapping.MapOptions)">
+            <summary>
+            Initializes a new instance of the <see cref="T:Microsoft.Owin.Mapping.MapMiddleware"/> class
+            </summary>
+            <param name="next">The normal pipeline taken for a negative match</param>
+            <param name="options"></param>
+        </member>
+        <member name="M:Microsoft.Owin.Mapping.MapMiddleware.Invoke(System.Collections.Generic.IDictionary{System.String,System.Object})">
+            <summary>
+            Process an individual request.
+            </summary>
+            <param name="environment"></param>
+            <returns></returns>
+        </member>
+        <member name="T:Microsoft.Owin.Mapping.MapOptions">
+            <summary>
+            Options for the Map middleware
+            </summary>
+        </member>
+        <member name="P:Microsoft.Owin.Mapping.MapOptions.PathMatch">
+            <summary>
+            The path to match
+            </summary>
+        </member>
+        <member name="P:Microsoft.Owin.Mapping.MapOptions.Branch">
+            <summary>
+            The branch taken for a positive match
+            </summary>
+        </member>
+        <member name="T:Owin.MapWhenExtensions">
+            <summary>
+            Extension methods for the MapWhenMiddleware
+            </summary>
+        </member>
+        <member name="M:Owin.MapWhenExtensions.MapWhen(Owin.IAppBuilder,System.Func{Microsoft.Owin.IOwinContext,System.Boolean},System.Action{Owin.IAppBuilder})">
+            <summary>
+            Branches the request pipeline based on the result of the given predicate.
+            </summary>
+            <param name="app"></param>
+            <param name="predicate">Invoked with the request environment to determine if the branch should be taken</param>
+            <param name="configuration">Configures a branch to take</param>
+            <returns></returns>
+        </member>
+        <member name="M:Owin.MapWhenExtensions.MapWhenAsync(Owin.IAppBuilder,System.Func{Microsoft.Owin.IOwinContext,System.Threading.Tasks.Task{System.Boolean}},System.Action{Owin.IAppBuilder})">
+            <summary>
+            Branches the request pipeline based on the async result of the given predicate.
+            </summary>
+            <param name="app"></param>
+            <param name="predicate">Invoked asynchronously with the request environment to determine if the branch should be taken</param>
+            <param name="configuration">Configures a branch to take</param>
+            <returns></returns>
+        </member>
+        <member name="T:Microsoft.Owin.Mapping.MapWhenMiddleware">
+            <summary>
+            Determines if the request should take a specific branch of the pipeline by passing the environment
+            to a user defined callback.
+            </summary>
+        </member>
+        <member name="M:Microsoft.Owin.Mapping.MapWhenMiddleware.#ctor(System.Func{System.Collections.Generic.IDictionary{System.String,System.Object},System.Threading.Tasks.Task},Microsoft.Owin.Mapping.MapWhenOptions)">
+            <summary>
+            Initializes a new instance of the <see cref="T:Microsoft.Owin.Mapping.MapWhenMiddleware"/> class
+            </summary>
+            <param name="next">The normal application pipeline</param>
+            <param name="options"></param>
+        </member>
+        <member name="M:Microsoft.Owin.Mapping.MapWhenMiddleware.Invoke(System.Collections.Generic.IDictionary{System.String,System.Object})">
+            <summary>
+            Process an individual request.
+            </summary>
+            <param name="environment"></param>
+            <returns></returns>
+        </member>
+        <member name="T:Microsoft.Owin.Mapping.MapWhenOptions">
+            <summary>
+            Options for the MapWhen middleware
+            </summary>
+        </member>
+        <member name="P:Microsoft.Owin.Mapping.MapWhenOptions.Predicate">
+            <summary>
+            The user callback that determines if the branch should be taken
+            </summary>
+        </member>
+        <member name="P:Microsoft.Owin.Mapping.MapWhenOptions.PredicateAsync">
+            <summary>
+            The async user callback that determines if the branch should be taken
+            </summary>
+        </member>
+        <member name="P:Microsoft.Owin.Mapping.MapWhenOptions.Branch">
+            <summary>
+            The branch taken for a positive match
+            </summary>
+        </member>
+        <member name="T:Microsoft.Owin.OwinContext">
+            <summary>
+            This wraps OWIN environment dictionary and provides strongly typed accessors.
+            </summary>
+        </member>
+        <member name="M:Microsoft.Owin.OwinContext.#ctor">
+            <summary>
+            Create a new context with only request and response header collections.
+            </summary>
+        </member>
+        <member name="M:Microsoft.Owin.OwinContext.#ctor(System.Collections.Generic.IDictionary{System.String,System.Object})">
+            <summary>
+            Create a new wrapper.
+            </summary>
+            <param name="environment">OWIN environment dictionary which stores state information about the request, response and relevant server state.</param>
+        </member>
+        <member name="M:Microsoft.Owin.OwinContext.Get``1(System.String)">
+            <summary>
+            Gets a value from the OWIN environment, or returns default(T) if not present.
+            </summary>
+            <typeparam name="T">The type of the value.</typeparam>
+            <param name="key">The key of the value to get.</param>
+            <returns>The value with the specified key or the default(T) if not present.</returns>
+        </member>
+        <member name="M:Microsoft.Owin.OwinContext.Set``1(System.String,``0)">
+            <summary>
+            Sets the given key and value in the OWIN environment.
+            </summary>
+            <typeparam name="T">The type of the value.</typeparam>
+            <param name="key">The key of the value to set.</param>
+            <param name="value">The value to set.</param>
+            <returns>This instance.</returns>
+        </member>
+        <member name="P:Microsoft.Owin.OwinContext.Request">
+            <summary>
+            Gets a wrapper exposing request specific properties.
+            </summary>
+            <returns>A wrapper exposing request specific properties.</returns>
+        </member>
+        <member name="P:Microsoft.Owin.OwinContext.Response">
+            <summary>
+            Gets a wrapper exposing response specific properties.
+            </summary>
+            <returns>A wrapper exposing response specific properties.</returns>
+        </member>
+        <member name="P:Microsoft.Owin.OwinContext.Authentication">
+            <summary>
+            Gets the Authentication middleware functionality available on the current request.
+            </summary>
+            <returns>The authentication middleware functionality available on the current request.</returns>
+        </member>
+        <member name="P:Microsoft.Owin.OwinContext.Environment">
+            <summary>
+            Gets the OWIN environment.
+            </summary>
+            <returns>The OWIN environment.</returns>
+        </member>
+        <member name="P:Microsoft.Owin.OwinContext.TraceOutput">
+            <summary>
+            Gets or sets the host.TraceOutput environment value.
+            </summary>
+            <returns>The host.TraceOutput TextWriter.</returns>
+        </member>
+        <member name="T:Microsoft.Owin.PathString">
+            <summary>
+            Provides correct escaping for Path and PathBase values when needed to reconstruct a request or redirect URI string
+            </summary>
+        </member>
+        <member name="F:Microsoft.Owin.PathString.Empty">
+            <summary>
+            Represents the empty path. This field is read-only.
+            </summary>
+        </member>
+        <member name="M:Microsoft.Owin.PathString.#ctor(System.String)">
+            <summary>
+            Initialize the path string with a given value. This value must be in un-escaped format. Use
+            PathString.FromUriComponent(value) if you have a path value which is in an escaped format.
+            </summary>
+            <param name="value">The unescaped path to be assigned to the Value property.</param>
+        </member>
+        <member name="M:Microsoft.Owin.PathString.ToString">
+            <summary>
+            Provides the path string escaped in a way which is correct for combining into the URI representation. 
+            </summary>
+            <returns>The escaped path value</returns>
+        </member>
+        <member name="M:Microsoft.Owin.PathString.ToUriComponent">
+            <summary>
+            Provides the path string escaped in a way which is correct for combining into the URI representation.
+            </summary>
+            <returns>The escaped path value</returns>
+        </member>
+        <member name="M:Microsoft.Owin.PathString.FromUriComponent(System.String)">
+            <summary>
+            Returns an PathString given the path as it is escaped in the URI format. The string MUST NOT contain any
+            value that is not a path.
+            </summary>
+            <param name="uriComponent">The escaped path as it appears in the URI format.</param>
+            <returns>The resulting PathString</returns>
+        </member>
+        <member name="M:Microsoft.Owin.PathString.FromUriComponent(System.Uri)">
+            <summary>
+            Returns an PathString given the path as from a Uri object. Relative Uri objects are not supported.
+            </summary>
+            <param name="uri">The Uri object</param>
+            <returns>The resulting PathString</returns>
+        </member>
+        <member name="M:Microsoft.Owin.PathString.StartsWithSegments(Microsoft.Owin.PathString)">
+            <summary>
+            Checks if this instance starts with or exactly matches the other instance. Only full segments are matched.
+            </summary>
+            <param name="other"></param>
+            <returns></returns>
+        </member>
+        <member name="M:Microsoft.Owin.PathString.StartsWithSegments(Microsoft.Owin.PathString,Microsoft.Owin.PathString@)">
+            <summary>
+            Checks if this instance starts with or exactly matches the other instance. Only full segments are matched.
+            </summary>
+            <param name="other"></param>
+            <param name="remaining">Any remaining segments from this instance not included in the other instance.</param>
+            <returns></returns>
+        </member>
+        <member name="M:Microsoft.Owin.PathString.Add(Microsoft.Owin.PathString)">
+            <summary>
+            Adds two PathString instances into a combined PathString value. 
+            </summary>
+            <returns>The combined PathString value</returns>
+        </member>
+        <member name="M:Microsoft.Owin.PathString.Add(Microsoft.Owin.QueryString)">
+            <summary>
+            Combines a PathString and QueryString into the joined URI formatted string value. 
+            </summary>
+            <returns>The joined URI formatted string value</returns>
+        </member>
+        <member name="M:Microsoft.Owin.PathString.Equals(Microsoft.Owin.PathString)">
+            <summary>
+            Compares this PathString value to another value. The default comparison is StringComparison.OrdinalIgnoreCase.
+            </summary>
+            <param name="other">The second PathString for comparison.</param>
+            <returns>True if both PathString values are equal</returns>
+        </member>
+        <member name="M:Microsoft.Owin.PathString.Equals(Microsoft.Owin.PathString,System.StringComparison)">
+            <summary>
+            Compares this PathString value to another value using a specific StringComparison type
+            </summary>
+            <param name="other">The second PathString for comparison</param>
+            <param name="comparisonType">The StringComparison type to use</param>
+            <returns>True if both PathString values are equal</returns>
+        </member>
+        <member name="M:Microsoft.Owin.PathString.Equals(System.Object)">
+            <summary>
+            Compares this PathString value to another value. The default comparison is StringComparison.OrdinalIgnoreCase.
+            </summary>
+            <param name="obj">The second PathString for comparison.</param>
+            <returns>True if both PathString values are equal</returns>
+        </member>
+        <member name="M:Microsoft.Owin.PathString.GetHashCode">
+            <summary>
+            Returns the hash code for the PathString value. The hash code is provided by the OrdinalIgnoreCase implementation.
+            </summary>
+            <returns>The hash code</returns>
+        </member>
+        <member name="M:Microsoft.Owin.PathString.op_Equality(Microsoft.Owin.PathString,Microsoft.Owin.PathString)">
+            <summary>
+            Operator call through to Equals
+            </summary>
+            <param name="left">The left parameter</param>
+            <param name="right">The right parameter</param>
+            <returns>True if both PathString values are equal</returns>
+        </member>
+        <member name="M:Microsoft.Owin.PathString.op_Inequality(Microsoft.Owin.PathString,Microsoft.Owin.PathString)">
+            <summary>
+            Operator call through to Equals
+            </summary>
+            <param name="left">The left parameter</param>
+            <param name="right">The right parameter</param>
+            <returns>True if both PathString values are not equal</returns>
+        </member>
+        <member name="M:Microsoft.Owin.PathString.op_Addition(Microsoft.Owin.PathString,Microsoft.Owin.PathString)">
+            <summary>
+            Operator call through to Add
+            </summary>
+            <param name="left">The left parameter</param>
+            <param name="right">The right parameter</param>
+            <returns>The PathString combination of both values</returns>
+        </member>
+        <member name="M:Microsoft.Owin.PathString.op_Addition(Microsoft.Owin.PathString,Microsoft.Owin.QueryString)">
+            <summary>
+            Operator call through to Add
+            </summary>
+            <param name="left">The left parameter</param>
+            <param name="right">The right parameter</param>
+            <returns>The PathString combination of both values</returns>
+        </member>
+        <member name="P:Microsoft.Owin.PathString.Value">
+            <summary>
+            The unescaped path value
+            </summary>
+        </member>
+        <member name="P:Microsoft.Owin.PathString.HasValue">
+            <summary>
+            True if the path is not empty
+            </summary>
+        </member>
+        <member name="T:Microsoft.Owin.QueryString">
+            <summary>
+            Provides correct handling for QueryString value when needed to reconstruct a request or redirect URI string
+            </summary>
+        </member>
+        <member name="F:Microsoft.Owin.QueryString.Empty">
+            <summary>
+            Represents the empty query string. This field is read-only.
+            </summary>
+        </member>
+        <member name="M:Microsoft.Owin.QueryString.#ctor(System.String)">
+            <summary>
+            Initialize the query string with a given value. This value must be in escaped and delimited format without
+            a leading '?' character. 
+            </summary>
+            <param name="value">The query string to be assigned to the Value property.</param>
+        </member>
+        <member name="M:Microsoft.Owin.QueryString.#ctor(System.String,System.String)">
+            <summary>
+            Initialize a query string with a single given parameter name and value. The value is 
+            </summary>
+            <param name="name">The unencoded parameter name</param>
+            <param name="value">The unencoded parameter value</param>
+        </member>
+        <member name="M:Microsoft.Owin.QueryString.ToString">
+            <summary>
+            Provides the query string escaped in a way which is correct for combining into the URI representation. 
+            A leading '?' character will be prepended unless the Value is null or empty. Characters which are potentially
+            dangerous are escaped.
+            </summary>
+            <returns>The query string value</returns>
+        </member>
+        <member name="M:Microsoft.Owin.QueryString.ToUriComponent">
+            <summary>
+            Provides the query string escaped in a way which is correct for combining into the URI representation. 
+            A leading '?' character will be prepended unless the Value is null or empty. Characters which are potentially
+            dangerous are escaped.
+            </summary>
+            <returns>The query string value</returns>
+        </member>
+        <member name="M:Microsoft.Owin.QueryString.FromUriComponent(System.String)">
+            <summary>
+            Returns an QueryString given the query as it is escaped in the URI format. The string MUST NOT contain any
+            value that is not a query.
+            </summary>
+            <param name="uriComponent">The escaped query as it appears in the URI format.</param>
+            <returns>The resulting QueryString</returns>
+        </member>
+        <member name="M:Microsoft.Owin.QueryString.FromUriComponent(System.Uri)">
+            <summary>
+            Returns an QueryString given the query as from a Uri object. Relative Uri objects are not supported.
+            </summary>
+            <param name="uri">The Uri object</param>
+            <returns>The resulting QueryString</returns>
+        </member>
+        <member name="M:Microsoft.Owin.QueryString.Equals(Microsoft.Owin.QueryString)">
+            <summary>
+            Indicates whether the current instance is equal to the other instance.
+            </summary>
+            <param name="other"></param>
+            <returns></returns>
+        </member>
+        <member name="M:Microsoft.Owin.QueryString.Equals(System.Object)">
+            <summary>
+            Indicates whether the current instance is equal to the other instance.
+            </summary>
+            <param name="obj"></param>
+            <returns></returns>
+        </member>
+        <member name="M:Microsoft.Owin.QueryString.GetHashCode">
+            <summary>
+            Returns the hash code for this instance.
+            </summary>
+            <returns></returns>
+        </member>
+        <member name="M:Microsoft.Owin.QueryString.op_Equality(Microsoft.Owin.QueryString,Microsoft.Owin.QueryString)">
+            <summary>
+            Compares the two instances for equality.
+            </summary>
+            <param name="left"></param>
+            <param name="right"></param>
+            <returns></returns>
+        </member>
+        <member name="M:Microsoft.Owin.QueryString.op_Inequality(Microsoft.Owin.QueryString,Microsoft.Owin.QueryString)">
+            <summary>
+            Compares the two instances for inequality.
+            </summary>
+            <param name="left"></param>
+            <param name="right"></param>
+            <returns></returns>
+        </member>
+        <member name="P:Microsoft.Owin.QueryString.Value">
+            <summary>
+            The escaped query string without the leading '?' character
+            </summary>
+        </member>
+        <member name="P:Microsoft.Owin.QueryString.HasValue">
+            <summary>
+            True if the query string is not empty
+            </summary>
+        </member>
+        <member name="T:Microsoft.Owin.RequestCookieCollection">
+            <summary>
+            A wrapper for the request Cookie header
+            </summary>
+        </member>
+        <member name="M:Microsoft.Owin.RequestCookieCollection.#ctor(System.Collections.Generic.IDictionary{System.String,System.String})">
+            <summary>
+            Create a new wrapper
+            </summary>
+            <param name="store"></param>
+        </member>
+        <member name="M:Microsoft.Owin.RequestCookieCollection.GetEnumerator">
+            <summary>
+            
+            </summary>
+            <returns></returns>
+        </member>
+        <member name="M:Microsoft.Owin.RequestCookieCollection.System#Collections#IEnumerable#GetEnumerator">
+            <summary>
+            
+            </summary>
+            <returns></returns>
+        </member>
+        <member name="P:Microsoft.Owin.RequestCookieCollection.Item(System.String)">
+            <summary>
+            Returns null rather than throwing KeyNotFoundException
+            </summary>
+            <param name="key"></param>
+            <returns></returns>
+        </member>
+        <member name="T:Microsoft.Owin.Resources">
+            <summary>
+              A strongly-typed resource class, for looking up localized strings, etc.
+            </summary>
+        </member>
+        <member name="P:Microsoft.Owin.Resources.ResourceManager">
+            <summary>
+              Returns the cached ResourceManager instance used by this class.
+            </summary>
+        </member>
+        <member name="P:Microsoft.Owin.Resources.Culture">
+            <summary>
+              Overrides the current thread's CurrentUICulture property for all
+              resource lookups using this strongly typed resource class.
+            </summary>
+        </member>
+        <member name="P:Microsoft.Owin.Resources.Exception_ConversionTakesOneParameter">
+            <summary>
+              Looks up a localized string similar to Conversion delegate must take one parameter..
+            </summary>
+        </member>
+        <member name="P:Microsoft.Owin.Resources.Exception_CookieLimitTooSmall">
+            <summary>
+              Looks up a localized string similar to The cookie key and options are larger than ChunksSize, leaving no room for data..
+            </summary>
+        </member>
+        <member name="P:Microsoft.Owin.Resources.Exception_ImcompleteChunkedCookie">
+            <summary>
+              Looks up a localized string similar to The chunked cookie is incomplete. Only {0} of the expected {1} chunks were found, totaling {2} characters. A client size limit may have been exceeded..
+            </summary>
+        </member>
+        <member name="P:Microsoft.Owin.Resources.Exception_MiddlewareNotSupported">
+            <summary>
+              Looks up a localized string similar to The type &apos;{0}&apos; does not match any known middleware pattern..
+            </summary>
+        </member>
+        <member name="P:Microsoft.Owin.Resources.Exception_MissingOnSendingHeaders">
+            <summary>
+              Looks up a localized string similar to The OWIN key &apos;server.OnSendingHeaders&apos; is not available for this request..
+            </summary>
+        </member>
+        <member name="P:Microsoft.Owin.Resources.Exception_NoConstructorFound">
+            <summary>
+              Looks up a localized string similar to The class &apos;{0}&apos; does not have a constructor taking {1} arguments..
+            </summary>
+        </member>
+        <member name="P:Microsoft.Owin.Resources.Exception_NoConversionExists">
+            <summary>
+              Looks up a localized string similar to No conversion available between {0} and {1}..
+            </summary>
+        </member>
+        <member name="P:Microsoft.Owin.Resources.Exception_PathMustNotEndWithSlash">
+            <summary>
+              Looks up a localized string similar to The path must not end with a &apos;/&apos;.
+            </summary>
+        </member>
+        <member name="P:Microsoft.Owin.Resources.Exception_PathMustStartWithSlash">
+            <summary>
+              Looks up a localized string similar to The path must start with a &apos;/&apos; followed by one or more characters..
+            </summary>
+        </member>
+        <member name="P:Microsoft.Owin.Resources.Exception_PathRequired">
+            <summary>
+              Looks up a localized string similar to The path is required.
+            </summary>
+        </member>
+        <member name="P:Microsoft.Owin.Resources.Exception_QueryStringMustStartWithDelimiter">
+            <summary>
+              Looks up a localized string similar to The query string must start with a &apos;?&apos; unless null or empty..
+            </summary>
+        </member>
+        <member name="T:Microsoft.Owin.Builder.AppBuilder">
+            <summary>
+            A standard implementation of IAppBuilder 
+            </summary>
+        </member>
+        <member name="M:Microsoft.Owin.Builder.AppBuilder.#ctor">
+            <summary>
+            Initializes a new instance of the the type.
+            </summary>
+        </member>
+        <member name="M:Microsoft.Owin.Builder.AppBuilder.#ctor(System.Collections.Generic.IDictionary{System.Tuple{System.Type,System.Type},System.Delegate},System.Collections.Generic.IDictionary{System.String,System.Object})">
+            <summary>
+            
+            </summary>
+            <param name="conversions"></param>
+            <param name="properties"></param>
+        </member>
+        <member name="M:Microsoft.Owin.Builder.AppBuilder.Use(System.Object,System.Object[])">
+            <summary>
+            Adds a middleware node to the OWIN function pipeline. The middleware are
+            invoked in the order they are added: the first middleware passed to Use will
+            be the outermost function, and the last middleware passed to Use will be the
+            innermost.
+            </summary>
+            <param name="middleware">
+            The middleware parameter determines which behavior is being chained into the
+            pipeline. 
+            
+            If the middleware given to Use is a Delegate, then it will be invoked with the "next app" in 
+            the chain as the first parameter. If the delegate takes more than the single argument, 
+            then the additional values must be provided to Use in the args array.
+            
+            If the middleware given to Use is a Type, then the public constructor will be 
+            invoked with the "next app" in the chain as the first parameter. The resulting object
+            must have a public Invoke method. If the object has constructors which take more than
+            the single "next app" argument, then additional values may be provided in the args array.
+            </param>
+            <param name="args">
+            Any additional args passed to Use will be passed as additional values, following the "next app"
+            parameter, when the OWIN call pipeline is build.
+            
+            They are passed as additional parameters if the middleware parameter is a Delegate, or as additional
+            constructor arguments if the middle parameter is a Type.
+            </param>
+            <returns>
+            The IAppBuilder itself is returned. This enables you to chain your use statements together.
+            </returns>
+        </member>
+        <member name="M:Microsoft.Owin.Builder.AppBuilder.New">
+            <summary>
+            The New method creates a new instance of an IAppBuilder. This is needed to create
+            a tree structure in your processing, rather than a linear pipeline. The new instance share the
+            same Properties, but will be created with a new, empty middleware list.
+            
+            To create a tangent pipeline you would first call New, followed by several calls to Use on 
+            the new builder, ending with a call to Build on the new builder. The return value from Build
+            will be the entry-point to your tangent pipeline. This entry-point may now be added to the
+            main pipeline as an argument to a switching middleware, which will either call the tangent
+            pipeline or the "next app", based on something in the request.
+            
+            That said - all of that work is typically hidden by a middleware like Map, which will do that
+            for you.
+            </summary>
+            <returns>The new instance of the IAppBuilder implementation</returns>
+        </member>
+        <member name="M:Microsoft.Owin.Builder.AppBuilder.Build(System.Type)">
+            <summary>
+            The Build is called at the point when all of the middleware should be chained
+            together. This is typically done by the hosting component which created the app builder,
+            and does not need to be called by the startup method if the IAppBuilder is passed in.
+            </summary>
+            <param name="returnType">
+            The Type argument indicates which calling convention should be returned, and
+            is typically typeof(<typeref name="Func&lt;IDictionary&lt;string,object&gt;, Task&gt;"/>) for the OWIN
+            calling convention.
+            </param>
+            <returns>
+            Returns an instance of the pipeline's entry point. This object may be safely cast to the
+            type which was provided
+            </returns>
+        </member>
+        <member name="P:Microsoft.Owin.Builder.AppBuilder.Properties">
+            <summary>
+            Contains arbitrary properties which may added, examined, and modified by
+            components during the startup sequence. 
+            </summary>
+            <returns>Returns <see cref="T:System.Collections.Generic.IDictionary`2" />.</returns>
+        </member>
+        <member name="T:Microsoft.Owin.Builder.NotFound">
+            <summary>
+            Simple object used by AppBuilder as seed OWIN callable if the
+            builder.Properties["builder.DefaultApp"] is not set
+            </summary>
+        </member>
+        <member name="T:Microsoft.Owin.Infrastructure.ISystemClock">
+            <summary>
+            Abstracts the system clock to facilitate testing.
+            </summary>
+        </member>
+        <member name="P:Microsoft.Owin.Infrastructure.ISystemClock.UtcNow">
+            <summary>
+            Retrieves the current system time in UTC.
+            </summary>
+        </member>
+        <member name="T:Microsoft.Owin.Infrastructure.SystemClock">
+            <summary>
+            Provides access to the normal system clock.
+            </summary>
+        </member>
+        <member name="P:Microsoft.Owin.Infrastructure.SystemClock.UtcNow">
+            <summary>
+            Retrieves the current system time in UTC.
+            </summary>
+        </member>
+        <member name="T:Microsoft.Owin.Infrastructure.WebUtilities">
+            <summary>
+            Response generation utilities.
+            </summary>
+        </member>
+        <member name="M:Microsoft.Owin.Infrastructure.WebUtilities.AddQueryString(System.String,System.String)">
+            <summary>
+            Append the given query to the uri.
+            </summary>
+            <param name="uri">The base uri.</param>
+            <param name="queryString">The query string to append, if any.</param>
+            <returns>The combine result.</returns>
+        </member>
+        <member name="M:Microsoft.Owin.Infrastructure.WebUtilities.AddQueryString(System.String,System.String,System.String)">
+            <summary>
+            Append the given query key and value to the uri.
+            </summary>
+            <param name="uri">The base uri.</param>
+            <param name="name">The name of the query key.</param>
+            <param name="value">The query value.</param>
+            <returns>The combine result.</returns>
+        </member>
+        <member name="M:Microsoft.Owin.Infrastructure.WebUtilities.AddQueryString(System.String,System.Collections.Generic.IDictionary{System.String,System.String})">
+            <summary>
+            Append the given query keys and values to the uri.
+            </summary>
+            <param name="uri">The base uri.</param>
+            <param name="queryString">A collection of name value query pairs to append.</param>
+            <returns>The combine result.</returns>
+        </member>
+        <member name="T:Microsoft.Owin.ResponseCookieCollection">
+            <summary>
+            A wrapper for the response Set-Cookie header
+            </summary>
+        </member>
+        <member name="M:Microsoft.Owin.ResponseCookieCollection.#ctor(Microsoft.Owin.IHeaderDictionary)">
+            <summary>
+            Create a new wrapper
+            </summary>
+            <param name="headers"></param>
+        </member>
+        <member name="M:Microsoft.Owin.ResponseCookieCollection.Append(System.String,System.String)">
+            <summary>
+            Add a new cookie and value
+            </summary>
+            <param name="key"></param>
+            <param name="value"></param>
+        </member>
+        <member name="M:Microsoft.Owin.ResponseCookieCollection.Append(System.String,System.String,Microsoft.Owin.CookieOptions)">
+            <summary>
+            Add a new cookie
+            </summary>
+            <param name="key"></param>
+            <param name="value"></param>
+            <param name="options"></param>
+        </member>
+        <member name="M:Microsoft.Owin.ResponseCookieCollection.Delete(System.String)">
+            <summary>
+            Sets an expired cookie
+            </summary>
+            <param name="key"></param>
+        </member>
+        <member name="M:Microsoft.Owin.ResponseCookieCollection.Delete(System.String,Microsoft.Owin.CookieOptions)">
+            <summary>
+            Sets an expired cookie
+            </summary>
+            <param name="key"></param>
+            <param name="options"></param>
+        </member>
+        <member name="T:Microsoft.Owin.Security.AuthenticateResult">
+            <summary>
+            Acts as the return value from calls to the IAuthenticationManager's AuthenticeAsync methods.
+            </summary>
+        </member>
+        <member name="M:Microsoft.Owin.Security.AuthenticateResult.#ctor(System.Security.Principal.IIdentity,Microsoft.Owin.Security.AuthenticationProperties,Microsoft.Owin.Security.AuthenticationDescription)">
+            <summary>
+            Create an instance of the result object
+            </summary>
+            <param name="identity">Assigned to Identity. May be null.</param>
+            <param name="properties">Assigned to Properties. Contains extra information carried along with the identity.</param>
+            <param name="description">Assigned to Description. Contains information describing the authentication provider.</param>
+        </member>
+        <member name="P:Microsoft.Owin.Security.AuthenticateResult.Identity">
+            <summary>
+            Contains the claims that were authenticated by the given AuthenticationType. If the authentication
+            type was not successful the Identity property will be null.
+            </summary>
+        </member>
+        <member name="P:Microsoft.Owin.Security.AuthenticateResult.Properties">
+            <summary>
+            Contains extra values that were provided with the original SignIn call.
+            </summary>
+        </member>
+        <member name="P:Microsoft.Owin.Security.AuthenticateResult.Description">
+            <summary>
+            Contains description properties for the middleware authentication type in general. Does not
+            vary per request.
+            </summary>
+        </member>
+        <member name="T:Microsoft.Owin.Security.AuthenticationDescription">
+            <summary>
+            Contains information describing an authentication provider.
+            </summary>
+        </member>
+        <member name="M:Microsoft.Owin.Security.AuthenticationDescription.#ctor">
+            <summary>
+            Initializes a new instance of the <see cref="T:Microsoft.Owin.Security.AuthenticationDescription"/> class
+            </summary>
+        </member>
+        <member name="M:Microsoft.Owin.Security.AuthenticationDescription.#ctor(System.Collections.Generic.IDictionary{System.String,System.Object})">
+            <summary>
+            Initializes a new instance of the <see cref="T:Microsoft.Owin.Security.AuthenticationDescription"/> class
+            </summary>
+            <param name="properties"></param>
+        </member>
+        <member name="P:Microsoft.Owin.Security.AuthenticationDescription.Properties">
+            <summary>
+            Contains metadata about the authentication provider.
+            </summary>
+        </member>
+        <member name="P:Microsoft.Owin.Security.AuthenticationDescription.AuthenticationType">
+            <summary>
+            Gets or sets the name used to reference the authentication middleware instance.
+            </summary>
+        </member>
+        <member name="P:Microsoft.Owin.Security.AuthenticationDescription.Caption">
+            <summary>
+            Gets or sets the display name for the authentication provider.
+            </summary>
+        </member>
+        <member name="T:Microsoft.Owin.Extensions.IntegratedPipelineExtensions">
+            <summary>
+            Extension methods used to indicate at which stage in the integrated pipeline prior middleware should run.
+            </summary>
+        </member>
+        <member name="M:Microsoft.Owin.Extensions.IntegratedPipelineExtensions.UseStageMarker(Owin.IAppBuilder,System.String)">
+            <summary>
+            Call after other middleware to specify when they should run in the integrated pipeline.
+            </summary>
+            <param name="app">The IAppBuilder.</param>
+            <param name="stageName">The name of the integrated pipeline in which to run.</param>
+            <returns>The original IAppBuilder for chaining.</returns>
+        </member>
+        <member name="M:Microsoft.Owin.Extensions.IntegratedPipelineExtensions.UseStageMarker(Owin.IAppBuilder,Owin.PipelineStage)">
+            <summary>
+            Call after other middleware to specify when they should run in the integrated pipeline.
+            </summary>
+            <param name="app">The IAppBuilder.</param>
+            <param name="stage">The stage of the integrated pipeline in which to run.</param>
+            <returns>The original IAppBuilder for chaining.</returns>
+        </member>
+        <member name="T:Microsoft.Owin.Logging.AppBuilderLoggerExtensions">
+            <summary>
+            Logging extension methods for IAppBuilder.
+            </summary>
+        </member>
+        <member name="M:Microsoft.Owin.Logging.AppBuilderLoggerExtensions.SetLoggerFactory(Owin.IAppBuilder,Microsoft.Owin.Logging.ILoggerFactory)">
+            <summary>
+            Sets the server.LoggerFactory in the Properties collection.
+            </summary>
+            <param name="app"></param>
+            <param name="loggerFactory"></param>
+        </member>
+        <member name="M:Microsoft.Owin.Logging.AppBuilderLoggerExtensions.GetLoggerFactory(Owin.IAppBuilder)">
+            <summary>
+            Retrieves the server.LoggerFactory from the Properties collection.
+            </summary>
+            <param name="app"></param>
+            <returns></returns>
+        </member>
+        <member name="M:Microsoft.Owin.Logging.AppBuilderLoggerExtensions.CreateLogger(Owin.IAppBuilder,System.String)">
+            <summary>
+            Creates a new ILogger instance from the server.LoggerFactory in the Properties collection.
+            </summary>
+            <param name="app"></param>
+            <param name="name"></param>
+            <returns></returns>
+        </member>
+        <member name="M:Microsoft.Owin.Logging.AppBuilderLoggerExtensions.CreateLogger(Owin.IAppBuilder,System.Type)">
+            <summary>
+            Creates a new ILogger instance from the server.LoggerFactory in the Properties collection.
+            </summary>
+            <param name="app"></param>
+            <param name="component"></param>
+            <returns></returns>
+        </member>
+        <member name="M:Microsoft.Owin.Logging.AppBuilderLoggerExtensions.CreateLogger``1(Owin.IAppBuilder)">
+            <summary>
+            Creates a new ILogger instance from the server.LoggerFactory in the Properties collection.
+            </summary>
+            <typeparam name="TType"></typeparam>
+            <param name="app"></param>
+            <returns></returns>
+        </member>
+        <member name="T:Microsoft.Owin.Logging.ILoggerFactory">
+            <summary>
+            Used to create logger instances of the given name.
+            </summary>
+        </member>
+        <member name="M:Microsoft.Owin.Logging.ILoggerFactory.Create(System.String)">
+            <summary>
+            Creates a new ILogger instance of the given name.
+            </summary>
+            <param name="name"></param>
+            <returns></returns>
+        </member>
+        <member name="T:Microsoft.Owin.Logging.ILogger">
+            <summary>
+            A generic interface for logging.
+            </summary>
+        </member>
+        <member name="M:Microsoft.Owin.Logging.ILogger.WriteCore(System.Diagnostics.TraceEventType,System.Int32,System.Object,System.Exception,System.Func{System.Object,System.Exception,System.String})">
+            <summary>
+            Aggregates most logging patterns to a single method.  This must be compatible with the Func representation in the OWIN environment.
+            
+            To check IsEnabled call WriteCore with only TraceEventType and check the return value, no event will be written.
+            </summary>
+            <param name="eventType"></param>
+            <param name="eventId"></param>
+            <param name="state"></param>
+            <param name="exception"></param>
+            <param name="formatter"></param>
+            <returns></returns>
+        </member>
+        <member name="T:Microsoft.Owin.BuilderProperties.Address">
+            <summary>
+            Contains the parts of an address.
+            </summary>
+        </member>
+        <member name="M:Microsoft.Owin.BuilderProperties.Address.#ctor(System.Collections.Generic.IDictionary{System.String,System.Object})">
+            <summary>
+            Initializes a new instance.
+            </summary>
+            <param name="dictionary"></param>
+        </member>
+        <member name="M:Microsoft.Owin.BuilderProperties.Address.#ctor(System.String,System.String,System.String,System.String)">
+            <summary>
+            Initializes a new <see cref="T:Microsoft.Owin.BuilderProperties.Address"/> with the given parts.
+            </summary>
+            <param name="scheme">The scheme.</param>
+            <param name="host">The host.</param>
+            <param name="port">The port.</param>
+            <param name="path">The path.</param>
+        </member>
+        <member name="M:Microsoft.Owin.BuilderProperties.Address.Create">
+            <summary>
+            Creates a new <see cref="T:Microsoft.Owin.BuilderProperties.Address"/>
+            </summary>
+            <returns>A new <see cref="T:Microsoft.Owin.BuilderProperties.Address" /></returns>
+        </member>
+        <member name="M:Microsoft.Owin.BuilderProperties.Address.Equals(Microsoft.Owin.BuilderProperties.Address)">
+            <summary>
+            Determines whether the specified object is equal to the current object.
+            </summary>
+            <param name="other">The other object.</param>
+            <returns>true if the specified object is equal to the current object; otherwise, false.</returns>
+        </member>
+        <member name="M:Microsoft.Owin.BuilderProperties.Address.Equals(System.Object)">
+            <summary>
+            Determines whether the specified object is equal to the current object.
+            </summary>
+            <param name="obj">The other object.</param>
+            <returns>true if the specified object is equal to the current object; otherwise, false.</returns>
+        </member>
+        <member name="M:Microsoft.Owin.BuilderProperties.Address.GetHashCode">
+            <summary>
+            Returns the hash code for this instance.
+            </summary>
+            <returns>The hash code for this instance.</returns>
+        </member>
+        <member name="M:Microsoft.Owin.BuilderProperties.Address.op_Equality(Microsoft.Owin.BuilderProperties.Address,Microsoft.Owin.BuilderProperties.Address)">
+            <summary>
+            Determines whether two specified instances of <see cref="T:Microsoft.Owin.BuilderProperties.Address" /> are equal.
+            </summary>
+            <param name="left">The first object to compare.</param>
+            <param name="right">The second object to compare.</param>
+            <returns>true if left and right represent the same address; otherwise, false.</returns>
+        </member>
+        <member name="M:Microsoft.Owin.BuilderProperties.Address.op_Inequality(Microsoft.Owin.BuilderProperties.Address,Microsoft.Owin.BuilderProperties.Address)">
+            <summary>
+            Determines whether two specified instances of <see cref="T:Microsoft.Owin.BuilderProperties.Address" /> are not equal.
+            </summary>
+            <param name="left">The first object to compare.</param>
+            <param name="right">The second object to compare.</param>
+            <returns>true if left and right do not represent the same address; otherwise, false.</returns>
+        </member>
+        <member name="M:Microsoft.Owin.BuilderProperties.Address.Get``1(System.String)">
+            <summary>
+            Gets a specified key and value from the underlying dictionary.
+            </summary>
+            <typeparam name="T"></typeparam>
+            <param name="key">The key.</param>
+            <returns></returns>
+        </member>
+        <member name="M:Microsoft.Owin.BuilderProperties.Address.Set(System.String,System.Object)">
+            <summary>
+            Sets a specified key and value in the underlying dictionary.
+            </summary>
+            <param name="key">The key.</param>
+            <param name="value">The value.</param>
+            <returns></returns>
+        </member>
+        <member name="P:Microsoft.Owin.BuilderProperties.Address.Dictionary">
+            <summary>
+            Gets the internal dictionary for this collection.
+            </summary>
+            <returns>The internal dictionary for this collection.</returns>
+        </member>
+        <member name="P:Microsoft.Owin.BuilderProperties.Address.Scheme">
+            <summary>
+            The uri scheme.
+            </summary>
+        </member>
+        <member name="P:Microsoft.Owin.BuilderProperties.Address.Host">
+            <summary>
+            The uri host.
+            </summary>
+        </member>
+        <member name="P:Microsoft.Owin.BuilderProperties.Address.Port">
+            <summary>
+            The uri port.
+            </summary>
+        </member>
+        <member name="P:Microsoft.Owin.BuilderProperties.Address.Path">
+            <summary>
+            The uri path.
+            </summary>
+        </member>
+        <member name="T:Microsoft.Owin.BuilderProperties.AddressCollection">
+            <summary>
+            Wraps the host.Addresses list.
+            </summary>
+        </member>
+        <member name="M:Microsoft.Owin.BuilderProperties.AddressCollection.#ctor(System.Collections.Generic.IList{System.Collections.Generic.IDictionary{System.String,System.Object}})">
+            <summary>
+            Initializes a new instance of the <see cref="T:Microsoft.Owin.BuilderProperties.AddressCollection" /> class.
+            </summary>
+            <param name="list">The address list to set to the collection.</param>
+        </member>
+        <member name="M:Microsoft.Owin.BuilderProperties.AddressCollection.Add(Microsoft.Owin.BuilderProperties.Address)">
+            <summary>
+            Adds the specified address to the collection.
+            </summary>
+            <param name="address">The address to add to the collection.</param>
+        </member>
+        <member name="M:Microsoft.Owin.BuilderProperties.AddressCollection.System#Collections#IEnumerable#GetEnumerator">
+            <summary>
+            Gets the enumerator that iterates through the collection.
+            </summary>
+            <returns>The enumerator that can be used to iterate through the collection.</returns>
+        </member>
+        <member name="M:Microsoft.Owin.BuilderProperties.AddressCollection.GetEnumerator">
+            <summary>
+            Gets the enumerator that iterates through the collection.
+            </summary>
+            <returns>The enumerator that can be used to iterate through the collection.</returns>
+        </member>
+        <member name="M:Microsoft.Owin.BuilderProperties.AddressCollection.Create">
+            <summary>
+            Creates a new empty instance of <see cref="T:Microsoft.Owin.BuilderProperties.AddressCollection" />.
+            </summary>
+            <returns>A new empty instance of <see cref="T:Microsoft.Owin.BuilderProperties.AddressCollection" />.</returns>
+        </member>
+        <member name="M:Microsoft.Owin.BuilderProperties.AddressCollection.Equals(Microsoft.Owin.BuilderProperties.AddressCollection)">
+            <summary>
+            Determines whether the current collection is equal to the specified collection.
+            </summary>
+            <param name="other">The other collection to compare to the current collection.</param>
+            <returns>true if current collection is equal to the specified collection; otherwise, false.</returns>
+        </member>
+        <member name="M:Microsoft.Owin.BuilderProperties.AddressCollection.Equals(System.Object)">
+            <summary>
+            Determines whether the current collection is equal to the specified object.
+            </summary>
+            <param name="obj">The object to compare to the current collection.</param>
+            <returns>true if current collection is equal to the specified object; otherwise, false.</returns>
+        </member>
+        <member name="M:Microsoft.Owin.BuilderProperties.AddressCollection.GetHashCode">
+            <summary>
+            Gets the hash code for this instance.
+            </summary>
+            <returns>The hash code for this instance.</returns>
+        </member>
+        <member name="M:Microsoft.Owin.BuilderProperties.AddressCollection.op_Equality(Microsoft.Owin.BuilderProperties.AddressCollection,Microsoft.Owin.BuilderProperties.AddressCollection)">
+            <summary>
+            Determines whether the first collection is equal to the second collection.
+            </summary>
+            <param name="left">The first collection to compare.</param>
+            <param name="right">The second collection to compare.</param>
+            <returns>true if both collections are equal; otherwise, false.</returns>
+        </member>
+        <member name="M:Microsoft.Owin.BuilderProperties.AddressCollection.op_Inequality(Microsoft.Owin.BuilderProperties.AddressCollection,Microsoft.Owin.BuilderProperties.AddressCollection)">
+            <summary>
+            Determines whether the first collection is not equal to the second collection.
+            </summary>
+            <param name="left">The first collection to compare.</param>
+            <param name="right">The second collection to compare.</param>
+            <returns>true if both collections are not equal; otherwise, false.</returns>
+        </member>
+        <member name="P:Microsoft.Owin.BuilderProperties.AddressCollection.List">
+            <summary>
+            Gets the underlying address list.
+            </summary>
+            <returns>The underlying address list.</returns>
+        </member>
+        <member name="P:Microsoft.Owin.BuilderProperties.AddressCollection.Count">
+            <summary>
+            Gets the number of elements in the collection.
+            </summary>
+            <returns>The number of elements in the collection.</returns>
+        </member>
+        <member name="P:Microsoft.Owin.BuilderProperties.AddressCollection.Item(System.Int32)">
+            <summary>
+            Gets the item with the specified index from the collection.
+            </summary>
+            <param name="index">The index.</param>
+            <returns>The item with the specified index.</returns>
+        </member>
+        <member name="T:Microsoft.Owin.BuilderProperties.AppProperties">
+            <summary>
+            A wrapper for the <see cref="P:Microsoft.Owin.Builder.AppBuilder.Properties" /> IDictionary.
+            </summary>
+        </member>
+        <member name="M:Microsoft.Owin.BuilderProperties.AppProperties.#ctor(System.Collections.Generic.IDictionary{System.String,System.Object})">
+            <summary>
+            Initializes a new instance of the <see cref="T:Microsoft.Owin.BuilderProperties.AppProperties" /> class.
+            </summary>
+            <param name="dictionary"></param>
+        </member>
+        <member name="M:Microsoft.Owin.BuilderProperties.AppProperties.Equals(Microsoft.Owin.BuilderProperties.AppProperties)">
+            <summary>
+            Determines whether the current AppProperties is equal to the specified AppProperties.
+            </summary>
+            <param name="other">The other AppProperties to compare with the current instance.</param>
+            <returns>true if the current AppProperties is equal to the specified AppProperties; otherwise, false.</returns>
+        </member>
+        <member name="M:Microsoft.Owin.BuilderProperties.AppProperties.Equals(System.Object)">
+            <summary>
+            Determines whether the current AppProperties is equal to the specified object.
+            </summary>
+            <param name="obj">The object to compare with the current instance.</param>
+            <returns>true if the current AppProperties is equal to the specified object; otherwise, false.</returns>
+        </member>
+        <member name="M:Microsoft.Owin.BuilderProperties.AppProperties.GetHashCode">
+            <summary>
+            Returns the hash code for this instance.
+            </summary>
+            <returns>The hash code for this instance.</returns>
+        </member>
+        <member name="M:Microsoft.Owin.BuilderProperties.AppProperties.op_Equality(Microsoft.Owin.BuilderProperties.AppProperties,Microsoft.Owin.BuilderProperties.AppProperties)">
+            <summary>
+            Determines whether the first AppPProperties is equal to the second AppProperties.
+            </summary>
+            <param name="left">The first AppPropeties to compare.</param>
+            <param name="right">The second AppPropeties to compare.</param>
+            <returns>true if both AppProperties are equal; otherwise, false.</returns>
+        </member>
+        <member name="M:Microsoft.Owin.BuilderProperties.AppProperties.op_Inequality(Microsoft.Owin.BuilderProperties.AppProperties,Microsoft.Owin.BuilderProperties.AppProperties)">
+            <summary>
+            Determines whether the first AppPProperties is not equal to the second AppProperties.
+            </summary>
+            <param name="left">The first AppPropeties to compare.</param>
+            <param name="right">The second AppPropeties to compare.</param>
+            <returns>true if both AppProperties are not equal; otherwise, false.</returns>
+        </member>
+        <member name="M:Microsoft.Owin.BuilderProperties.AppProperties.Get``1(System.String)">
+            <summary>
+            Gets the value from the dictionary with the specified key.
+            </summary>
+            <typeparam name="T">The type of the value.</typeparam>
+            <param name="key">The key of the value to get.</param>
+            <returns>The value with the specified key.</returns>
+        </member>
+        <member name="M:Microsoft.Owin.BuilderProperties.AppProperties.Set(System.String,System.Object)">
+            <summary>
+            Sets the value with the specified key.
+            </summary>
+            <param name="key">The key of the value to set.</param>
+            <param name="value">The value to set.</param>
+            <returns>This instance.</returns>
+        </member>
+        <member name="P:Microsoft.Owin.BuilderProperties.AppProperties.OwinVersion">
+            <summary>
+            Gets or sets the string value for “owin.Version”.
+            </summary>
+            <returns>The string value for “owin.Version”.</returns>
+        </member>
+        <member name="P:Microsoft.Owin.BuilderProperties.AppProperties.DefaultApp">
+            <summary>
+            Gets or sets the function delegate for “builder.DefaultApp”.
+            </summary>
+            <returns>The function delegate for “builder.DefaultApp”.</returns>
+        </member>
+        <member name="P:Microsoft.Owin.BuilderProperties.AppProperties.AddSignatureConversionDelegate">
+            <summary>
+            Gets or sets the action delegate for “builder.AddSignatureConversion”.
+            </summary>
+            <returns>The action delegate for “builder.AddSignatureConversion”.</returns>
+        </member>
+        <member name="P:Microsoft.Owin.BuilderProperties.AppProperties.AppName">
+            <summary>
+            Gets or sets the string value for “host.AppName”.
+            </summary>
+            <returns>The string value for “host.AppName”.</returns>
+        </member>
+        <member name="P:Microsoft.Owin.BuilderProperties.AppProperties.TraceOutput">
+            <summary>
+            Gets or sets the text writer for “host.TraceOutput”.
+            </summary>
+            <returns>The text writer for “host.TraceOutput”.</returns>
+        </member>
+        <member name="P:Microsoft.Owin.BuilderProperties.AppProperties.OnAppDisposing">
+            <summary>
+            Gets or sets the cancellation token for “host.OnAppDisposing”.
+            </summary>
+            <returns>The cancellation token for “host.OnAppDisposing”.</returns>
+        </member>
+        <member name="P:Microsoft.Owin.BuilderProperties.AppProperties.Addresses">
+            <summary>
+            Gets or sets the address collection for “host.Addresses”.
+            </summary>
+            <returns>The address collection for “host.Addresses”.</returns>
+        </member>
+        <member name="P:Microsoft.Owin.BuilderProperties.AppProperties.Capabilities">
+            <summary>
+            Gets or sets the list of “server.Capabilities”.
+            </summary>
+            <returns>The list of “server.Capabilities”.</returns>
+        </member>
+        <member name="P:Microsoft.Owin.BuilderProperties.AppProperties.Dictionary">
+            <summary>
+            Gets the underlying dictionary for this <see cref="T:Microsoft.Owin.BuilderProperties.AppProperties" /> instance.
+            </summary>
+            <returns>The underlying dictionary for this <see cref="T:Microsoft.Owin.BuilderProperties.AppProperties" /> instance.</returns>
+        </member>
+        <member name="T:Microsoft.Owin.BuilderProperties.Capabilities">
+            <summary>
+            Represents the capabilities for the builder properties.
+            </summary>
+        </member>
+        <member name="M:Microsoft.Owin.BuilderProperties.Capabilities.#ctor(System.Collections.Generic.IDictionary{System.String,System.Object})">
+            <summary>
+            Initializes a new instance of the <see cref="T:Microsoft.Owin.BuilderProperties.Capabilities" /> class.
+            </summary>
+            <param name="dictionary"></param>
+        </member>
+        <member name="M:Microsoft.Owin.BuilderProperties.Capabilities.Create">
+            <summary>
+            Initializes a new instance of the <see cref="T:Microsoft.Owin.BuilderProperties.Capabilities" /> class.
+            </summary>
+            <returns>A new instance of the <see cref="T:Microsoft.Owin.BuilderProperties.Capabilities" /> class.</returns>
+        </member>
+        <member name="M:Microsoft.Owin.BuilderProperties.Capabilities.Equals(Microsoft.Owin.BuilderProperties.Capabilities)">
+            <summary>
+            Determines whether the current Capabilities instance is equal to the specified Capabilities.
+            </summary>
+            <param name="other">The other Capabilities to compare with the current instance.</param>
+            <returns>true if the specified object is equal to the current object; otherwise, false.</returns>
+        </member>
+        <member name="M:Microsoft.Owin.BuilderProperties.Capabilities.Equals(System.Object)">
+            <summary>
+            Determines whether the current Capabilities is equal to the specified object.
+            </summary>
+            <param name="obj">The object to compare with the current instance.</param>
+            <returns>true if the current Capabilities is equal to the specified object; otherwise, false.</returns>
+        </member>
+        <member name="M:Microsoft.Owin.BuilderProperties.Capabilities.GetHashCode">
+            <summary>
+            Returns the hash code for this instance.
+            </summary>
+            <returns>The hash code for this instance.</returns>
+        </member>
+        <member name="M:Microsoft.Owin.BuilderProperties.Capabilities.op_Equality(Microsoft.Owin.BuilderProperties.Capabilities,Microsoft.Owin.BuilderProperties.Capabilities)">
+            <summary>
+            Determines whether two specified instances of <see cref="T:Microsoft.Owin.BuilderProperties.Capabilities" /> are equal.
+            </summary>
+            <param name="left">The first object to compare.</param>
+            <param name="right">The second object to compare.</param>
+            <returns>true if the two specified instances of <see cref="T:Microsoft.Owin.BuilderProperties.Capabilities" /> are equal; otherwise, false.</returns>
+        </member>
+        <member name="M:Microsoft.Owin.BuilderProperties.Capabilities.op_Inequality(Microsoft.Owin.BuilderProperties.Capabilities,Microsoft.Owin.BuilderProperties.Capabilities)">
+            <summary>
+            Determines whether two specified instances of <see cref="T:Microsoft.Owin.BuilderProperties.Capabilities" /> are not equal.
+            </summary>
+            <param name="left">The first object to compare.</param>
+            <param name="right">The second object to compare.</param>
+            <returns>true if the two specified instances of <see cref="T:Microsoft.Owin.BuilderProperties.Capabilities" /> are not equal; otherwise, false.</returns>
+        </member>
+        <member name="M:Microsoft.Owin.BuilderProperties.Capabilities.Get``1(System.String)">
+            <summary>
+            Gets the value from the dictionary with the specified key.
+            </summary>
+            <typeparam name="T">The type of the value.</typeparam>
+            <param name="key">The key of the value to get.</param>
+            <returns>The value with the specified key.</returns>
+        </member>
+        <member name="M:Microsoft.Owin.BuilderProperties.Capabilities.Set(System.String,System.Object)">
+            <summary>
+            Sets the given key and value in the underlying dictionary.
+            </summary>
+            <param name="key">The key of the value to set.</param>
+            <param name="value">The value to set.</param>
+            <returns>This instance.</returns>
+        </member>
+        <member name="P:Microsoft.Owin.BuilderProperties.Capabilities.Dictionary">
+            <summary>
+            The underling IDictionary
+            </summary>
+        </member>
+        <member name="P:Microsoft.Owin.BuilderProperties.Capabilities.SendFileVersion">
+            <summary>
+            Gets or sets the string value for "sendfile.Version"
+            </summary>
+            <returns>the string value for "sendfile.Version"</returns>
+        </member>
+        <member name="P:Microsoft.Owin.BuilderProperties.Capabilities.WebSocketVersion">
+            <summary>
+            Gets or sets the websocket version.
+            </summary>
+            <returns>The websocket version.</returns>
+        </member>
+        <member name="T:Microsoft.Owin.CookieOptions">
+            <summary>
+            Options used to create a new cookie.
+            </summary>
+        </member>
+        <member name="M:Microsoft.Owin.CookieOptions.#ctor">
+            <summary>
+            Creates a default cookie with a path of '/'.
+            </summary>
+        </member>
+        <member name="P:Microsoft.Owin.CookieOptions.Domain">
+            <summary>
+            Gets or sets the domain to associate the cookie with.
+            </summary>
+            <returns>The domain to associate the cookie with.</returns>
+        </member>
+        <member name="P:Microsoft.Owin.CookieOptions.Path">
+            <summary>
+            Gets or sets the cookie path.
+            </summary>
+            <returns>The cookie path.</returns>
+        </member>
+        <member name="P:Microsoft.Owin.CookieOptions.Expires">
+            <summary>
+            Gets or sets the expiration date and time for the cookie.
+            </summary>
+            <returns>The expiration date and time for the cookie.</returns>
+        </member>
+        <member name="P:Microsoft.Owin.CookieOptions.Secure">
+            <summary>
+            Gets or sets a value that indicates whether to transmit the cookie using Secure Sockets Layer (SSL)—that is, over HTTPS only.
+            </summary>
+            <returns>true to transmit the cookie only over an SSL connection (HTTPS); otherwise, false.</returns>
+        </member>
+        <member name="P:Microsoft.Owin.CookieOptions.HttpOnly">
+            <summary>
+            Gets or sets a value that indicates whether a cookie is accessible by client-side script.
+            </summary>
+            <returns>true if a cookie is accessible by client-side script; otherwise, false.</returns>
+        </member>
+        <member name="T:Microsoft.Owin.Logging.LoggerFactory">
+            <summary>
+            Provides a default ILoggerFactory.
+            </summary>
+        </member>
+        <member name="P:Microsoft.Owin.Logging.LoggerFactory.Default">
+            <summary>
+            Provides a default ILoggerFactory based on System.Diagnostics.TraceSorce.
+            </summary>
+        </member>
+        <member name="T:Microsoft.Owin.Builder.AppBuilderExtensions">
+            <summary>
+            Extension methods for IAppBuilder.
+            </summary>
+        </member>
+        <member name="M:Microsoft.Owin.Builder.AppBuilderExtensions.Build(Owin.IAppBuilder)">
+            <summary>
+            The Build is called at the point when all of the middleware should be chained
+            together. May be called to build pipeline branches.
+            </summary>
+            <param name="builder"></param>
+            <returns>The request processing entry point for this section of the pipeline.</returns>
+        </member>
+        <member name="M:Microsoft.Owin.Builder.AppBuilderExtensions.Build``1(Owin.IAppBuilder)">
+            <summary>
+            The Build is called at the point when all of the middleware should be chained
+            together. May be called to build pipeline branches.
+            </summary>
+            <typeparam name="TApp">The application signature.</typeparam>
+            <param name="builder"></param>
+            <returns>The request processing entry point for this section of the pipeline.</returns>
+        </member>
+        <member name="M:Microsoft.Owin.Builder.AppBuilderExtensions.AddSignatureConversion(Owin.IAppBuilder,System.Delegate)">
+            <summary>
+            Adds converters for adapting between disparate application signatures.
+            </summary>
+            <param name="builder"></param>
+            <param name="conversion"></param>
+        </member>
+        <member name="M:Microsoft.Owin.Builder.AppBuilderExtensions.AddSignatureConversion``2(Owin.IAppBuilder,System.Func{``0,``1})">
+            <summary>
+            Adds converters for adapting between disparate application signatures.
+            </summary>
+            <typeparam name="T1"></typeparam>
+            <typeparam name="T2"></typeparam>
+            <param name="builder"></param>
+            <param name="conversion"></param>
+        </member>
+        <member name="T:Microsoft.Owin.Helpers.WebHelpers">
+            <summary>
+            Provides helper methods for processing requests.
+            </summary>
+        </member>
+        <member name="M:Microsoft.Owin.Helpers.WebHelpers.ParseForm(System.String)">
+            <summary>
+            Parses an HTTP form body.
+            </summary>
+            <param name="text">The HTTP form body to parse.</param>
+            <returns>The <see cref="T:Microsoft.Owin.IFormCollection" /> object containing the parsed HTTP form body.</returns>
+        </member>
+        <member name="T:Microsoft.Owin.Infrastructure.AppFuncTransition">
+            <summary>
+            Converts between an OwinMiddlware and an <typeref name="Func&lt;IDictionary&lt;string,object&gt;, Task&gt;"/>.
+            </summary>
+        </member>
+        <member name="T:Microsoft.Owin.OwinMiddleware">
+            <summary>
+            An abstract base class for a standard middleware pattern.
+            </summary>
+        </member>
+        <member name="M:Microsoft.Owin.OwinMiddleware.#ctor(Microsoft.Owin.OwinMiddleware)">
+            <summary>
+            Instantiates the middleware with an optional pointer to the next component.
+            </summary>
+            <param name="next"></param>
+        </member>
+        <member name="M:Microsoft.Owin.OwinMiddleware.Invoke(Microsoft.Owin.IOwinContext)">
+            <summary>
+            Process an individual request.
+            </summary>
+            <param name="context"></param>
+            <returns></returns>
+        </member>
+        <member name="P:Microsoft.Owin.OwinMiddleware.Next">
+            <summary>
+            The optional next component.
+            </summary>
+        </member>
+        <member name="M:Microsoft.Owin.Infrastructure.AppFuncTransition.#ctor(System.Func{System.Collections.Generic.IDictionary{System.String,System.Object},System.Threading.Tasks.Task})">
+            <summary>
+            
+            </summary>
+            <param name="next"></param>
+        </member>
+        <member name="M:Microsoft.Owin.Infrastructure.AppFuncTransition.Invoke(Microsoft.Owin.IOwinContext)">
+            <summary>
+            
+            </summary>
+            <param name="context"></param>
+            <returns></returns>
+        </member>
+        <member name="T:Microsoft.Owin.Infrastructure.OwinMiddlewareTransition">
+            <summary>
+            Transitions between <typeref name="Func&lt;IDictionary&lt;string,object&gt;, Task&gt;"/> and OwinMiddleware.
+            </summary>
+        </member>
+        <member name="M:Microsoft.Owin.Infrastructure.OwinMiddlewareTransition.#ctor(Microsoft.Owin.OwinMiddleware)">
+            <summary>
+            
+            </summary>
+            <param name="next"></param>
+        </member>
+        <member name="M:Microsoft.Owin.Infrastructure.OwinMiddlewareTransition.Invoke(System.Collections.Generic.IDictionary{System.String,System.Object})">
+            <summary>
+            
+            </summary>
+            <param name="environment">OWIN environment dictionary which stores state information about the request, response and relevant server state.</param>
+            <returns></returns>
+        </member>
+        <member name="T:Microsoft.Owin.Infrastructure.SignatureConversions">
+            <summary>
+            Adds adapters between <typeref name="Func&lt;IDictionary&lt;string,object&gt;, Task&gt;"/> and OwinMiddleware.
+            </summary>
+        </member>
+        <member name="M:Microsoft.Owin.Infrastructure.SignatureConversions.AddConversions(Owin.IAppBuilder)">
+            <summary>
+            Adds adapters between <typeref name="Func&lt;IDictionary&lt;string,object&gt;, Task&gt;"/> and OwinMiddleware.
+            </summary>
+            <param name="app"></param>
+        </member>
+        <member name="T:Microsoft.Owin.OwinRequest">
+            <summary>
+            This wraps OWIN environment dictionary and provides strongly typed accessors.
+            </summary>
+        </member>
+        <member name="M:Microsoft.Owin.OwinRequest.#ctor">
+            <summary>
+            Create a new context with only request and response header collections.
+            </summary>
+        </member>
+        <member name="M:Microsoft.Owin.OwinRequest.#ctor(System.Collections.Generic.IDictionary{System.String,System.Object})">
+            <summary>
+            Create a new environment wrapper exposing request properties.
+            </summary>
+            <param name="environment">OWIN environment dictionary which stores state information about the request, response and relevant server state.</param>
+        </member>
+        <member name="M:Microsoft.Owin.OwinRequest.ReadFormAsync">
+            <summary>
+            Asynchronously reads and parses the request body as a form.
+            </summary>
+            <returns>The parsed form data.</returns>
+        </member>
+        <member name="M:Microsoft.Owin.OwinRequest.Get``1(System.String)">
+            <summary>
+            Gets a value from the OWIN environment, or returns default(T) if not present.
+            </summary>
+            <typeparam name="T">The type of the value.</typeparam>
+            <param name="key">The key of the value to get.</param>
+            <returns>The value with the specified key or the default(T) if not present.</returns>
+        </member>
+        <member name="M:Microsoft.Owin.OwinRequest.Set``1(System.String,``0)">
+            <summary>
+            Sets the given key and value in the OWIN environment.
+            </summary>
+            <typeparam name="T">The type of the value.</typeparam>
+            <param name="key">The key of the value to set.</param>
+            <param name="value">The value to set.</param>
+            <returns>This instance.</returns>
+        </member>
+        <member name="P:Microsoft.Owin.OwinRequest.Environment">
+            <summary>
+            Gets the OWIN environment.
+            </summary>
+            <returns>The OWIN environment.</returns>
+        </member>
+        <member name="P:Microsoft.Owin.OwinRequest.Context">
+            <summary>
+            Gets the request context.
+            </summary>
+            <returns>The request context.</returns>
+        </member>
+        <member name="P:Microsoft.Owin.OwinRequest.Method">
+            <summary>
+            Gets or set the HTTP method.
+            </summary>
+            <returns>The HTTP method.</returns>
+        </member>
+        <member name="P:Microsoft.Owin.OwinRequest.Scheme">
+            <summary>
+            Gets or set the HTTP request scheme from owin.RequestScheme.
+            </summary>
+            <returns>The HTTP request scheme from owin.RequestScheme.</returns>
+        </member>
+        <member name="P:Microsoft.Owin.OwinRequest.IsSecure">
+            <summary>
+            Returns true if the owin.RequestScheme is https.
+            </summary>
+            <returns>true if this request is using https; otherwise, false.</returns>
+        </member>
+        <member name="P:Microsoft.Owin.OwinRequest.Host">
+            <summary>
+            Gets or set the Host header. May include the port.
+            </summary>
+            <return>The Host header.</return>
+        </member>
+        <member name="P:Microsoft.Owin.OwinRequest.PathBase">
+            <summary>
+            Gets or set the owin.RequestPathBase.
+            </summary>
+            <returns>The owin.RequestPathBase.</returns>
+        </member>
+        <member name="P:Microsoft.Owin.OwinRequest.Path">
+            <summary>
+            Gets or set the request path from owin.RequestPath.
+            </summary>
+            <returns>The request path from owin.RequestPath.</returns>
+        </member>
+        <member name="P:Microsoft.Owin.OwinRequest.QueryString">
+            <summary>
+            Gets or set the query string from owin.RequestQueryString.
+            </summary>
+            <returns>The query string from owin.RequestQueryString.</returns>
+        </member>
+        <member name="P:Microsoft.Owin.OwinRequest.Query">
+            <summary>
+            Gets the query value collection parsed from owin.RequestQueryString.
+            </summary>
+            <returns>The query value collection parsed from owin.RequestQueryString.</returns>
+        </member>
+        <member name="P:Microsoft.Owin.OwinRequest.Uri">
+            <summary>
+            Gets the uniform resource identifier (URI) associated with the request.
+            </summary>
+            <returns>The uniform resource identifier (URI) associated with the request.</returns>
+        </member>
+        <member name="P:Microsoft.Owin.OwinRequest.Protocol">
+            <summary>
+            Gets or set the owin.RequestProtocol.
+            </summary>
+            <returns>The owin.RequestProtocol.</returns>
+        </member>
+        <member name="P:Microsoft.Owin.OwinRequest.Headers">
+            <summary>
+            Gets the request headers.
+            </summary>
+            <returns>The request headers.</returns>
+        </member>
+        <member name="P:Microsoft.Owin.OwinRequest.Cookies">
+            <summary>
+            Gets the collection of Cookies for this request.
+            </summary>
+            <returns>The collection of Cookies for this request.</returns>
+        </member>
+        <member name="P:Microsoft.Owin.OwinRequest.ContentType">
+            <summary>
+            Gets or sets the Content-Type header.
+            </summary>
+            <returns>The Content-Type header.</returns>
+        </member>
+        <member name="P:Microsoft.Owin.OwinRequest.CacheControl">
+            <summary>
+            Gets or sets the Cache-Control header.
+            </summary>
+            <returns>The Cache-Control header.</returns>
+        </member>
+        <member name="P:Microsoft.Owin.OwinRequest.MediaType">
+            <summary>
+            Gets or sets the Media-Type header.
+            </summary>
+            <returns>The Media-Type header.</returns>
+        </member>
+        <member name="P:Microsoft.Owin.OwinRequest.Accept">
+            <summary>
+            Gets or set the Accept header.
+            </summary>
+            <returns>The Accept header.</returns>
+        </member>
+        <member name="P:Microsoft.Owin.OwinRequest.Body">
+            <summary>
+            Gets or set the owin.RequestBody Stream.
+            </summary>
+            <returns>The owin.RequestBody Stream.</returns>
+        </member>
+        <member name="P:Microsoft.Owin.OwinRequest.CallCancelled">
+            <summary>
+            Gets or sets the cancellation token for the request.
+            </summary>
+            <returns>The cancellation token for the request.</returns>
+        </member>
+        <member name="P:Microsoft.Owin.OwinRequest.LocalIpAddress">
+            <summary>
+            Gets or set the server.LocalIpAddress.
+            </summary>
+            <returns>The server.LocalIpAddress.</returns>
+        </member>
+        <member name="P:Microsoft.Owin.OwinRequest.LocalPort">
+            <summary>
+            Gets or set the server.LocalPort.
+            </summary>
+            <returns>The server.LocalPort.</returns>
+        </member>
+        <member name="P:Microsoft.Owin.OwinRequest.RemoteIpAddress">
+            <summary>
+            Gets or set the server.RemoteIpAddress.
+            </summary>
+            <returns>The server.RemoteIpAddress.</returns>
+        </member>
+        <member name="P:Microsoft.Owin.OwinRequest.RemotePort">
+            <summary>
+            Gets or set the server.RemotePort.
+            </summary>
+            <returns>The server.RemotePort.</returns>
+        </member>
+        <member name="P:Microsoft.Owin.OwinRequest.User">
+            <summary>
+            Gets or set the server.User.
+            </summary>
+            <returns>The server.User.</returns>
+        </member>
+        <member name="T:Microsoft.Owin.OwinResponse">
+            <summary>
+            This wraps OWIN environment dictionary and provides strongly typed accessors.
+            </summary>
+        </member>
+        <member name="M:Microsoft.Owin.OwinResponse.#ctor">
+            <summary>
+            Create a new context with only request and response header collections.
+            </summary>
+        </member>
+        <member name="M:Microsoft.Owin.OwinResponse.#ctor(System.Collections.Generic.IDictionary{System.String,System.Object})">
+            <summary>
+            Creates a new environment wrapper exposing response properties.
+            </summary>
+            <param name="environment">OWIN environment dictionary which stores state information about the request, response and relevant server state.</param>
+        </member>
+        <member name="M:Microsoft.Owin.OwinResponse.OnSendingHeaders(System.Action{System.Object},System.Object)">
+            <summary>
+            Registers for an event that fires when the response headers are sent.
+            </summary>
+            <param name="callback">The callback method.</param>
+            <param name="state">The callback state.</param>
+        </member>
+        <member name="M:Microsoft.Owin.OwinResponse.Redirect(System.String)">
+            <summary>
+            Sets a 302 response status code and the Location header.
+            </summary>
+            <param name="location">The location where to redirect the client.</param>
+        </member>
+        <member name="M:Microsoft.Owin.OwinResponse.Write(System.String)">
+            <summary>
+            Writes the given text to the response body stream using UTF-8.
+            </summary>
+            <param name="text">The response data.</param>
+        </member>
+        <member name="M:Microsoft.Owin.OwinResponse.Write(System.Byte[])">
+            <summary>
+            Writes the given bytes to the response body stream.
+            </summary>
+            <param name="data">The response data.</param>
+        </member>
+        <member name="M:Microsoft.Owin.OwinResponse.Write(System.Byte[],System.Int32,System.Int32)">
+            <summary>
+            Writes the given bytes to the response body stream.
+            </summary>
+            <param name="data">The response data.</param>
+            <param name="offset">The zero-based byte offset in the <paramref name="data" /> parameter at which to begin copying bytes.</param>
+            <param name="count">The number of bytes to write.</param>
+        </member>
+        <member name="M:Microsoft.Owin.OwinResponse.WriteAsync(System.String)">
+            <summary>
+            Asynchronously writes the given text to the response body stream using UTF-8.
+            </summary>
+            <param name="text">The response data.</param>
+            <returns>A Task tracking the state of the write operation.</returns>
+        </member>
+        <member name="M:Microsoft.Owin.OwinResponse.WriteAsync(System.String,System.Threading.CancellationToken)">
+            <summary>
+            Asynchronously writes the given text to the response body stream using UTF-8.
+            </summary>
+            <param name="text">The response data.</param>
+            <param name="token">A token used to indicate cancellation.</param>
+            <returns>A Task tracking the state of the write operation.</returns>
+        </member>
+        <member name="M:Microsoft.Owin.OwinResponse.WriteAsync(System.Byte[])">
+            <summary>
+            Asynchronously writes the given bytes to the response body stream.
+            </summary>
+            <param name="data">The response data.</param>
+            <returns>A Task tracking the state of the write operation.</returns>
+        </member>
+        <member name="M:Microsoft.Owin.OwinResponse.WriteAsync(System.Byte[],System.Threading.CancellationToken)">
+            <summary>
+            Asynchronously writes the given bytes to the response body stream.
+            </summary>
+            <param name="data">The response data.</param>
+            <param name="token">A token used to indicate cancellation.</param>
+            <returns>A Task tracking the state of the write operation.</returns>
+        </member>
+        <member name="M:Microsoft.Owin.OwinResponse.WriteAsync(System.Byte[],System.Int32,System.Int32,System.Threading.CancellationToken)">
+            <summary>
+            Asynchronously writes the given bytes to the response body stream.
+            </summary>
+            <param name="data">The response data.</param>
+            <param name="offset">The zero-based byte offset in the <paramref name="data" /> parameter at which to begin copying bytes.</param>
+            <param name="count">The number of bytes to write.</param>
+            <param name="token">A token used to indicate cancellation.</param>
+            <returns>A Task tracking the state of the write operation.</returns>
+        </member>
+        <member name="M:Microsoft.Owin.OwinResponse.Get``1(System.String)">
+            <summary>
+            Gets a value from the OWIN environment, or returns default(T) if not present.
+            </summary>
+            <typeparam name="T">The type of the value.</typeparam>
+            <param name="key">The key of the value to get.</param>
+            <returns>The value with the specified key or the default(T) if not present.</returns>
+        </member>
+        <member name="M:Microsoft.Owin.OwinResponse.Set``1(System.String,``0)">
+            <summary>
+            Sets the given key and value in the OWIN environment.
+            </summary>
+            <typeparam name="T">The type of the value.</typeparam>
+            <param name="key">The key of the value to set.</param>
+            <param name="value">The value to set.</param>
+            <returns>This instance.</returns>
+        </member>
+        <member name="P:Microsoft.Owin.OwinResponse.Environment">
+            <summary>
+            Gets the OWIN environment.
+            </summary>
+            <returns>The OWIN environment.</returns>
+        </member>
+        <member name="P:Microsoft.Owin.OwinResponse.Context">
+            <summary>
+            Gets the request context.
+            </summary>
+            <returns>The request context.</returns>
+        </member>
+        <member name="P:Microsoft.Owin.OwinResponse.StatusCode">
+            <summary>
+            Gets or sets the optional owin.ResponseStatusCode.
+            </summary>
+            <returns>The optional owin.ResponseStatusCode, or 200 if not set.</returns>
+        </member>
+        <member name="P:Microsoft.Owin.OwinResponse.ReasonPhrase">
+            <summary>
+            Gets or sets the the optional owin.ResponseReasonPhrase.
+            </summary>
+            <returns>The the optional owin.ResponseReasonPhrase.</returns>
+        </member>
+        <member name="P:Microsoft.Owin.OwinResponse.Protocol">
+            <summary>
+            Gets or sets the owin.ResponseProtocol.
+            </summary>
+            <returns>The owin.ResponseProtocol.</returns>
+        </member>
+        <member name="P:Microsoft.Owin.OwinResponse.Headers">
+            <summary>
+            Gets the response header collection.
+            </summary>
+            <returns>The response header collection.</returns>
+        </member>
+        <member name="P:Microsoft.Owin.OwinResponse.Cookies">
+            <summary>
+            Gets a collection used to manipulate the Set-Cookie header.
+            </summary>
+            <returns>A collection used to manipulate the Set-Cookie header.</returns>
+        </member>
+        <member name="P:Microsoft.Owin.OwinResponse.ContentLength">
+            <summary>
+            Gets or sets the Content-Length header.
+            </summary>
+            <returns>The Content-Length header.</returns>
+        </member>
+        <member name="P:Microsoft.Owin.OwinResponse.ContentType">
+            <summary>
+            Gets or sets the Content-Type header.
+            </summary>
+            <returns>The Content-Type header.</returns>
+        </member>
+        <member name="P:Microsoft.Owin.OwinResponse.Expires">
+            <summary>
+            Gets or sets the Expires header.
+            </summary>
+            <returns>The Expires header.</returns>
+        </member>
+        <member name="P:Microsoft.Owin.OwinResponse.ETag">
+            <summary>
+            Gets or sets the E-Tag header.
+            </summary>
+            <returns>The E-Tag header.</returns>
+        </member>
+        <member name="P:Microsoft.Owin.OwinResponse.Body">
+            <summary>
+            Gets or sets the owin.ResponseBody Stream.
+            </summary>
+            <returns>The owin.ResponseBody Stream.</returns>
+        </member>
+        <member name="T:Microsoft.Owin.Security.IAuthenticationManager">
+            <summary>
+            Used to interact with authentication middleware that have been chained in the pipeline
+            </summary>
+        </member>
+        <member name="M:Microsoft.Owin.Security.IAuthenticationManager.GetAuthenticationTypes">
+            <summary>
+            Lists all of the description data provided by authentication middleware that have been chained
+            </summary>
+            <returns>The authentication descriptions</returns>
+        </member>
+        <member name="M:Microsoft.Owin.Security.IAuthenticationManager.GetAuthenticationTypes(System.Func{Microsoft.Owin.Security.AuthenticationDescription,System.Boolean})">
+            <summary>
+            Lists the description data of all of the authentication middleware which are true for a given predicate
+            </summary>
+            <param name="predicate">A function provided by the caller which returns true for descriptions that should be in the returned list</param>
+            <returns>The authentication descriptions</returns>
+        </member>
+        <member name="M:Microsoft.Owin.Security.IAuthenticationManager.AuthenticateAsync(System.String)">
+            <summary>
+            Call back through the middleware to ask for a specific form of authentication to be performed
+            on the current request
+            </summary>
+            <param name="authenticationType">Identifies which middleware should respond to the request
+            for authentication. This value is compared to the middleware's Options.AuthenticationType property.</param>
+            <returns>Returns an object with the results of the authentication. The AuthenticationResult.Identity
+            may be null if authentication failed. Even if the Identity property is null, there may still be 
+            AuthenticationResult.properties and AuthenticationResult.Description information returned.</returns>
+        </member>
+        <member name="M:Microsoft.Owin.Security.IAuthenticationManager.AuthenticateAsync(System.String[])">
+            <summary>
+            Called to perform any number of authentication mechanisms on the current request.
+            </summary>
+            <param name="authenticationTypes">Identifies one or more middleware which should attempt to respond</param>
+            <returns>Returns the AuthenticationResult information from the middleware which responded. The 
+            order is determined by the order the middleware are in the pipeline. Latest added is first in the list.</returns>
+        </member>
+        <member name="M:Microsoft.Owin.Security.IAuthenticationManager.Challenge(Microsoft.Owin.Security.AuthenticationProperties,System.String[])">
+            <summary>
+            Add information into the response environment that will cause the authentication middleware to challenge
+            the caller to authenticate. This also changes the status code of the response to 401. The nature of that 
+            challenge varies greatly, and ranges from adding a response header or changing the 401 status code to 
+            a 302 redirect.
+            </summary>
+            <param name="properties">Additional arbitrary values which may be used by particular authentication types.</param>
+            <param name="authenticationTypes">Identify which middleware should perform their alterations on the
+            response. If the authenticationTypes is null or empty, that means the 
+            AuthenticationMode.Active middleware should perform their alterations on the response.</param>
+        </member>
+        <member name="M:Microsoft.Owin.Security.IAuthenticationManager.Challenge(System.String[])">
+            <summary>
+            Add information into the response environment that will cause the authentication middleware to challenge
+            the caller to authenticate. This also changes the status code of the response to 401. The nature of that 
+            challenge varies greatly, and ranges from adding a response header or changing the 401 status code to 
+            a 302 redirect.
+            </summary>
+            <param name="authenticationTypes">Identify which middleware should perform their alterations on the
+            response. If the authenticationTypes is null or empty, that means the 
+            AuthenticationMode.Active middleware should perform their alterations on the response.</param>
+        </member>
+        <member name="M:Microsoft.Owin.Security.IAuthenticationManager.SignIn(Microsoft.Owin.Security.AuthenticationProperties,System.Security.Claims.ClaimsIdentity[])">
+            <summary>
+            Add information to the response environment that will cause the appropriate authentication middleware
+            to grant a claims-based identity to the recipient of the response. The exact mechanism of this may vary.
+            Examples include setting a cookie, to adding a fragment on the redirect url, or producing an OAuth2
+            access code or token response.
+            </summary>
+            <param name="properties">Contains additional properties the middleware are expected to persist along with
+            the claims. These values will be returned as the AuthenticateResult.properties collection when AuthenticateAsync
+            is called on subsequent requests.</param>
+            <param name="identities">Determines which claims are granted to the signed in user. The 
+            ClaimsIdentity.AuthenticationType property is compared to the middleware's Options.AuthenticationType 
+            value to determine which claims are granted by which middleware. The recommended use is to have a single
+            ClaimsIdentity which has the AuthenticationType matching a specific middleware.</param>
+        </member>
+        <member name="M:Microsoft.Owin.Security.IAuthenticationManager.SignIn(System.Security.Claims.ClaimsIdentity[])">
+            <summary>
+            Add information to the response environment that will cause the appropriate authentication middleware
+            to grant a claims-based identity to the recipient of the response. The exact mechanism of this may vary.
+            Examples include setting a cookie, to adding a fragment on the redirect url, or producing an OAuth2
+            access code or token response.
+            </summary>
+            <param name="identities">Determines which claims are granted to the signed in user. The 
+            ClaimsIdentity.AuthenticationType property is compared to the middleware's Options.AuthenticationType 
+            value to determine which claims are granted by which middleware. The recommended use is to have a single
+            ClaimsIdentity which has the AuthenticationType matching a specific middleware.</param>
+        </member>
+        <member name="M:Microsoft.Owin.Security.IAuthenticationManager.SignOut(Microsoft.Owin.Security.AuthenticationProperties,System.String[])">
+            <summary>
+            Add information to the response environment that will cause the appropriate authentication middleware
+            to revoke any claims identity associated the the caller. The exact method varies.
+            </summary>
+            <param name="properties">Additional arbitrary values which may be used by particular authentication types.</param>
+            <param name="authenticationTypes">Identifies which middleware should perform the work to sign out.
+            Multiple authentication types may be provided to clear out more than one cookie at a time, or to clear
+            cookies and redirect to an external single-sign out url.</param>
+        </member>
+        <member name="M:Microsoft.Owin.Security.IAuthenticationManager.SignOut(System.String[])">
+            <summary>
+            Add information to the response environment that will cause the appropriate authentication middleware
+            to revoke any claims identity associated the the caller. The exact method varies.
+            </summary>
+            <param name="authenticationTypes">Identifies which middleware should perform the work to sign out.
+            Multiple authentication types may be provided to clear out more than one cookie at a time, or to clear
+            cookies and redirect to an external single-sign out url.</param>
+        </member>
+        <member name="P:Microsoft.Owin.Security.IAuthenticationManager.User">
+            <summary>
+            Returns the current user for the request
+            </summary>
+        </member>
+        <member name="P:Microsoft.Owin.Security.IAuthenticationManager.AuthenticationResponseChallenge">
+            <summary>
+            Exposes the security.Challenge environment value as a strong type.
+            </summary>
+        </member>
+        <member name="P:Microsoft.Owin.Security.IAuthenticationManager.AuthenticationResponseGrant">
+            <summary>
+            Exposes the security.SignIn environment value as a strong type.
+            </summary>
+        </member>
+        <member name="P:Microsoft.Owin.Security.IAuthenticationManager.AuthenticationResponseRevoke">
+            <summary>
+            Exposes the security.SignOut environment value as a strong type.
+            </summary>
+        </member>
+        <member name="M:Microsoft.Owin.Security.AuthenticationManager.Authenticate(System.String[],System.Action{System.Security.Principal.IIdentity,System.Collections.Generic.IDictionary{System.String,System.String},System.Collections.Generic.IDictionary{System.String,System.Object},System.Object},System.Object)">
+            <summary>
+            
+            </summary>
+            <param name="authenticationTypes"></param>
+            <param name="callback"></param>
+            <param name="state"></param>
+            <returns></returns>
+        </member>
+        <member name="P:Microsoft.Owin.Security.AuthenticationManager.AuthenticationResponseChallenge">
+            <summary>
+            Exposes the security.Challenge environment value as a strong type.
+            </summary>
+        </member>
+        <member name="P:Microsoft.Owin.Security.AuthenticationManager.AuthenticationResponseGrant">
+            <summary>
+            Exposes the security.SignIn environment value as a strong type.
+            </summary>
+        </member>
+        <member name="P:Microsoft.Owin.Security.AuthenticationManager.AuthenticationResponseRevoke">
+            <summary>
+            Exposes the security.SignOut environment value as a strong type.
+            </summary>
+        </member>
+        <member name="T:Microsoft.Owin.Security.AuthenticationProperties">
+            <summary>
+            Dictionary used to store state values about the authentication session.
+            </summary>
+        </member>
+        <member name="M:Microsoft.Owin.Security.AuthenticationProperties.#ctor">
+            <summary>
+            Initializes a new instance of the <see cref="T:Microsoft.Owin.Security.AuthenticationProperties"/> class
+            </summary>
+        </member>
+        <member name="M:Microsoft.Owin.Security.AuthenticationProperties.#ctor(System.Collections.Generic.IDictionary{System.String,System.String})">
+            <summary>
+            Initializes a new instance of the <see cref="T:Microsoft.Owin.Security.AuthenticationProperties"/> class
+            </summary>
+            <param name="dictionary"></param>
+        </member>
+        <member name="P:Microsoft.Owin.Security.AuthenticationProperties.Dictionary">
+            <summary>
+            State values about the authentication session.
+            </summary>
+        </member>
+        <member name="P:Microsoft.Owin.Security.AuthenticationProperties.IsPersistent">
+            <summary>
+            Gets or sets whether the authentication session is persisted across multiple requests.
+            </summary>
+        </member>
+        <member name="P:Microsoft.Owin.Security.AuthenticationProperties.RedirectUri">
+            <summary>
+            Gets or sets the full path or absolute URI to be used as an http redirect response value. 
+            </summary>
+        </member>
+        <member name="P:Microsoft.Owin.Security.AuthenticationProperties.IssuedUtc">
+            <summary>
+            Gets or sets the time at which the authentication ticket was issued.
+            </summary>
+        </member>
+        <member name="P:Microsoft.Owin.Security.AuthenticationProperties.ExpiresUtc">
+            <summary>
+            Gets or sets the time at which the authentication ticket expires.
+            </summary>
+        </member>
+        <member name="P:Microsoft.Owin.Security.AuthenticationProperties.AllowRefresh">
+            <summary>
+            Gets or sets if refreshing the authentication session should be allowed.
+            </summary>
+        </member>
+        <member name="T:Microsoft.Owin.Security.AuthenticationResponseChallenge">
+            <summary>
+            Exposes the security.Challenge environment value as a strong type.
+            </summary>
+        </member>
+        <member name="M:Microsoft.Owin.Security.AuthenticationResponseChallenge.#ctor(System.String[],Microsoft.Owin.Security.AuthenticationProperties)">
+            <summary>
+            Initializes a new instance of the <see cref="T:Microsoft.Owin.Security.AuthenticationResponseChallenge"/> class
+            </summary>
+            <param name="authenticationTypes"></param>
+            <param name="properties"></param>
+        </member>
+        <member name="P:Microsoft.Owin.Security.AuthenticationResponseChallenge.AuthenticationTypes">
+            <summary>
+            List of the authentication types that should send a challenge in the response.
+            </summary>
+        </member>
+        <member name="P:Microsoft.Owin.Security.AuthenticationResponseChallenge.Properties">
+            <summary>
+            Dictionary used to store state values about the authentication session.
+            </summary>
+        </member>
+        <member name="T:Microsoft.Owin.Security.AuthenticationResponseGrant">
+            <summary>
+            Exposes the security.SignIn environment value as a strong type.
+            </summary>
+        </member>
+        <member name="M:Microsoft.Owin.Security.AuthenticationResponseGrant.#ctor(System.Security.Claims.ClaimsIdentity,Microsoft.Owin.Security.AuthenticationProperties)">
+            <summary>
+            Initializes a new instance of the <see cref="T:Microsoft.Owin.Security.AuthenticationResponseGrant"/> class.
+            </summary>
+            <param name="identity"></param>
+            <param name="properties"></param>
+        </member>
+        <member name="M:Microsoft.Owin.Security.AuthenticationResponseGrant.#ctor(System.Security.Claims.ClaimsPrincipal,Microsoft.Owin.Security.AuthenticationProperties)">
+            <summary>
+            Initializes a new instance of the <see cref="T:Microsoft.Owin.Security.AuthenticationResponseGrant"/> class.
+            </summary>
+            <param name="principal"></param>
+            <param name="properties"></param>
+        </member>
+        <member name="P:Microsoft.Owin.Security.AuthenticationResponseGrant.Identity">
+            <summary>
+            The identity associated with the user sign in.
+            </summary>
+        </member>
+        <member name="P:Microsoft.Owin.Security.AuthenticationResponseGrant.Principal">
+            <summary>
+            The security principal associated with the user sign in.
+            </summary>
+        </member>
+        <member name="P:Microsoft.Owin.Security.AuthenticationResponseGrant.Properties">
+            <summary>
+            Dictionary used to store state values about the authentication session.
+            </summary>
+        </member>
+        <member name="T:Microsoft.Owin.Security.AuthenticationResponseRevoke">
+            <summary>
+            Exposes the security.SignOut and security.SignOutProperties environment values as a strong type.
+            </summary>
+        </member>
+        <member name="M:Microsoft.Owin.Security.AuthenticationResponseRevoke.#ctor(System.String[])">
+            <summary>
+            Initializes a new instance of the <see cref="T:Microsoft.Owin.Security.AuthenticationResponseRevoke"/> class
+            </summary>
+            <param name="authenticationTypes"></param>
+        </member>
+        <member name="M:Microsoft.Owin.Security.AuthenticationResponseRevoke.#ctor(System.String[],Microsoft.Owin.Security.AuthenticationProperties)">
+            <summary>
+            Initializes a new instance of the <see cref="T:Microsoft.Owin.Security.AuthenticationResponseRevoke"/> class
+            </summary>
+            <param name="authenticationTypes"></param>
+            <param name="properties"></param>
+        </member>
+        <member name="P:Microsoft.Owin.Security.AuthenticationResponseRevoke.AuthenticationTypes">
+            <summary>
+            List of the authentication types that should be revoked on sign out.
+            </summary>
+        </member>
+        <member name="P:Microsoft.Owin.Security.AuthenticationResponseRevoke.Properties">
+            <summary>
+            Dictionary used to store state values about the authentication session.
+            </summary>
+        </member>
+        <member name="T:Microsoft.Owin.Logging.DiagnosticsLoggerFactory">
+            <summary>
+            Provides an ILoggerFactory based on System.Diagnostics.TraceSource.
+            </summary>
+        </member>
+        <member name="M:Microsoft.Owin.Logging.DiagnosticsLoggerFactory.#ctor">
+            <summary>
+            Initializes a new instance of the <see cref="T:Microsoft.Owin.Logging.DiagnosticsLoggerFactory"/> class. 
+            </summary>
+            <summary>
+            Creates a factory named "Microsoft.Owin".
+            </summary>
+        </member>
+        <member name="M:Microsoft.Owin.Logging.DiagnosticsLoggerFactory.#ctor(System.Diagnostics.SourceSwitch,System.Diagnostics.TraceListener)">
+            <summary>
+            Initializes a new instance of the <see cref="T:Microsoft.Owin.Logging.DiagnosticsLoggerFactory"/> class.
+            </summary>
+            <param name="rootSourceSwitch"></param>
+            <param name="rootTraceListener"></param>
+        </member>
+        <member name="M:Microsoft.Owin.Logging.DiagnosticsLoggerFactory.Create(System.String)">
+            <summary>
+            Creates a new DiagnosticsLogger for the given component name.
+            </summary>
+            <param name="name"></param>
+            <returns></returns>
+        </member>
+        <member name="T:Microsoft.Owin.Logging.LoggerExtensions">
+            <summary>
+            ILogger extension methods for common scenarios.
+            </summary>
+        </member>
+        <member name="M:Microsoft.Owin.Logging.LoggerExtensions.IsEnabled(Microsoft.Owin.Logging.ILogger,System.Diagnostics.TraceEventType)">
+            <summary>
+            Checks if the given TraceEventType is enabled.
+            </summary>
+            <param name="logger"></param>
+            <param name="eventType"></param>
+            <returns></returns>
+        </member>
+        <member name="M:Microsoft.Owin.Logging.LoggerExtensions.WriteVerbose(Microsoft.Owin.Logging.ILogger,System.String)">
+            <summary>
+            Writes a verbose log message.
+            </summary>
+            <param name="logger"></param>
+            <param name="data"></param>
+        </member>
+        <member name="M:Microsoft.Owin.Logging.LoggerExtensions.WriteInformation(Microsoft.Owin.Logging.ILogger,System.String)">
+            <summary>
+            Writes an informational log message.
+            </summary>
+            <param name="logger"></param>
+            <param name="message"></param>
+        </member>
+        <member name="M:Microsoft.Owin.Logging.LoggerExtensions.WriteWarning(Microsoft.Owin.Logging.ILogger,System.String,System.String[])">
+            <summary>
+            Writes a warning log message.
+            </summary>
+            <param name="logger"></param>
+            <param name="message"></param>
+            <param name="args"></param>
+        </member>
+        <member name="M:Microsoft.Owin.Logging.LoggerExtensions.WriteWarning(Microsoft.Owin.Logging.ILogger,System.String,System.Exception)">
+            <summary>
+            Writes a warning log message.
+            </summary>
+            <param name="logger"></param>
+            <param name="message"></param>
+            <param name="error"></param>
+        </member>
+        <member name="M:Microsoft.Owin.Logging.LoggerExtensions.WriteError(Microsoft.Owin.Logging.ILogger,System.String)">
+            <summary>
+            Writes an error log message.
+            </summary>
+            <param name="logger"></param>
+            <param name="message"></param>
+        </member>
+        <member name="M:Microsoft.Owin.Logging.LoggerExtensions.WriteError(Microsoft.Owin.Logging.ILogger,System.String,System.Exception)">
+            <summary>
+            Writes an error log message.
+            </summary>
+            <param name="logger"></param>
+            <param name="message"></param>
+            <param name="error"></param>
+        </member>
+        <member name="M:Microsoft.Owin.Logging.LoggerExtensions.WriteCritical(Microsoft.Owin.Logging.ILogger,System.String)">
+            <summary>
+            Writes a critical log message.
+            </summary>
+            <param name="logger"></param>
+            <param name="message"></param>
+        </member>
+        <member name="M:Microsoft.Owin.Logging.LoggerExtensions.WriteCritical(Microsoft.Owin.Logging.ILogger,System.String,System.Exception)">
+            <summary>
+            Writes a critical log message.
+            </summary>
+            <param name="logger"></param>
+            <param name="message"></param>
+            <param name="error"></param>
+        </member>
+    </members>
+</doc>
