Copying //UE4/Dev-Core to //UE4/Main

diff --git a/Engine/Source/Programs/CrashReporter/CrashReportWebSite/packages/Microsoft.Owin.Host.SystemWeb.3.0.0/lib/net45/Microsoft.Owin.Host.SystemWeb.xml b/Engine/Source/Programs/CrashReporter/CrashReportWebSite/packages/Microsoft.Owin.Host.SystemWeb.3.0.0/lib/net45/Microsoft.Owin.Host.SystemWeb.xml
new file mode 100644
--- /dev/null
+++ b/Engine/Source/Programs/CrashReporter/CrashReportWebSite/packages/Microsoft.Owin.Host.SystemWeb.3.0.0/lib/net45/Microsoft.Owin.Host.SystemWeb.xml
@@ -0,0 +1,526 @@
+<?xml version="1.0"?>
+<doc>
+    <assembly>
+        <name>Microsoft.Owin.Host.SystemWeb</name>
+    </assembly>
+    <members>
+        <member name="T:Owin.Loader.DefaultLoader">
+            <summary>
+            Locates the startup class based on the following convention:
+            AssemblyName.Startup, with a method named Configuration
+            </summary>
+        </member>
+        <member name="M:Owin.Loader.DefaultLoader.#ctor">
+            <summary>
+            
+            </summary>
+        </member>
+        <member name="M:Owin.Loader.DefaultLoader.#ctor(System.Func{System.String,System.Collections.Generic.IList{System.String},System.Action{Owin.IAppBuilder}})">
+            <summary>
+            Allows for a fallback loader to be specified.
+            </summary>
+            <param name="next"></param>
+        </member>
+        <member name="M:Owin.Loader.DefaultLoader.#ctor(System.Func{System.String,System.Collections.Generic.IList{System.String},System.Action{Owin.IAppBuilder}},System.Func{System.Type,System.Object})">
+            <summary>
+            Allows for a fallback loader and a Dependency Injection activator to be specified.
+            </summary>
+            <param name="next"></param>
+            <param name="activator"></param>
+        </member>
+        <member name="M:Owin.Loader.DefaultLoader.#ctor(System.Func{System.String,System.Collections.Generic.IList{System.String},System.Action{Owin.IAppBuilder}},System.Func{System.Type,System.Object},System.Collections.Generic.IEnumerable{System.Reflection.Assembly})">
+            <summary>
+            
+            </summary>
+            <param name="next"></param>
+            <param name="activator"></param>
+            <param name="referencedAssemblies"></param>
+        </member>
+        <member name="M:Owin.Loader.DefaultLoader.Load(System.String,System.Collections.Generic.IList{System.String})">
+            <summary>
+            Executes the loader, searching for the entry point by name.
+            </summary>
+            <param name="startupName">The name of the assembly and type entry point</param>
+            <param name="errorDetails"></param>
+            <returns></returns>
+        </member>
+        <member name="M:Owin.Loader.DefaultLoader.DotByDot(System.String)">
+            <summary>
+            
+            </summary>
+            <param name="text"></param>
+            <returns></returns>
+        </member>
+        <member name="T:Owin.Loader.NullLoader">
+            <summary>
+            A default fallback loader that does nothing.
+            </summary>
+        </member>
+        <member name="M:Owin.Loader.NullLoader.Load(System.String,System.Collections.Generic.IList{System.String})">
+            <summary>
+            A placeholder method that always returns null.
+            </summary>
+            <param name="startup"></param>
+            <param name="errors"></param>
+            <returns>null.</returns>
+        </member>
+        <member name="P:Owin.Loader.NullLoader.Instance">
+            <summary>
+            A singleton instance of the NullLoader type.
+            </summary>
+        </member>
+        <member name="T:SharedResourceNamespace.LoaderResources">
+            <summary>
+              A strongly-typed resource class, for looking up localized strings, etc.
+            </summary>
+        </member>
+        <member name="P:SharedResourceNamespace.LoaderResources.ResourceManager">
+            <summary>
+              Returns the cached ResourceManager instance used by this class.
+            </summary>
+        </member>
+        <member name="P:SharedResourceNamespace.LoaderResources.Culture">
+            <summary>
+              Overrides the current thread's CurrentUICulture property for all
+              resource lookups using this strongly typed resource class.
+            </summary>
+        </member>
+        <member name="P:SharedResourceNamespace.LoaderResources.AssemblyNotFound">
+            <summary>
+              Looks up a localized string similar to For the app startup parameter value &apos;{0}&apos;, the assembly &apos;{1}&apos; was not found..
+            </summary>
+        </member>
+        <member name="P:SharedResourceNamespace.LoaderResources.ClassNotFoundInAssembly">
+            <summary>
+              Looks up a localized string similar to For the app startup parameter value &apos;{0}&apos;, the class &apos;{1}&apos; was not found in assembly &apos;{2}&apos;..
+            </summary>
+        </member>
+        <member name="P:SharedResourceNamespace.LoaderResources.Exception_AttributeNameConflict">
+            <summary>
+              Looks up a localized string similar to The OwinStartup attribute discovered in assembly &apos;{0}&apos; referencing startup type &apos;{1}&apos; conflicts with the attribute in assembly &apos;{2}&apos; referencing startup type &apos;{3}&apos; because they have the same FriendlyName &apos;{4}&apos;. Remove or rename one of the attributes, or reference the desired type directly..
+            </summary>
+        </member>
+        <member name="P:SharedResourceNamespace.LoaderResources.Exception_StartupTypeConflict">
+            <summary>
+              Looks up a localized string similar to The discovered startup type &apos;{0}&apos; conflicts with the type &apos;{1}&apos;. Remove or rename one of the types, or reference the desired type directly..
+            </summary>
+        </member>
+        <member name="P:SharedResourceNamespace.LoaderResources.FriendlyNameMismatch">
+            <summary>
+              Looks up a localized string similar to The OwinStartupAttribute.FriendlyName value &apos;{0}&apos; does not match the given value &apos;{1}&apos; in Assembly &apos;{2}&apos;..
+            </summary>
+        </member>
+        <member name="P:SharedResourceNamespace.LoaderResources.MethodNotFoundInClass">
+            <summary>
+              Looks up a localized string similar to No &apos;{0}&apos; method was found in class &apos;{1}&apos;..
+            </summary>
+        </member>
+        <member name="P:SharedResourceNamespace.LoaderResources.NoAssemblyWithStartupClass">
+            <summary>
+              Looks up a localized string similar to No assembly found containing a Startup or [AssemblyName].Startup class..
+            </summary>
+        </member>
+        <member name="P:SharedResourceNamespace.LoaderResources.NoOwinStartupAttribute">
+            <summary>
+              Looks up a localized string similar to No assembly found containing an OwinStartupAttribute..
+            </summary>
+        </member>
+        <member name="P:SharedResourceNamespace.LoaderResources.StartupTypePropertyEmpty">
+            <summary>
+              Looks up a localized string similar to The OwinStartupAttribute.StartupType value is empty in Assembly &apos;{0}&apos;..
+            </summary>
+        </member>
+        <member name="P:SharedResourceNamespace.LoaderResources.StartupTypePropertyMissing">
+            <summary>
+              Looks up a localized string similar to The type &apos;{0}&apos; referenced from assembly &apos;{1}&apos; does not define a property &apos;StartupType&apos; of type &apos;Type&apos;..
+            </summary>
+        </member>
+        <member name="P:SharedResourceNamespace.LoaderResources.TypeOrMethodNotFound">
+            <summary>
+              Looks up a localized string similar to The given type or method &apos;{0}&apos; was not found. Try specifying the Assembly..
+            </summary>
+        </member>
+        <member name="P:SharedResourceNamespace.LoaderResources.UnexpectedMethodSignature">
+            <summary>
+              Looks up a localized string similar to The &apos;{0}&apos; method on class &apos;{1}&apos; does not have the expected signature &apos;void {0}(IAppBuilder)&apos;..
+            </summary>
+        </member>
+        <member name="T:Microsoft.Owin.Host.SystemWeb.DataProtection.MachineKeyDataProtectionProvider">
+            <summary>
+            Used to provide the data protection services that are derived from the MachineKey API. It is the best choice of
+            data protection when you application is hosted by ASP.NET and all servers in the farm are running with the same Machine Key values.
+            </summary>
+        </member>
+        <member name="M:Microsoft.Owin.Host.SystemWeb.DataProtection.MachineKeyDataProtectionProvider.Create(System.String[])">
+            <summary>
+            Returns a new instance of IDataProtection for the provider.
+            </summary>
+            <param name="purposes">Additional entropy used to ensure protected data may only be unprotected for the correct purposes.</param>
+            <returns>An instance of a data protection service</returns>
+        </member>
+        <member name="T:System.Web.HttpContextBaseExtensions">
+            <summary>
+            Provides extension methods for <see cref="T:System.Web.HttpContextBase"/>.
+            </summary>
+        </member>
+        <member name="M:System.Web.HttpContextBaseExtensions.GetOwinContext(System.Web.HttpContextBase)">
+            <summary>
+            Gets the <see cref="T:Microsoft.Owin.IOwinContext"/> for the current request.
+            </summary>
+            <param name="context"></param>
+            <returns></returns>
+        </member>
+        <member name="M:System.Web.HttpContextBaseExtensions.GetOwinContext(System.Web.HttpRequestBase)">
+            <summary>
+            Gets the <see cref="T:Microsoft.Owin.IOwinContext"/> for the current request.
+            </summary>
+            <param name="request"></param>
+            <returns></returns>
+        </member>
+        <member name="T:System.Web.HttpContextExtensions">
+            <summary>
+            Provides extension methods for <see cref="T:System.Web.HttpContext"/>.
+            </summary>
+        </member>
+        <member name="M:System.Web.HttpContextExtensions.GetOwinContext(System.Web.HttpContext)">
+            <summary>
+            Gets the <see cref="T:Microsoft.Owin.IOwinContext"/> for the current request.
+            </summary>
+            <param name="context"></param>
+            <returns></returns>
+        </member>
+        <member name="M:System.Web.HttpContextExtensions.GetOwinContext(System.Web.HttpRequest)">
+            <summary>
+            Gets the <see cref="T:Microsoft.Owin.IOwinContext"/> for the current request.
+            </summary>
+            <param name="request"></param>
+            <returns></returns>
+        </member>
+        <member name="T:Microsoft.Owin.Host.SystemWeb.OwinHttpHandler">
+            <summary>
+            Processes requests from System.Web as OWIN requests.
+            </summary>
+        </member>
+        <member name="M:Microsoft.Owin.Host.SystemWeb.OwinHttpHandler.#ctor">
+            <summary>
+            Processes requests using the default OWIN application.
+            </summary>
+        </member>
+        <member name="M:Microsoft.Owin.Host.SystemWeb.OwinHttpHandler.BeginProcessRequest(System.Web.HttpContextBase,System.AsyncCallback,System.Object)">
+            <summary>
+            Initiates an asynchronous call to the HTTP handler.
+            </summary>
+            <param name="httpContext">
+            An System.Web.HttpContextBase object that provides references to intrinsic server
+            objects (for example, Request, Response, Session, and Server) used to service
+            HTTP requests.
+            </param>
+            <param name="callback">
+            The System.AsyncCallback to call when the asynchronous method call is complete.
+            If callback is null, the delegate is not called.
+            </param>
+            <param name="extraData">
+            Any extra data needed to process the request.
+            </param>
+            <returns>
+            An System.IAsyncResult that contains information about the status of the process.
+            </returns>
+        </member>
+        <member name="M:Microsoft.Owin.Host.SystemWeb.OwinHttpHandler.EndProcessRequest(System.IAsyncResult)">
+            <summary>
+            Provides an asynchronous process End method when the process ends.
+            </summary>
+            <param name="result">
+            An System.IAsyncResult that contains information about the status of the process.
+            </param>
+        </member>
+        <member name="P:Microsoft.Owin.Host.SystemWeb.OwinHttpHandler.IsReusable">
+            <summary>
+            Gets a value indicating whether another request can use the System.Web.IHttpHandler instance.
+            </summary>
+            <returns>
+            true.
+            </returns>
+        </member>
+        <member name="T:Microsoft.Owin.Host.SystemWeb.OwinRouteHandler">
+            <summary>
+            Processes a route through an OWIN pipeline.
+            </summary>
+        </member>
+        <member name="M:Microsoft.Owin.Host.SystemWeb.OwinRouteHandler.#ctor(System.Action{Owin.IAppBuilder})">
+            <summary>
+            Initialize an OwinRouteHandler
+            </summary>
+            <param name="startup">The method to initialize the pipeline that processes requests for the route.</param>
+        </member>
+        <member name="M:Microsoft.Owin.Host.SystemWeb.OwinRouteHandler.GetHttpHandler(System.Web.Routing.RequestContext)">
+            <summary>
+            Provides the object that processes the request.
+            </summary>
+            <returns>
+            An object that processes the request.
+            </returns>
+            <param name="requestContext">An object that encapsulates information about the request.</param>
+        </member>
+        <member name="T:Microsoft.Owin.Host.SystemWeb.PreApplicationStart">
+            <summary>
+            Registers the OWIN request processing module at application startup.
+            </summary>
+        </member>
+        <member name="M:Microsoft.Owin.Host.SystemWeb.PreApplicationStart.Initialize">
+            <summary>
+            Registers the OWIN request processing module.
+            </summary>
+        </member>
+        <member name="T:Microsoft.Owin.Host.SystemWeb.Resources">
+            <summary>
+              A strongly-typed resource class, for looking up localized strings, etc.
+            </summary>
+        </member>
+        <member name="P:Microsoft.Owin.Host.SystemWeb.Resources.ResourceManager">
+            <summary>
+              Returns the cached ResourceManager instance used by this class.
+            </summary>
+        </member>
+        <member name="P:Microsoft.Owin.Host.SystemWeb.Resources.Culture">
+            <summary>
+              Overrides the current thread's CurrentUICulture property for all
+              resource lookups using this strongly typed resource class.
+            </summary>
+        </member>
+        <member name="P:Microsoft.Owin.Host.SystemWeb.Resources.Exception_AppLoderFailure">
+            <summary>
+              Looks up a localized string similar to The following errors occurred while attempting to load the app..
+            </summary>
+        </member>
+        <member name="P:Microsoft.Owin.Host.SystemWeb.Resources.Exception_CannotRegisterAfterHeadersSent">
+            <summary>
+              Looks up a localized string similar to Cannot register for &apos;OnSendingHeaders&apos; event after response headers have been sent..
+            </summary>
+        </member>
+        <member name="P:Microsoft.Owin.Host.SystemWeb.Resources.Exception_DuplicateKey">
+            <summary>
+              Looks up a localized string similar to The key &apos;{0}&apos; is already present in the dictionary..
+            </summary>
+        </member>
+        <member name="P:Microsoft.Owin.Host.SystemWeb.Resources.Exception_HowToDisableAutoAppStartup">
+            <summary>
+              Looks up a localized string similar to To disable OWIN startup discovery, add the appSetting owin:AutomaticAppStartup with a value of &quot;false&quot; in your web.config..
+            </summary>
+        </member>
+        <member name="P:Microsoft.Owin.Host.SystemWeb.Resources.Exception_HowToSpecifyAppStartup">
+            <summary>
+              Looks up a localized string similar to To specify the OWIN startup Assembly, Class, or Method, add the appSetting owin:AppStartup with the fully qualified startup class or configuration method name in your web.config..
+            </summary>
+        </member>
+        <member name="P:Microsoft.Owin.Host.SystemWeb.Resources.Exception_UnsupportedPipelineStage">
+            <summary>
+              Looks up a localized string similar to The given stage &apos;{0}&apos; is not supported..
+            </summary>
+        </member>
+        <member name="P:Microsoft.Owin.Host.SystemWeb.Resources.HttpContext_OwinEnvironmentNotFound">
+            <summary>
+              Looks up a localized string similar to No owin.Environment item was found in the context..
+            </summary>
+        </member>
+        <member name="P:Microsoft.Owin.Host.SystemWeb.Resources.Trace_ClientCertException">
+            <summary>
+              Looks up a localized string similar to An exception was thrown while trying to load the client certificate:.
+            </summary>
+        </member>
+        <member name="P:Microsoft.Owin.Host.SystemWeb.Resources.Trace_EntryPointException">
+            <summary>
+              Looks up a localized string similar to The OWIN entry point threw an exception:.
+            </summary>
+        </member>
+        <member name="P:Microsoft.Owin.Host.SystemWeb.Resources.Trace_OwinCallContextCallbackException">
+            <summary>
+              Looks up a localized string similar to The IAsyncResult callback for OwinCallHandler threw an exception:.
+            </summary>
+        </member>
+        <member name="P:Microsoft.Owin.Host.SystemWeb.Resources.Trace_RegisterModuleException">
+            <summary>
+              Looks up a localized string similar to Failed to register the OWIN module:.
+            </summary>
+        </member>
+        <member name="P:Microsoft.Owin.Host.SystemWeb.Resources.Trace_RequestDisconnectCallbackExceptions">
+            <summary>
+              Looks up a localized string similar to The application threw one or more exceptions when notified of a client disconnect:.
+            </summary>
+        </member>
+        <member name="P:Microsoft.Owin.Host.SystemWeb.Resources.Trace_ShutdownDetectionSetupException">
+            <summary>
+              Looks up a localized string similar to Shutdown detection setup failed:.
+            </summary>
+        </member>
+        <member name="P:Microsoft.Owin.Host.SystemWeb.Resources.Trace_ShutdownException">
+            <summary>
+              Looks up a localized string similar to One or more exceptions were thrown during app pool shutdown:.
+            </summary>
+        </member>
+        <member name="P:Microsoft.Owin.Host.SystemWeb.Resources.Trace_TimerCallbackException">
+            <summary>
+              Looks up a localized string similar to An exception was thrown from a timer callback:.
+            </summary>
+        </member>
+        <member name="P:Microsoft.Owin.Host.SystemWeb.Resources.Trace_WebSocketException">
+            <summary>
+              Looks up a localized string similar to An exception was thrown while processing the WebSocket:.
+            </summary>
+        </member>
+        <member name="P:Microsoft.Owin.Host.SystemWeb.Resources.Trace_WebSocketsSupportDetected">
+            <summary>
+              Looks up a localized string similar to Support for WebSockets has been detected..
+            </summary>
+        </member>
+        <member name="P:Microsoft.Owin.Host.SystemWeb.Resources.Trace_WebSocketsSupportNotDetected">
+            <summary>
+              Looks up a localized string similar to No support for WebSockets has been detected..
+            </summary>
+        </member>
+        <member name="T:System.Web.Routing.RouteCollectionExtensions">
+            <summary>
+            Provides extension methods for registering OWIN applications as System.Web routes.
+            </summary>
+        </member>
+        <member name="M:System.Web.Routing.RouteCollectionExtensions.MapOwinPath(System.Web.Routing.RouteCollection,System.String)">
+            <summary>
+            Registers a route for the default OWIN application.
+            </summary>
+            <param name="routes">The route collection.</param>
+            <param name="pathBase">The route path to map to the default OWIN application.</param>
+            <returns>The created route.</returns>
+        </member>
+        <member name="M:System.Web.Routing.RouteCollectionExtensions.MapOwinPath``1(System.Web.Routing.RouteCollection,System.String,``0)">
+            <summary>
+            Registers a route for a specific OWIN application entry point.
+            </summary>
+            <typeparam name="TApp">The OWIN application entry point type.</typeparam>
+            <param name="routes">The route collection.</param>
+            <param name="pathBase">The route path to map to the given OWIN application.</param>
+            <param name="app">The OWIN application entry point.</param>
+            <returns>The created route.</returns>
+        </member>
+        <member name="M:System.Web.Routing.RouteCollectionExtensions.MapOwinPath(System.Web.Routing.RouteCollection,System.String,System.Action{Owin.IAppBuilder})">
+            <summary>
+            Invokes the System.Action startup delegate to build the OWIN application
+            and then registers a route for it on the given path.
+            </summary>
+            <param name="routes">The route collection.</param>
+            <param name="pathBase">The route path to map to the given OWIN application.</param>
+            <param name="startup">A System.Action delegate invoked to build the OWIN application.</param>
+            <returns>The created route.</returns>
+        </member>
+        <member name="M:System.Web.Routing.RouteCollectionExtensions.MapOwinPath(System.Web.Routing.RouteCollection,System.String,System.String)">
+            <summary>
+            Registers a route for the default OWIN application.
+            </summary>
+            <param name="routes">The route collection.</param>
+            <param name="name">The given name of the route.</param>
+            <param name="pathBase">The route path to map to the default OWIN application.</param>
+            <returns>The created route.</returns>
+        </member>
+        <member name="M:System.Web.Routing.RouteCollectionExtensions.MapOwinPath``1(System.Web.Routing.RouteCollection,System.String,System.String,``0)">
+            <summary>
+            Registers a route for a specific OWIN application entry point.
+            </summary>
+            <typeparam name="TApp">The OWIN application entry point type.</typeparam>
+            <param name="routes">The route collection.</param>
+            <param name="name">The given name of the route.</param>
+            <param name="pathBase">The route path to map to the given OWIN application.</param>
+            <param name="app">The OWIN application entry point.</param>
+            <returns>The created route.</returns>
+        </member>
+        <member name="M:System.Web.Routing.RouteCollectionExtensions.MapOwinPath(System.Web.Routing.RouteCollection,System.String,System.String,System.Action{Owin.IAppBuilder})">
+            <summary>
+            Invokes the System.Action startup delegate to build the OWIN application
+            and then registers a route for it on the given path.
+            </summary>
+            <param name="routes">The route collection.</param>
+            <param name="name">The given name of the route.</param>
+            <param name="pathBase">The route path to map to the given OWIN application.</param>
+            <param name="startup">A System.Action delegate invoked to build the OWIN application.</param>
+            <returns>The created route.</returns>
+        </member>
+        <member name="M:System.Web.Routing.RouteCollectionExtensions.MapOwinRoute(System.Web.Routing.RouteCollection,System.String,System.Action{Owin.IAppBuilder})">
+            <summary>
+            Provides a way to define routes for an OWIN pipeline.
+            </summary>
+            <param name="routes">The route collection.</param>
+            <param name="routeUrl">The URL pattern for the route.</param>
+            <param name="startup">The method to initialize the pipeline that processes requests for the route.</param>
+        </member>
+        <member name="M:System.Web.Routing.RouteCollectionExtensions.MapOwinRoute(System.Web.Routing.RouteCollection,System.String,System.Web.Routing.RouteValueDictionary,System.Action{Owin.IAppBuilder})">
+            <summary>
+            Provides a way to define routes for an OWIN pipeline.
+            </summary>
+            <param name="routes">The route collection.</param>
+            <param name="routeUrl">The URL pattern for the route.</param>
+            <param name="defaults">The values to use if the URL does not contain all the parameters.</param>
+            <param name="startup">The method to initialize the pipeline that processes requests for the route.</param>
+        </member>
+        <member name="M:System.Web.Routing.RouteCollectionExtensions.MapOwinRoute(System.Web.Routing.RouteCollection,System.String,System.Web.Routing.RouteValueDictionary,System.Web.Routing.RouteValueDictionary,System.Action{Owin.IAppBuilder})">
+            <summary>
+            Provides a way to define routes for an OWIN pipeline.
+            </summary>
+            <param name="routes">The route collection.</param>
+            <param name="routeUrl">The URL pattern for the route.</param>
+            <param name="defaults">The values to use if the URL does not contain all the parameters.</param>
+            <param name="constraints">A regular expression that specifies valid values for a URL parameter.</param>
+            <param name="startup">The method to initialize the pipeline that processes requests for the route.</param>
+        </member>
+        <member name="M:System.Web.Routing.RouteCollectionExtensions.MapOwinRoute(System.Web.Routing.RouteCollection,System.String,System.Web.Routing.RouteValueDictionary,System.Web.Routing.RouteValueDictionary,System.Web.Routing.RouteValueDictionary,System.Action{Owin.IAppBuilder})">
+            <summary>
+            Provides a way to define routes for an OWIN pipeline.
+            </summary>
+            <param name="routes">The route collection.</param>
+            <param name="routeUrl">The URL pattern for the route.</param>
+            <param name="defaults">The values to use if the URL does not contain all the parameters.</param>
+            <param name="constraints">A regular expression that specifies valid values for a URL parameter.</param>
+            <param name="dataTokens">Custom values that are passed to the route handler, but which are not used to determine whether the route matches a specific URL pattern. These values are passed to the route handler, where they can be used for processing the request.</param>
+            <param name="startup">The method to initialize the pipeline that processes requests for the route.</param>
+        </member>
+        <member name="M:System.Web.Routing.RouteCollectionExtensions.MapOwinRoute(System.Web.Routing.RouteCollection,System.String,System.String,System.Action{Owin.IAppBuilder})">
+            <summary>
+            Provides a way to define routes for an OWIN pipeline.
+            </summary>
+            <param name="routes">The route collection.</param>
+            <param name="routeName">The name of the route.</param>
+            <param name="routeUrl">The URL pattern for the route.</param>
+            <param name="startup">The method to initialize the pipeline that processes requests for the route.</param>
+        </member>
+        <member name="M:System.Web.Routing.RouteCollectionExtensions.MapOwinRoute(System.Web.Routing.RouteCollection,System.String,System.String,System.Web.Routing.RouteValueDictionary,System.Action{Owin.IAppBuilder})">
+            <summary>
+            Provides a way to define routes for an OWIN pipeline.
+            </summary>
+            <param name="routes">The route collection.</param>
+            <param name="routeName">The name of the route.</param>
+            <param name="routeUrl">The URL pattern for the route.</param>
+            <param name="defaults">The values to use if the URL does not contain all the parameters.</param>
+            <param name="startup">The method to initialize the pipeline that processes requests for the route.</param>
+        </member>
+        <member name="M:System.Web.Routing.RouteCollectionExtensions.MapOwinRoute(System.Web.Routing.RouteCollection,System.String,System.String,System.Web.Routing.RouteValueDictionary,System.Web.Routing.RouteValueDictionary,System.Action{Owin.IAppBuilder})">
+            <summary>
+            Provides a way to define routes for an OWIN pipeline.
+            </summary>
+            <param name="routes">The route collection.</param>
+            <param name="routeName">The name of the route.</param>
+            <param name="routeUrl">The URL pattern for the route.</param>
+            <param name="defaults">The values to use if the URL does not contain all the parameters.</param>
+            <param name="constraints">A regular expression that specifies valid values for a URL parameter.</param>
+            <param name="startup">The method to initialize the pipeline that processes requests for the route.</param>
+        </member>
+        <member name="M:System.Web.Routing.RouteCollectionExtensions.MapOwinRoute(System.Web.Routing.RouteCollection,System.String,System.String,System.Web.Routing.RouteValueDictionary,System.Web.Routing.RouteValueDictionary,System.Web.Routing.RouteValueDictionary,System.Action{Owin.IAppBuilder})">
+            <summary>
+            Provides a way to define routes for an OWIN pipeline.
+            </summary>
+            <param name="routes">The route collection.</param>
+            <param name="routeName">The name of the route.</param>
+            <param name="routeUrl">The URL pattern for the route.</param>
+            <param name="defaults">The values to use if the URL does not contain all the parameters.</param>
+            <param name="constraints">A regular expression that specifies valid values for a URL parameter.</param>
+            <param name="dataTokens">Custom values that are passed to the route handler, but which are not used to determine whether the route matches a specific URL pattern. These values are passed to the route handler, where they can be used for processing the request.</param>
+            <param name="startup">The method to initialize the pipeline that processes requests for the route.</param>
+        </member>
+    </members>
+</doc>
