Copying //UE4/Dev-Core to //UE4/Main

diff --git a/Engine/Source/Programs/CrashReporter/CrashReportWebSite/packages/Newtonsoft.Json.5.0.1/lib/net20/Newtonsoft.Json.xml b/Engine/Source/Programs/CrashReporter/CrashReportWebSite/packages/Newtonsoft.Json.5.0.1/lib/net20/Newtonsoft.Json.xml
new file mode 100644
--- /dev/null
+++ b/Engine/Source/Programs/CrashReporter/CrashReportWebSite/packages/Newtonsoft.Json.5.0.1/lib/net20/Newtonsoft.Json.xml
@@ -0,0 +1,8781 @@
+<?xml version="1.0"?>
+<doc>
+    <assembly>
+        <name>Newtonsoft.Json</name>
+    </assembly>
+    <members>
+        <member name="T:Newtonsoft.Json.Bson.BsonObjectId">
+            <summary>
+            Represents a BSON Oid (object id).
+            </summary>
+        </member>
+        <member name="M:Newtonsoft.Json.Bson.BsonObjectId.#ctor(System.Byte[])">
+            <summary>
+            Initializes a new instance of the <see cref="T:Newtonsoft.Json.Bson.BsonObjectId"/> class.
+            </summary>
+            <param name="value">The Oid value.</param>
+        </member>
+        <member name="P:Newtonsoft.Json.Bson.BsonObjectId.Value">
+            <summary>
+            Gets or sets the value of the Oid.
+            </summary>
+            <value>The value of the Oid.</value>
+        </member>
+        <member name="T:Newtonsoft.Json.Bson.BsonReader">
+            <summary>
+            Represents a reader that provides fast, non-cached, forward-only access to serialized Json data.
+            </summary>
+        </member>
+        <member name="T:Newtonsoft.Json.JsonReader">
+            <summary>
+            Represents a reader that provides fast, non-cached, forward-only access to serialized Json data.
+            </summary>
+        </member>
+        <member name="M:Newtonsoft.Json.JsonReader.#ctor">
+            <summary>
+            Initializes a new instance of the <see cref="T:Newtonsoft.Json.JsonReader"/> class with the specified <see cref="T:System.IO.TextReader"/>.
+            </summary>
+        </member>
+        <member name="M:Newtonsoft.Json.JsonReader.Read">
+            <summary>
+            Reads the next JSON token from the stream.
+            </summary>
+            <returns>true if the next token was read successfully; false if there are no more tokens to read.</returns>
+        </member>
+        <member name="M:Newtonsoft.Json.JsonReader.ReadAsInt32">
+            <summary>
+            Reads the next JSON token from the stream as a <see cref="T:System.Nullable`1"/>.
+            </summary>
+            <returns>A <see cref="T:System.Nullable`1"/>. This method will return <c>null</c> at the end of an array.</returns>
+        </member>
+        <member name="M:Newtonsoft.Json.JsonReader.ReadAsString">
+            <summary>
+            Reads the next JSON token from the stream as a <see cref="T:System.String"/>.
+            </summary>
+            <returns>A <see cref="T:System.String"/>. This method will return <c>null</c> at the end of an array.</returns>
+        </member>
+        <member name="M:Newtonsoft.Json.JsonReader.ReadAsBytes">
+            <summary>
+            Reads the next JSON token from the stream as a <see cref="T:Byte[]"/>.
+            </summary>
+            <returns>A <see cref="T:Byte[]"/> or a null reference if the next JSON token is null. This method will return <c>null</c> at the end of an array.</returns>
+        </member>
+        <member name="M:Newtonsoft.Json.JsonReader.ReadAsDecimal">
+            <summary>
+            Reads the next JSON token from the stream as a <see cref="T:System.Nullable`1"/>.
+            </summary>
+            <returns>A <see cref="T:System.Nullable`1"/>. This method will return <c>null</c> at the end of an array.</returns>
+        </member>
+        <member name="M:Newtonsoft.Json.JsonReader.ReadAsDateTime">
+            <summary>
+            Reads the next JSON token from the stream as a <see cref="T:System.Nullable`1"/>.
+            </summary>
+            <returns>A <see cref="T:System.String"/>. This method will return <c>null</c> at the end of an array.</returns>
+        </member>
+        <member name="M:Newtonsoft.Json.JsonReader.Skip">
+            <summary>
+            Skips the children of the current token.
+            </summary>
+        </member>
+        <member name="M:Newtonsoft.Json.JsonReader.SetToken(Newtonsoft.Json.JsonToken)">
+            <summary>
+            Sets the current token.
+            </summary>
+            <param name="newToken">The new token.</param>
+        </member>
+        <member name="M:Newtonsoft.Json.JsonReader.SetToken(Newtonsoft.Json.JsonToken,System.Object)">
+            <summary>
+            Sets the current token and value.
+            </summary>
+            <param name="newToken">The new token.</param>
+            <param name="value">The value.</param>
+        </member>
+        <member name="M:Newtonsoft.Json.JsonReader.SetStateBasedOnCurrent">
+            <summary>
+            Sets the state based on current token type.
+            </summary>
+        </member>
+        <member name="M:Newtonsoft.Json.JsonReader.System#IDisposable#Dispose">
+            <summary>
+            Performs application-defined tasks associated with freeing, releasing, or resetting unmanaged resources.
+            </summary>
+        </member>
+        <member name="M:Newtonsoft.Json.JsonReader.Dispose(System.Boolean)">
+            <summary>
+            Releases unmanaged and - optionally - managed resources
+            </summary>
+            <param name="disposing"><c>true</c> to release both managed and unmanaged resources; <c>false</c> to release only unmanaged resources.</param>
+        </member>
+        <member name="M:Newtonsoft.Json.JsonReader.Close">
+            <summary>
+            Changes the <see cref="T:Newtonsoft.Json.JsonReader.State"/> to Closed. 
+            </summary>
+        </member>
+        <member name="P:Newtonsoft.Json.JsonReader.CurrentState">
+            <summary>
+            Gets the current reader state.
+            </summary>
+            <value>The current reader state.</value>
+        </member>
+        <member name="P:Newtonsoft.Json.JsonReader.CloseInput">
+            <summary>
+            Gets or sets a value indicating whether the underlying stream or
+            <see cref="T:System.IO.TextReader"/> should be closed when the reader is closed.
+            </summary>
+            <value>
+            true to close the underlying stream or <see cref="T:System.IO.TextReader"/> when
+            the reader is closed; otherwise false. The default is true.
+            </value>
+        </member>
+        <member name="P:Newtonsoft.Json.JsonReader.QuoteChar">
+            <summary>
+            Gets the quotation mark character used to enclose the value of a string.
+            </summary>
+        </member>
+        <member name="P:Newtonsoft.Json.JsonReader.DateTimeZoneHandling">
+            <summary>
+            Get or set how <see cref="T:System.DateTime"/> time zones are handling when reading JSON.
+            </summary>
+        </member>
+        <member name="P:Newtonsoft.Json.JsonReader.DateParseHandling">
+            <summary>
+            Get or set how date formatted strings, e.g. "\/Date(1198908717056)\/" and "2012-03-21T05:40Z", are parsed when reading JSON.
+            </summary>
+        </member>
+        <member name="P:Newtonsoft.Json.JsonReader.FloatParseHandling">
+            <summary>
+            Get or set how floating point numbers, e.g. 1.0 and 9.9, are parsed when reading JSON text.
+            </summary>
+        </member>
+        <member name="P:Newtonsoft.Json.JsonReader.MaxDepth">
+            <summary>
+            Gets or sets the maximum depth allowed when reading JSON. Reading past this depth will throw a <see cref="T:Newtonsoft.Json.JsonReaderException"/>.
+            </summary>
+        </member>
+        <member name="P:Newtonsoft.Json.JsonReader.TokenType">
+            <summary>
+            Gets the type of the current JSON token. 
+            </summary>
+        </member>
+        <member name="P:Newtonsoft.Json.JsonReader.Value">
+            <summary>
+            Gets the text value of the current JSON token.
+            </summary>
+        </member>
+        <member name="P:Newtonsoft.Json.JsonReader.ValueType">
+            <summary>
+            Gets The Common Language Runtime (CLR) type for the current JSON token.
+            </summary>
+        </member>
+        <member name="P:Newtonsoft.Json.JsonReader.Depth">
+            <summary>
+            Gets the depth of the current token in the JSON document.
+            </summary>
+            <value>The depth of the current token in the JSON document.</value>
+        </member>
+        <member name="P:Newtonsoft.Json.JsonReader.Path">
+            <summary>
+            Gets the path of the current JSON token. 
+            </summary>
+        </member>
+        <member name="P:Newtonsoft.Json.JsonReader.Culture">
+            <summary>
+            Gets or sets the culture used when reading JSON. Defaults to <see cref="P:System.Globalization.CultureInfo.InvariantCulture"/>.
+            </summary>
+        </member>
+        <member name="T:Newtonsoft.Json.JsonReader.State">
+            <summary>
+            Specifies the state of the reader.
+            </summary>
+        </member>
+        <member name="F:Newtonsoft.Json.JsonReader.State.Start">
+            <summary>
+            The Read method has not been called.
+            </summary>
+        </member>
+        <member name="F:Newtonsoft.Json.JsonReader.State.Complete">
+            <summary>
+            The end of the file has been reached successfully.
+            </summary>
+        </member>
+        <member name="F:Newtonsoft.Json.JsonReader.State.Property">
+            <summary>
+            Reader is at a property.
+            </summary>
+        </member>
+        <member name="F:Newtonsoft.Json.JsonReader.State.ObjectStart">
+            <summary>
+            Reader is at the start of an object.
+            </summary>
+        </member>
+        <member name="F:Newtonsoft.Json.JsonReader.State.Object">
+            <summary>
+            Reader is in an object.
+            </summary>
+        </member>
+        <member name="F:Newtonsoft.Json.JsonReader.State.ArrayStart">
+            <summary>
+            Reader is at the start of an array.
+            </summary>
+        </member>
+        <member name="F:Newtonsoft.Json.JsonReader.State.Array">
+            <summary>
+            Reader is in an array.
+            </summary>
+        </member>
+        <member name="F:Newtonsoft.Json.JsonReader.State.Closed">
+            <summary>
+            The Close method has been called.
+            </summary>
+        </member>
+        <member name="F:Newtonsoft.Json.JsonReader.State.PostValue">
+            <summary>
+            Reader has just read a value.
+            </summary>
+        </member>
+        <member name="F:Newtonsoft.Json.JsonReader.State.ConstructorStart">
+            <summary>
+            Reader is at the start of a constructor.
+            </summary>
+        </member>
+        <member name="F:Newtonsoft.Json.JsonReader.State.Constructor">
+            <summary>
+            Reader in a constructor.
+            </summary>
+        </member>
+        <member name="F:Newtonsoft.Json.JsonReader.State.Error">
+            <summary>
+            An error occurred that prevents the read operation from continuing.
+            </summary>
+        </member>
+        <member name="F:Newtonsoft.Json.JsonReader.State.Finished">
+            <summary>
+            The end of the file has been reached successfully.
+            </summary>
+        </member>
+        <member name="M:Newtonsoft.Json.Bson.BsonReader.#ctor(System.IO.Stream)">
+            <summary>
+            Initializes a new instance of the <see cref="T:Newtonsoft.Json.Bson.BsonReader"/> class.
+            </summary>
+            <param name="stream">The stream.</param>
+        </member>
+        <member name="M:Newtonsoft.Json.Bson.BsonReader.#ctor(System.IO.BinaryReader)">
+            <summary>
+            Initializes a new instance of the <see cref="T:Newtonsoft.Json.Bson.BsonReader"/> class.
+            </summary>
+            <param name="reader">The reader.</param>
+        </member>
+        <member name="M:Newtonsoft.Json.Bson.BsonReader.#ctor(System.IO.Stream,System.Boolean,System.DateTimeKind)">
+            <summary>
+            Initializes a new instance of the <see cref="T:Newtonsoft.Json.Bson.BsonReader"/> class.
+            </summary>
+            <param name="stream">The stream.</param>
+            <param name="readRootValueAsArray">if set to <c>true</c> the root object will be read as a JSON array.</param>
+            <param name="dateTimeKindHandling">The <see cref="T:System.DateTimeKind"/> used when reading <see cref="T:System.DateTime"/> values from BSON.</param>
+        </member>
+        <member name="M:Newtonsoft.Json.Bson.BsonReader.#ctor(System.IO.BinaryReader,System.Boolean,System.DateTimeKind)">
+            <summary>
+            Initializes a new instance of the <see cref="T:Newtonsoft.Json.Bson.BsonReader"/> class.
+            </summary>
+            <param name="reader">The reader.</param>
+            <param name="readRootValueAsArray">if set to <c>true</c> the root object will be read as a JSON array.</param>
+            <param name="dateTimeKindHandling">The <see cref="T:System.DateTimeKind"/> used when reading <see cref="T:System.DateTime"/> values from BSON.</param>
+        </member>
+        <member name="M:Newtonsoft.Json.Bson.BsonReader.ReadAsBytes">
+            <summary>
+            Reads the next JSON token from the stream as a <see cref="T:Byte[]"/>.
+            </summary>
+            <returns>
+            A <see cref="T:Byte[]"/> or a null reference if the next JSON token is null. This method will return <c>null</c> at the end of an array.
+            </returns>
+        </member>
+        <member name="M:Newtonsoft.Json.Bson.BsonReader.ReadAsDecimal">
+            <summary>
+            Reads the next JSON token from the stream as a <see cref="T:System.Nullable`1"/>.
+            </summary>
+            <returns>A <see cref="T:System.Nullable`1"/>. This method will return <c>null</c> at the end of an array.</returns>
+        </member>
+        <member name="M:Newtonsoft.Json.Bson.BsonReader.ReadAsInt32">
+            <summary>
+            Reads the next JSON token from the stream as a <see cref="T:System.Nullable`1"/>.
+            </summary>
+            <returns>A <see cref="T:System.Nullable`1"/>. This method will return <c>null</c> at the end of an array.</returns>
+        </member>
+        <member name="M:Newtonsoft.Json.Bson.BsonReader.ReadAsString">
+            <summary>
+            Reads the next JSON token from the stream as a <see cref="T:System.String"/>.
+            </summary>
+            <returns>A <see cref="T:System.String"/>. This method will return <c>null</c> at the end of an array.</returns>
+        </member>
+        <member name="M:Newtonsoft.Json.Bson.BsonReader.ReadAsDateTime">
+            <summary>
+            Reads the next JSON token from the stream as a <see cref="T:System.Nullable`1"/>.
+            </summary>
+            <returns>A <see cref="T:System.String"/>. This method will return <c>null</c> at the end of an array.</returns>
+        </member>
+        <member name="M:Newtonsoft.Json.Bson.BsonReader.Read">
+            <summary>
+            Reads the next JSON token from the stream.
+            </summary>
+            <returns>
+            true if the next token was read successfully; false if there are no more tokens to read.
+            </returns>
+        </member>
+        <member name="M:Newtonsoft.Json.Bson.BsonReader.Close">
+            <summary>
+            Changes the <see cref="T:Newtonsoft.Json.JsonReader.State"/> to Closed.
+            </summary>
+        </member>
+        <member name="P:Newtonsoft.Json.Bson.BsonReader.JsonNet35BinaryCompatibility">
+            <summary>
+            Gets or sets a value indicating whether binary data reading should compatible with incorrect Json.NET 3.5 written binary.
+            </summary>
+            <value>
+            	<c>true</c> if binary data reading will be compatible with incorrect Json.NET 3.5 written binary; otherwise, <c>false</c>.
+            </value>
+        </member>
+        <member name="P:Newtonsoft.Json.Bson.BsonReader.ReadRootValueAsArray">
+            <summary>
+            Gets or sets a value indicating whether the root object will be read as a JSON array.
+            </summary>
+            <value>
+            	<c>true</c> if the root object will be read as a JSON array; otherwise, <c>false</c>.
+            </value>
+        </member>
+        <member name="P:Newtonsoft.Json.Bson.BsonReader.DateTimeKindHandling">
+            <summary>
+            Gets or sets the <see cref="T:System.DateTimeKind"/> used when reading <see cref="T:System.DateTime"/> values from BSON.
+            </summary>
+            <value>The <see cref="T:System.DateTimeKind"/> used when reading <see cref="T:System.DateTime"/> values from BSON.</value>
+        </member>
+        <member name="T:Newtonsoft.Json.Bson.BsonWriter">
+            <summary>
+            Represents a writer that provides a fast, non-cached, forward-only way of generating JSON data.
+            </summary>
+        </member>
+        <member name="T:Newtonsoft.Json.JsonWriter">
+            <summary>
+            Represents a writer that provides a fast, non-cached, forward-only way of generating Json data.
+            </summary>
+        </member>
+        <member name="M:Newtonsoft.Json.JsonWriter.#ctor">
+            <summary>
+            Creates an instance of the <c>JsonWriter</c> class. 
+            </summary>
+        </member>
+        <member name="M:Newtonsoft.Json.JsonWriter.Flush">
+            <summary>
+            Flushes whatever is in the buffer to the underlying streams and also flushes the underlying stream.
+            </summary>
+        </member>
+        <member name="M:Newtonsoft.Json.JsonWriter.Close">
+            <summary>
+            Closes this stream and the underlying stream.
+            </summary>
+        </member>
+        <member name="M:Newtonsoft.Json.JsonWriter.WriteStartObject">
+            <summary>
+            Writes the beginning of a Json object.
+            </summary>
+        </member>
+        <member name="M:Newtonsoft.Json.JsonWriter.WriteEndObject">
+            <summary>
+            Writes the end of a Json object.
+            </summary>
+        </member>
+        <member name="M:Newtonsoft.Json.JsonWriter.WriteStartArray">
+            <summary>
+            Writes the beginning of a Json array.
+            </summary>
+        </member>
+        <member name="M:Newtonsoft.Json.JsonWriter.WriteEndArray">
+            <summary>
+            Writes the end of an array.
+            </summary>
+        </member>
+        <member name="M:Newtonsoft.Json.JsonWriter.WriteStartConstructor(System.String)">
+            <summary>
+            Writes the start of a constructor with the given name.
+            </summary>
+            <param name="name">The name of the constructor.</param>
+        </member>
+        <member name="M:Newtonsoft.Json.JsonWriter.WriteEndConstructor">
+            <summary>
+            Writes the end constructor.
+            </summary>
+        </member>
+        <member name="M:Newtonsoft.Json.JsonWriter.WritePropertyName(System.String)">
+            <summary>
+            Writes the property name of a name/value pair on a JSON object.
+            </summary>
+            <param name="name">The name of the property.</param>
+        </member>
+        <member name="M:Newtonsoft.Json.JsonWriter.WritePropertyName(System.String,System.Boolean)">
+            <summary>
+            Writes the property name of a name/value pair on a JSON object.
+            </summary>
+            <param name="name">The name of the property.</param>
+            <param name="escape">A flag to indicate whether the text should be escaped when it is written as a JSON property name.</param>
+        </member>
+        <member name="M:Newtonsoft.Json.JsonWriter.WriteEnd">
+            <summary>
+            Writes the end of the current Json object or array.
+            </summary>
+        </member>
+        <member name="M:Newtonsoft.Json.JsonWriter.WriteToken(Newtonsoft.Json.JsonReader)">
+            <summary>
+            Writes the current <see cref="T:Newtonsoft.Json.JsonReader"/> token and its children.
+            </summary>
+            <param name="reader">The <see cref="T:Newtonsoft.Json.JsonReader"/> to read the token from.</param>
+        </member>
+        <member name="M:Newtonsoft.Json.JsonWriter.WriteToken(Newtonsoft.Json.JsonReader,System.Boolean)">
+            <summary>
+            Writes the current <see cref="T:Newtonsoft.Json.JsonReader"/> token.
+            </summary>
+            <param name="reader">The <see cref="T:Newtonsoft.Json.JsonReader"/> to read the token from.</param>
+            <param name="writeChildren">A flag indicating whether the current token's children should be written.</param>
+        </member>
+        <member name="M:Newtonsoft.Json.JsonWriter.WriteEnd(Newtonsoft.Json.JsonToken)">
+            <summary>
+            Writes the specified end token.
+            </summary>
+            <param name="token">The end token to write.</param>
+        </member>
+        <member name="M:Newtonsoft.Json.JsonWriter.WriteIndent">
+            <summary>
+            Writes indent characters.
+            </summary>
+        </member>
+        <member name="M:Newtonsoft.Json.JsonWriter.WriteValueDelimiter">
+            <summary>
+            Writes the JSON value delimiter.
+            </summary>
+        </member>
+        <member name="M:Newtonsoft.Json.JsonWriter.WriteIndentSpace">
+            <summary>
+            Writes an indent space.
+            </summary>
+        </member>
+        <member name="M:Newtonsoft.Json.JsonWriter.WriteNull">
+            <summary>
+            Writes a null value.
+            </summary>
+        </member>
+        <member name="M:Newtonsoft.Json.JsonWriter.WriteUndefined">
+            <summary>
+            Writes an undefined value.
+            </summary>
+        </member>
+        <member name="M:Newtonsoft.Json.JsonWriter.WriteRaw(System.String)">
+            <summary>
+            Writes raw JSON without changing the writer's state.
+            </summary>
+            <param name="json">The raw JSON to write.</param>
+        </member>
+        <member name="M:Newtonsoft.Json.JsonWriter.WriteRawValue(System.String)">
+            <summary>
+            Writes raw JSON where a value is expected and updates the writer's state.
+            </summary>
+            <param name="json">The raw JSON to write.</param>
+        </member>
+        <member name="M:Newtonsoft.Json.JsonWriter.WriteValue(System.String)">
+            <summary>
+            Writes a <see cref="T:System.String"/> value.
+            </summary>
+            <param name="value">The <see cref="T:System.String"/> value to write.</param>
+        </member>
+        <member name="M:Newtonsoft.Json.JsonWriter.WriteValue(System.Int32)">
+            <summary>
+            Writes a <see cref="T:System.Int32"/> value.
+            </summary>
+            <param name="value">The <see cref="T:System.Int32"/> value to write.</param>
+        </member>
+        <member name="M:Newtonsoft.Json.JsonWriter.WriteValue(System.UInt32)">
+            <summary>
+            Writes a <see cref="T:System.UInt32"/> value.
+            </summary>
+            <param name="value">The <see cref="T:System.UInt32"/> value to write.</param>
+        </member>
+        <member name="M:Newtonsoft.Json.JsonWriter.WriteValue(System.Int64)">
+            <summary>
+            Writes a <see cref="T:System.Int64"/> value.
+            </summary>
+            <param name="value">The <see cref="T:System.Int64"/> value to write.</param>
+        </member>
+        <member name="M:Newtonsoft.Json.JsonWriter.WriteValue(System.UInt64)">
+            <summary>
+            Writes a <see cref="T:System.UInt64"/> value.
+            </summary>
+            <param name="value">The <see cref="T:System.UInt64"/> value to write.</param>
+        </member>
+        <member name="M:Newtonsoft.Json.JsonWriter.WriteValue(System.Single)">
+            <summary>
+            Writes a <see cref="T:System.Single"/> value.
+            </summary>
+            <param name="value">The <see cref="T:System.Single"/> value to write.</param>
+        </member>
+        <member name="M:Newtonsoft.Json.JsonWriter.WriteValue(System.Double)">
+            <summary>
+            Writes a <see cref="T:System.Double"/> value.
+            </summary>
+            <param name="value">The <see cref="T:System.Double"/> value to write.</param>
+        </member>
+        <member name="M:Newtonsoft.Json.JsonWriter.WriteValue(System.Boolean)">
+            <summary>
+            Writes a <see cref="T:System.Boolean"/> value.
+            </summary>
+            <param name="value">The <see cref="T:System.Boolean"/> value to write.</param>
+        </member>
+        <member name="M:Newtonsoft.Json.JsonWriter.WriteValue(System.Int16)">
+            <summary>
+            Writes a <see cref="T:System.Int16"/> value.
+            </summary>
+            <param name="value">The <see cref="T:System.Int16"/> value to write.</param>
+        </member>
+        <member name="M:Newtonsoft.Json.JsonWriter.WriteValue(System.UInt16)">
+            <summary>
+            Writes a <see cref="T:System.UInt16"/> value.
+            </summary>
+            <param name="value">The <see cref="T:System.UInt16"/> value to write.</param>
+        </member>
+        <member name="M:Newtonsoft.Json.JsonWriter.WriteValue(System.Char)">
+            <summary>
+            Writes a <see cref="T:System.Char"/> value.
+            </summary>
+            <param name="value">The <see cref="T:System.Char"/> value to write.</param>
+        </member>
+        <member name="M:Newtonsoft.Json.JsonWriter.WriteValue(System.Byte)">
+            <summary>
+            Writes a <see cref="T:System.Byte"/> value.
+            </summary>
+            <param name="value">The <see cref="T:System.Byte"/> value to write.</param>
+        </member>
+        <member name="M:Newtonsoft.Json.JsonWriter.WriteValue(System.SByte)">
+            <summary>
+            Writes a <see cref="T:System.SByte"/> value.
+            </summary>
+            <param name="value">The <see cref="T:System.SByte"/> value to write.</param>
+        </member>
+        <member name="M:Newtonsoft.Json.JsonWriter.WriteValue(System.Decimal)">
+            <summary>
+            Writes a <see cref="T:System.Decimal"/> value.
+            </summary>
+            <param name="value">The <see cref="T:System.Decimal"/> value to write.</param>
+        </member>
+        <member name="M:Newtonsoft.Json.JsonWriter.WriteValue(System.DateTime)">
+            <summary>
+            Writes a <see cref="T:System.DateTime"/> value.
+            </summary>
+            <param name="value">The <see cref="T:System.DateTime"/> value to write.</param>
+        </member>
+        <member name="M:Newtonsoft.Json.JsonWriter.WriteValue(System.Guid)">
+            <summary>
+            Writes a <see cref="T:System.Guid"/> value.
+            </summary>
+            <param name="value">The <see cref="T:System.Guid"/> value to write.</param>
+        </member>
+        <member name="M:Newtonsoft.Json.JsonWriter.WriteValue(System.TimeSpan)">
+            <summary>
+            Writes a <see cref="T:System.TimeSpan"/> value.
+            </summary>
+            <param name="value">The <see cref="T:System.TimeSpan"/> value to write.</param>
+        </member>
+        <member name="M:Newtonsoft.Json.JsonWriter.WriteValue(System.Nullable{System.Int32})">
+            <summary>
+            Writes a <see cref="T:System.Nullable`1"/> value.
+            </summary>
+            <param name="value">The <see cref="T:System.Nullable`1"/> value to write.</param>
+        </member>
+        <member name="M:Newtonsoft.Json.JsonWriter.WriteValue(System.Nullable{System.UInt32})">
+            <summary>
+            Writes a <see cref="T:System.Nullable`1"/> value.
+            </summary>
+            <param name="value">The <see cref="T:System.Nullable`1"/> value to write.</param>
+        </member>
+        <member name="M:Newtonsoft.Json.JsonWriter.WriteValue(System.Nullable{System.Int64})">
+            <summary>
+            Writes a <see cref="T:System.Nullable`1"/> value.
+            </summary>
+            <param name="value">The <see cref="T:System.Nullable`1"/> value to write.</param>
+        </member>
+        <member name="M:Newtonsoft.Json.JsonWriter.WriteValue(System.Nullable{System.UInt64})">
+            <summary>
+            Writes a <see cref="T:System.Nullable`1"/> value.
+            </summary>
+            <param name="value">The <see cref="T:System.Nullable`1"/> value to write.</param>
+        </member>
+        <member name="M:Newtonsoft.Json.JsonWriter.WriteValue(System.Nullable{System.Single})">
+            <summary>
+            Writes a <see cref="T:System.Nullable`1"/> value.
+            </summary>
+            <param name="value">The <see cref="T:System.Nullable`1"/> value to write.</param>
+        </member>
+        <member name="M:Newtonsoft.Json.JsonWriter.WriteValue(System.Nullable{System.Double})">
+            <summary>
+            Writes a <see cref="T:System.Nullable`1"/> value.
+            </summary>
+            <param name="value">The <see cref="T:System.Nullable`1"/> value to write.</param>
+        </member>
+        <member name="M:Newtonsoft.Json.JsonWriter.WriteValue(System.Nullable{System.Boolean})">
+            <summary>
+            Writes a <see cref="T:System.Nullable`1"/> value.
+            </summary>
+            <param name="value">The <see cref="T:System.Nullable`1"/> value to write.</param>
+        </member>
+        <member name="M:Newtonsoft.Json.JsonWriter.WriteValue(System.Nullable{System.Int16})">
+            <summary>
+            Writes a <see cref="T:System.Nullable`1"/> value.
+            </summary>
+            <param name="value">The <see cref="T:System.Nullable`1"/> value to write.</param>
+        </member>
+        <member name="M:Newtonsoft.Json.JsonWriter.WriteValue(System.Nullable{System.UInt16})">
+            <summary>
+            Writes a <see cref="T:System.Nullable`1"/> value.
+            </summary>
+            <param name="value">The <see cref="T:System.Nullable`1"/> value to write.</param>
+        </member>
+        <member name="M:Newtonsoft.Json.JsonWriter.WriteValue(System.Nullable{System.Char})">
+            <summary>
+            Writes a <see cref="T:System.Nullable`1"/> value.
+            </summary>
+            <param name="value">The <see cref="T:System.Nullable`1"/> value to write.</param>
+        </member>
+        <member name="M:Newtonsoft.Json.JsonWriter.WriteValue(System.Nullable{System.Byte})">
+            <summary>
+            Writes a <see cref="T:System.Nullable`1"/> value.
+            </summary>
+            <param name="value">The <see cref="T:System.Nullable`1"/> value to write.</param>
+        </member>
+        <member name="M:Newtonsoft.Json.JsonWriter.WriteValue(System.Nullable{System.SByte})">
+            <summary>
+            Writes a <see cref="T:System.Nullable`1"/> value.
+            </summary>
+            <param name="value">The <see cref="T:System.Nullable`1"/> value to write.</param>
+        </member>
+        <member name="M:Newtonsoft.Json.JsonWriter.WriteValue(System.Nullable{System.Decimal})">
+            <summary>
+            Writes a <see cref="T:System.Nullable`1"/> value.
+            </summary>
+            <param name="value">The <see cref="T:System.Nullable`1"/> value to write.</param>
+        </member>
+        <member name="M:Newtonsoft.Json.JsonWriter.WriteValue(System.Nullable{System.DateTime})">
+            <summary>
+            Writes a <see cref="T:System.Nullable`1"/> value.
+            </summary>
+            <param name="value">The <see cref="T:System.Nullable`1"/> value to write.</param>
+        </member>
+        <member name="M:Newtonsoft.Json.JsonWriter.WriteValue(System.Nullable{System.Guid})">
+            <summary>
+            Writes a <see cref="T:System.Nullable`1"/> value.
+            </summary>
+            <param name="value">The <see cref="T:System.Nullable`1"/> value to write.</param>
+        </member>
+        <member name="M:Newtonsoft.Json.JsonWriter.WriteValue(System.Nullable{System.TimeSpan})">
+            <summary>
+            Writes a <see cref="T:System.Nullable`1"/> value.
+            </summary>
+            <param name="value">The <see cref="T:System.Nullable`1"/> value to write.</param>
+        </member>
+        <member name="M:Newtonsoft.Json.JsonWriter.WriteValue(System.Byte[])">
+            <summary>
+            Writes a <see cref="T:Byte[]"/> value.
+            </summary>
+            <param name="value">The <see cref="T:Byte[]"/> value to write.</param>
+        </member>
+        <member name="M:Newtonsoft.Json.JsonWriter.WriteValue(System.Uri)">
+            <summary>
+            Writes a <see cref="T:System.Uri"/> value.
+            </summary>
+            <param name="value">The <see cref="T:System.Uri"/> value to write.</param>
+        </member>
+        <member name="M:Newtonsoft.Json.JsonWriter.WriteValue(System.Object)">
+            <summary>
+            Writes a <see cref="T:System.Object"/> value.
+            An error will raised if the value cannot be written as a single JSON token.
+            </summary>
+            <param name="value">The <see cref="T:System.Object"/> value to write.</param>
+        </member>
+        <member name="M:Newtonsoft.Json.JsonWriter.WriteComment(System.String)">
+            <summary>
+            Writes out a comment <code>/*...*/</code> containing the specified text. 
+            </summary>
+            <param name="text">Text to place inside the comment.</param>
+        </member>
+        <member name="M:Newtonsoft.Json.JsonWriter.WriteWhitespace(System.String)">
+            <summary>
+            Writes out the given white space.
+            </summary>
+            <param name="ws">The string of white space characters.</param>
+        </member>
+        <member name="P:Newtonsoft.Json.JsonWriter.CloseOutput">
+            <summary>
+            Gets or sets a value indicating whether the underlying stream or
+            <see cref="T:System.IO.TextReader"/> should be closed when the writer is closed.
+            </summary>
+            <value>
+            true to close the underlying stream or <see cref="T:System.IO.TextReader"/> when
+            the writer is closed; otherwise false. The default is true.
+            </value>
+        </member>
+        <member name="P:Newtonsoft.Json.JsonWriter.Top">
+            <summary>
+            Gets the top.
+            </summary>
+            <value>The top.</value>
+        </member>
+        <member name="P:Newtonsoft.Json.JsonWriter.WriteState">
+            <summary>
+            Gets the state of the writer.
+            </summary>
+        </member>
+        <member name="P:Newtonsoft.Json.JsonWriter.Path">
+            <summary>
+            Gets the path of the writer. 
+            </summary>
+        </member>
+        <member name="P:Newtonsoft.Json.JsonWriter.Formatting">
+            <summary>
+            Indicates how JSON text output is formatted.
+            </summary>
+        </member>
+        <member name="P:Newtonsoft.Json.JsonWriter.DateFormatHandling">
+            <summary>
+            Get or set how dates are written to JSON text.
+            </summary>
+        </member>
+        <member name="P:Newtonsoft.Json.JsonWriter.DateTimeZoneHandling">
+            <summary>
+            Get or set how <see cref="T:System.DateTime"/> time zones are handling when writing JSON text.
+            </summary>
+        </member>
+        <member name="P:Newtonsoft.Json.JsonWriter.StringEscapeHandling">
+            <summary>
+            Get or set how strings are escaped when writing JSON text.
+            </summary>
+        </member>
+        <member name="P:Newtonsoft.Json.JsonWriter.FloatFormatHandling">
+            <summary>
+            Get or set how special floating point numbers, e.g. <see cref="F:System.Double.NaN"/>,
+            <see cref="F:System.Double.PositiveInfinity"/> and <see cref="F:System.Double.NegativeInfinity"/>,
+            are written to JSON text.
+            </summary>
+        </member>
+        <member name="P:Newtonsoft.Json.JsonWriter.DateFormatString">
+            <summary>
+            Get or set how <see cref="T:System.DateTime"/> and <see cref="T:System.DateTimeOffset"/> values are formatting when writing JSON text.
+            </summary>
+        </member>
+        <member name="P:Newtonsoft.Json.JsonWriter.Culture">
+            <summary>
+            Gets or sets the culture used when writing JSON. Defaults to <see cref="P:System.Globalization.CultureInfo.InvariantCulture"/>.
+            </summary>
+        </member>
+        <member name="M:Newtonsoft.Json.Bson.BsonWriter.#ctor(System.IO.Stream)">
+            <summary>
+            Initializes a new instance of the <see cref="T:Newtonsoft.Json.Bson.BsonWriter"/> class.
+            </summary>
+            <param name="stream">The stream.</param>
+        </member>
+        <member name="M:Newtonsoft.Json.Bson.BsonWriter.#ctor(System.IO.BinaryWriter)">
+            <summary>
+            Initializes a new instance of the <see cref="T:Newtonsoft.Json.Bson.BsonWriter"/> class.
+            </summary>
+            <param name="writer">The writer.</param>
+        </member>
+        <member name="M:Newtonsoft.Json.Bson.BsonWriter.Flush">
+            <summary>
+            Flushes whatever is in the buffer to the underlying streams and also flushes the underlying stream.
+            </summary>
+        </member>
+        <member name="M:Newtonsoft.Json.Bson.BsonWriter.WriteEnd(Newtonsoft.Json.JsonToken)">
+            <summary>
+            Writes the end.
+            </summary>
+            <param name="token">The token.</param>
+        </member>
+        <member name="M:Newtonsoft.Json.Bson.BsonWriter.WriteComment(System.String)">
+            <summary>
+            Writes out a comment <code>/*...*/</code> containing the specified text.
+            </summary>
+            <param name="text">Text to place inside the comment.</param>
+        </member>
+        <member name="M:Newtonsoft.Json.Bson.BsonWriter.WriteStartConstructor(System.String)">
+            <summary>
+            Writes the start of a constructor with the given name.
+            </summary>
+            <param name="name">The name of the constructor.</param>
+        </member>
+        <member name="M:Newtonsoft.Json.Bson.BsonWriter.WriteRaw(System.String)">
+            <summary>
+            Writes raw JSON.
+            </summary>
+            <param name="json">The raw JSON to write.</param>
+        </member>
+        <member name="M:Newtonsoft.Json.Bson.BsonWriter.WriteRawValue(System.String)">
+            <summary>
+            Writes raw JSON where a value is expected and updates the writer's state.
+            </summary>
+            <param name="json">The raw JSON to write.</param>
+        </member>
+        <member name="M:Newtonsoft.Json.Bson.BsonWriter.WriteStartArray">
+            <summary>
+            Writes the beginning of a Json array.
+            </summary>
+        </member>
+        <member name="M:Newtonsoft.Json.Bson.BsonWriter.WriteStartObject">
+            <summary>
+            Writes the beginning of a Json object.
+            </summary>
+        </member>
+        <member name="M:Newtonsoft.Json.Bson.BsonWriter.WritePropertyName(System.String)">
+            <summary>
+            Writes the property name of a name/value pair on a Json object.
+            </summary>
+            <param name="name">The name of the property.</param>
+        </member>
+        <member name="M:Newtonsoft.Json.Bson.BsonWriter.Close">
+            <summary>
+            Closes this stream and the underlying stream.
+            </summary>
+        </member>
+        <member name="M:Newtonsoft.Json.Bson.BsonWriter.WriteNull">
+            <summary>
+            Writes a null value.
+            </summary>
+        </member>
+        <member name="M:Newtonsoft.Json.Bson.BsonWriter.WriteUndefined">
+            <summary>
+            Writes an undefined value.
+            </summary>
+        </member>
+        <member name="M:Newtonsoft.Json.Bson.BsonWriter.WriteValue(System.String)">
+            <summary>
+            Writes a <see cref="T:System.String"/> value.
+            </summary>
+            <param name="value">The <see cref="T:System.String"/> value to write.</param>
+        </member>
+        <member name="M:Newtonsoft.Json.Bson.BsonWriter.WriteValue(System.Int32)">
+            <summary>
+            Writes a <see cref="T:System.Int32"/> value.
+            </summary>
+            <param name="value">The <see cref="T:System.Int32"/> value to write.</param>
+        </member>
+        <member name="M:Newtonsoft.Json.Bson.BsonWriter.WriteValue(System.UInt32)">
+            <summary>
+            Writes a <see cref="T:System.UInt32"/> value.
+            </summary>
+            <param name="value">The <see cref="T:System.UInt32"/> value to write.</param>
+        </member>
+        <member name="M:Newtonsoft.Json.Bson.BsonWriter.WriteValue(System.Int64)">
+            <summary>
+            Writes a <see cref="T:System.Int64"/> value.
+            </summary>
+            <param name="value">The <see cref="T:System.Int64"/> value to write.</param>
+        </member>
+        <member name="M:Newtonsoft.Json.Bson.BsonWriter.WriteValue(System.UInt64)">
+            <summary>
+            Writes a <see cref="T:System.UInt64"/> value.
+            </summary>
+            <param name="value">The <see cref="T:System.UInt64"/> value to write.</param>
+        </member>
+        <member name="M:Newtonsoft.Json.Bson.BsonWriter.WriteValue(System.Single)">
+            <summary>
+            Writes a <see cref="T:System.Single"/> value.
+            </summary>
+            <param name="value">The <see cref="T:System.Single"/> value to write.</param>
+        </member>
+        <member name="M:Newtonsoft.Json.Bson.BsonWriter.WriteValue(System.Double)">
+            <summary>
+            Writes a <see cref="T:System.Double"/> value.
+            </summary>
+            <param name="value">The <see cref="T:System.Double"/> value to write.</param>
+        </member>
+        <member name="M:Newtonsoft.Json.Bson.BsonWriter.WriteValue(System.Boolean)">
+            <summary>
+            Writes a <see cref="T:System.Boolean"/> value.
+            </summary>
+            <param name="value">The <see cref="T:System.Boolean"/> value to write.</param>
+        </member>
+        <member name="M:Newtonsoft.Json.Bson.BsonWriter.WriteValue(System.Int16)">
+            <summary>
+            Writes a <see cref="T:System.Int16"/> value.
+            </summary>
+            <param name="value">The <see cref="T:System.Int16"/> value to write.</param>
+        </member>
+        <member name="M:Newtonsoft.Json.Bson.BsonWriter.WriteValue(System.UInt16)">
+            <summary>
+            Writes a <see cref="T:System.UInt16"/> value.
+            </summary>
+            <param name="value">The <see cref="T:System.UInt16"/> value to write.</param>
+        </member>
+        <member name="M:Newtonsoft.Json.Bson.BsonWriter.WriteValue(System.Char)">
+            <summary>
+            Writes a <see cref="T:System.Char"/> value.
+            </summary>
+            <param name="value">The <see cref="T:System.Char"/> value to write.</param>
+        </member>
+        <member name="M:Newtonsoft.Json.Bson.BsonWriter.WriteValue(System.Byte)">
+            <summary>
+            Writes a <see cref="T:System.Byte"/> value.
+            </summary>
+            <param name="value">The <see cref="T:System.Byte"/> value to write.</param>
+        </member>
+        <member name="M:Newtonsoft.Json.Bson.BsonWriter.WriteValue(System.SByte)">
+            <summary>
+            Writes a <see cref="T:System.SByte"/> value.
+            </summary>
+            <param name="value">The <see cref="T:System.SByte"/> value to write.</param>
+        </member>
+        <member name="M:Newtonsoft.Json.Bson.BsonWriter.WriteValue(System.Decimal)">
+            <summary>
+            Writes a <see cref="T:System.Decimal"/> value.
+            </summary>
+            <param name="value">The <see cref="T:System.Decimal"/> value to write.</param>
+        </member>
+        <member name="M:Newtonsoft.Json.Bson.BsonWriter.WriteValue(System.DateTime)">
+            <summary>
+            Writes a <see cref="T:System.DateTime"/> value.
+            </summary>
+            <param name="value">The <see cref="T:System.DateTime"/> value to write.</param>
+        </member>
+        <member name="M:Newtonsoft.Json.Bson.BsonWriter.WriteValue(System.Byte[])">
+            <summary>
+            Writes a <see cref="T:Byte[]"/> value.
+            </summary>
+            <param name="value">The <see cref="T:Byte[]"/> value to write.</param>
+        </member>
+        <member name="M:Newtonsoft.Json.Bson.BsonWriter.WriteValue(System.Guid)">
+            <summary>
+            Writes a <see cref="T:System.Guid"/> value.
+            </summary>
+            <param name="value">The <see cref="T:System.Guid"/> value to write.</param>
+        </member>
+        <member name="M:Newtonsoft.Json.Bson.BsonWriter.WriteValue(System.TimeSpan)">
+            <summary>
+            Writes a <see cref="T:System.TimeSpan"/> value.
+            </summary>
+            <param name="value">The <see cref="T:System.TimeSpan"/> value to write.</param>
+        </member>
+        <member name="M:Newtonsoft.Json.Bson.BsonWriter.WriteValue(System.Uri)">
+            <summary>
+            Writes a <see cref="T:System.Uri"/> value.
+            </summary>
+            <param name="value">The <see cref="T:System.Uri"/> value to write.</param>
+        </member>
+        <member name="M:Newtonsoft.Json.Bson.BsonWriter.WriteObjectId(System.Byte[])">
+            <summary>
+            Writes a <see cref="T:Byte[]"/> value that represents a BSON object id.
+            </summary>
+            <param name="value">The Object ID value to write.</param>
+        </member>
+        <member name="M:Newtonsoft.Json.Bson.BsonWriter.WriteRegex(System.String,System.String)">
+            <summary>
+            Writes a BSON regex.
+            </summary>
+            <param name="pattern">The regex pattern.</param>
+            <param name="options">The regex options.</param>
+        </member>
+        <member name="P:Newtonsoft.Json.Bson.BsonWriter.DateTimeKindHandling">
+            <summary>
+            Gets or sets the <see cref="T:System.DateTimeKind"/> used when writing <see cref="T:System.DateTime"/> values to BSON.
+            When set to <see cref="F:System.DateTimeKind.Unspecified"/> no conversion will occur.
+            </summary>
+            <value>The <see cref="T:System.DateTimeKind"/> used when writing <see cref="T:System.DateTime"/> values to BSON.</value>
+        </member>
+        <member name="T:Newtonsoft.Json.ConstructorHandling">
+            <summary>
+            Specifies how constructors are used when initializing objects during deserialization by the <see cref="T:Newtonsoft.Json.JsonSerializer"/>.
+            </summary>
+        </member>
+        <member name="F:Newtonsoft.Json.ConstructorHandling.Default">
+            <summary>
+            First attempt to use the public default constructor, then fall back to single paramatized constructor, then the non-public default constructor.
+            </summary>
+        </member>
+        <member name="F:Newtonsoft.Json.ConstructorHandling.AllowNonPublicDefaultConstructor">
+            <summary>
+            Json.NET will use a non-public default constructor before falling back to a paramatized constructor.
+            </summary>
+        </member>
+        <member name="T:Newtonsoft.Json.Converters.BinaryConverter">
+            <summary>
+            Converts a binary value to and from a base 64 string value.
+            </summary>
+        </member>
+        <member name="T:Newtonsoft.Json.JsonConverter">
+            <summary>
+            Converts an object to and from JSON.
+            </summary>
+        </member>
+        <member name="M:Newtonsoft.Json.JsonConverter.WriteJson(Newtonsoft.Json.JsonWriter,System.Object,Newtonsoft.Json.JsonSerializer)">
+            <summary>
+            Writes the JSON representation of the object.
+            </summary>
+            <param name="writer">The <see cref="T:Newtonsoft.Json.JsonWriter"/> to write to.</param>
+            <param name="value">The value.</param>
+            <param name="serializer">The calling serializer.</param>
+        </member>
+        <member name="M:Newtonsoft.Json.JsonConverter.ReadJson(Newtonsoft.Json.JsonReader,System.Type,System.Object,Newtonsoft.Json.JsonSerializer)">
+            <summary>
+            Reads the JSON representation of the object.
+            </summary>
+            <param name="reader">The <see cref="T:Newtonsoft.Json.JsonReader"/> to read from.</param>
+            <param name="objectType">Type of the object.</param>
+            <param name="existingValue">The existing value of object being read.</param>
+            <param name="serializer">The calling serializer.</param>
+            <returns>The object value.</returns>
+        </member>
+        <member name="M:Newtonsoft.Json.JsonConverter.CanConvert(System.Type)">
+            <summary>
+            Determines whether this instance can convert the specified object type.
+            </summary>
+            <param name="objectType">Type of the object.</param>
+            <returns>
+            	<c>true</c> if this instance can convert the specified object type; otherwise, <c>false</c>.
+            </returns>
+        </member>
+        <member name="M:Newtonsoft.Json.JsonConverter.GetSchema">
+            <summary>
+            Gets the <see cref="T:Newtonsoft.Json.Schema.JsonSchema"/> of the JSON produced by the JsonConverter.
+            </summary>
+            <returns>The <see cref="T:Newtonsoft.Json.Schema.JsonSchema"/> of the JSON produced by the JsonConverter.</returns>
+        </member>
+        <member name="P:Newtonsoft.Json.JsonConverter.CanRead">
+            <summary>
+            Gets a value indicating whether this <see cref="T:Newtonsoft.Json.JsonConverter"/> can read JSON.
+            </summary>
+            <value><c>true</c> if this <see cref="T:Newtonsoft.Json.JsonConverter"/> can read JSON; otherwise, <c>false</c>.</value>
+        </member>
+        <member name="P:Newtonsoft.Json.JsonConverter.CanWrite">
+            <summary>
+            Gets a value indicating whether this <see cref="T:Newtonsoft.Json.JsonConverter"/> can write JSON.
+            </summary>
+            <value><c>true</c> if this <see cref="T:Newtonsoft.Json.JsonConverter"/> can write JSON; otherwise, <c>false</c>.</value>
+        </member>
+        <member name="M:Newtonsoft.Json.Converters.BinaryConverter.WriteJson(Newtonsoft.Json.JsonWriter,System.Object,Newtonsoft.Json.JsonSerializer)">
+            <summary>
+            Writes the JSON representation of the object.
+            </summary>
+            <param name="writer">The <see cref="T:Newtonsoft.Json.JsonWriter"/> to write to.</param>
+            <param name="value">The value.</param>
+            <param name="serializer">The calling serializer.</param>
+        </member>
+        <member name="M:Newtonsoft.Json.Converters.BinaryConverter.ReadJson(Newtonsoft.Json.JsonReader,System.Type,System.Object,Newtonsoft.Json.JsonSerializer)">
+            <summary>
+            Reads the JSON representation of the object.
+            </summary>
+            <param name="reader">The <see cref="T:Newtonsoft.Json.JsonReader"/> to read from.</param>
+            <param name="objectType">Type of the object.</param>
+            <param name="existingValue">The existing value of object being read.</param>
+            <param name="serializer">The calling serializer.</param>
+            <returns>The object value.</returns>
+        </member>
+        <member name="M:Newtonsoft.Json.Converters.BinaryConverter.CanConvert(System.Type)">
+            <summary>
+            Determines whether this instance can convert the specified object type.
+            </summary>
+            <param name="objectType">Type of the object.</param>
+            <returns>
+            	<c>true</c> if this instance can convert the specified object type; otherwise, <c>false</c>.
+            </returns>
+        </member>
+        <member name="T:Newtonsoft.Json.Converters.BsonObjectIdConverter">
+            <summary>
+            Converts a <see cref="T:Newtonsoft.Json.Bson.BsonObjectId"/> to and from JSON and BSON.
+            </summary>
+        </member>
+        <member name="M:Newtonsoft.Json.Converters.BsonObjectIdConverter.WriteJson(Newtonsoft.Json.JsonWriter,System.Object,Newtonsoft.Json.JsonSerializer)">
+            <summary>
+            Writes the JSON representation of the object.
+            </summary>
+            <param name="writer">The <see cref="T:Newtonsoft.Json.JsonWriter"/> to write to.</param>
+            <param name="value">The value.</param>
+            <param name="serializer">The calling serializer.</param>
+        </member>
+        <member name="M:Newtonsoft.Json.Converters.BsonObjectIdConverter.ReadJson(Newtonsoft.Json.JsonReader,System.Type,System.Object,Newtonsoft.Json.JsonSerializer)">
+            <summary>
+            Reads the JSON representation of the object.
+            </summary>
+            <param name="reader">The <see cref="T:Newtonsoft.Json.JsonReader"/> to read from.</param>
+            <param name="objectType">Type of the object.</param>
+            <param name="existingValue">The existing value of object being read.</param>
+            <param name="serializer">The calling serializer.</param>
+            <returns>The object value.</returns>
+        </member>
+        <member name="M:Newtonsoft.Json.Converters.BsonObjectIdConverter.CanConvert(System.Type)">
+            <summary>
+            Determines whether this instance can convert the specified object type.
+            </summary>
+            <param name="objectType">Type of the object.</param>
+            <returns>
+            	<c>true</c> if this instance can convert the specified object type; otherwise, <c>false</c>.
+            </returns>
+        </member>
+        <member name="T:Newtonsoft.Json.Converters.CustomCreationConverter`1">
+            <summary>
+            Create a custom object
+            </summary>
+            <typeparam name="T">The object type to convert.</typeparam>
+        </member>
+        <member name="M:Newtonsoft.Json.Converters.CustomCreationConverter`1.WriteJson(Newtonsoft.Json.JsonWriter,System.Object,Newtonsoft.Json.JsonSerializer)">
+            <summary>
+            Writes the JSON representation of the object.
+            </summary>
+            <param name="writer">The <see cref="T:Newtonsoft.Json.JsonWriter"/> to write to.</param>
+            <param name="value">The value.</param>
+            <param name="serializer">The calling serializer.</param>
+        </member>
+        <member name="M:Newtonsoft.Json.Converters.CustomCreationConverter`1.ReadJson(Newtonsoft.Json.JsonReader,System.Type,System.Object,Newtonsoft.Json.JsonSerializer)">
+            <summary>
+            Reads the JSON representation of the object.
+            </summary>
+            <param name="reader">The <see cref="T:Newtonsoft.Json.JsonReader"/> to read from.</param>
+            <param name="objectType">Type of the object.</param>
+            <param name="existingValue">The existing value of object being read.</param>
+            <param name="serializer">The calling serializer.</param>
+            <returns>The object value.</returns>
+        </member>
+        <member name="M:Newtonsoft.Json.Converters.CustomCreationConverter`1.Create(System.Type)">
+            <summary>
+            Creates an object which will then be populated by the serializer.
+            </summary>
+            <param name="objectType">Type of the object.</param>
+            <returns>The created object.</returns>
+        </member>
+        <member name="M:Newtonsoft.Json.Converters.CustomCreationConverter`1.CanConvert(System.Type)">
+            <summary>
+            Determines whether this instance can convert the specified object type.
+            </summary>
+            <param name="objectType">Type of the object.</param>
+            <returns>
+            	<c>true</c> if this instance can convert the specified object type; otherwise, <c>false</c>.
+            </returns>
+        </member>
+        <member name="P:Newtonsoft.Json.Converters.CustomCreationConverter`1.CanWrite">
+            <summary>
+            Gets a value indicating whether this <see cref="T:Newtonsoft.Json.JsonConverter"/> can write JSON.
+            </summary>
+            <value>
+            	<c>true</c> if this <see cref="T:Newtonsoft.Json.JsonConverter"/> can write JSON; otherwise, <c>false</c>.
+            </value>
+        </member>
+        <member name="T:Newtonsoft.Json.Converters.DataSetConverter">
+            <summary>
+            Converts a <see cref="T:System.Data.DataSet"/> to and from JSON.
+            </summary>
+        </member>
+        <member name="M:Newtonsoft.Json.Converters.DataSetConverter.WriteJson(Newtonsoft.Json.JsonWriter,System.Object,Newtonsoft.Json.JsonSerializer)">
+            <summary>
+            Writes the JSON representation of the object.
+            </summary>
+            <param name="writer">The <see cref="T:Newtonsoft.Json.JsonWriter"/> to write to.</param>
+            <param name="value">The value.</param>
+            <param name="serializer">The calling serializer.</param>
+        </member>
+        <member name="M:Newtonsoft.Json.Converters.DataSetConverter.ReadJson(Newtonsoft.Json.JsonReader,System.Type,System.Object,Newtonsoft.Json.JsonSerializer)">
+            <summary>
+            Reads the JSON representation of the object.
+            </summary>
+            <param name="reader">The <see cref="T:Newtonsoft.Json.JsonReader"/> to read from.</param>
+            <param name="objectType">Type of the object.</param>
+            <param name="existingValue">The existing value of object being read.</param>
+            <param name="serializer">The calling serializer.</param>
+            <returns>The object value.</returns>
+        </member>
+        <member name="M:Newtonsoft.Json.Converters.DataSetConverter.CanConvert(System.Type)">
+            <summary>
+            Determines whether this instance can convert the specified value type.
+            </summary>
+            <param name="valueType">Type of the value.</param>
+            <returns>
+            	<c>true</c> if this instance can convert the specified value type; otherwise, <c>false</c>.
+            </returns>
+        </member>
+        <member name="T:Newtonsoft.Json.Converters.DataTableConverter">
+            <summary>
+            Converts a <see cref="T:System.Data.DataTable"/> to and from JSON.
+            </summary>
+        </member>
+        <member name="M:Newtonsoft.Json.Converters.DataTableConverter.WriteJson(Newtonsoft.Json.JsonWriter,System.Object,Newtonsoft.Json.JsonSerializer)">
+            <summary>
+            Writes the JSON representation of the object.
+            </summary>
+            <param name="writer">The <see cref="T:Newtonsoft.Json.JsonWriter"/> to write to.</param>
+            <param name="value">The value.</param>
+            <param name="serializer">The calling serializer.</param>
+        </member>
+        <member name="M:Newtonsoft.Json.Converters.DataTableConverter.ReadJson(Newtonsoft.Json.JsonReader,System.Type,System.Object,Newtonsoft.Json.JsonSerializer)">
+            <summary>
+            Reads the JSON representation of the object.
+            </summary>
+            <param name="reader">The <see cref="T:Newtonsoft.Json.JsonReader"/> to read from.</param>
+            <param name="objectType">Type of the object.</param>
+            <param name="existingValue">The existing value of object being read.</param>
+            <param name="serializer">The calling serializer.</param>
+            <returns>The object value.</returns>
+        </member>
+        <member name="M:Newtonsoft.Json.Converters.DataTableConverter.CanConvert(System.Type)">
+            <summary>
+            Determines whether this instance can convert the specified value type.
+            </summary>
+            <param name="valueType">Type of the value.</param>
+            <returns>
+            	<c>true</c> if this instance can convert the specified value type; otherwise, <c>false</c>.
+            </returns>
+        </member>
+        <member name="T:Newtonsoft.Json.Converters.DateTimeConverterBase">
+            <summary>
+            Provides a base class for converting a <see cref="T:System.DateTime"/> to and from JSON.
+            </summary>
+        </member>
+        <member name="M:Newtonsoft.Json.Converters.DateTimeConverterBase.CanConvert(System.Type)">
+            <summary>
+            Determines whether this instance can convert the specified object type.
+            </summary>
+            <param name="objectType">Type of the object.</param>
+            <returns>
+            	<c>true</c> if this instance can convert the specified object type; otherwise, <c>false</c>.
+            </returns>
+        </member>
+        <member name="T:Newtonsoft.Json.Converters.KeyValuePairConverter">
+            <summary>
+            Converts a <see cref="T:System.Collections.Generic.KeyValuePair`2"/> to and from JSON.
+            </summary>
+        </member>
+        <member name="M:Newtonsoft.Json.Converters.KeyValuePairConverter.WriteJson(Newtonsoft.Json.JsonWriter,System.Object,Newtonsoft.Json.JsonSerializer)">
+            <summary>
+            Writes the JSON representation of the object.
+            </summary>
+            <param name="writer">The <see cref="T:Newtonsoft.Json.JsonWriter"/> to write to.</param>
+            <param name="value">The value.</param>
+            <param name="serializer">The calling serializer.</param>
+        </member>
+        <member name="M:Newtonsoft.Json.Converters.KeyValuePairConverter.ReadJson(Newtonsoft.Json.JsonReader,System.Type,System.Object,Newtonsoft.Json.JsonSerializer)">
+            <summary>
+            Reads the JSON representation of the object.
+            </summary>
+            <param name="reader">The <see cref="T:Newtonsoft.Json.JsonReader"/> to read from.</param>
+            <param name="objectType">Type of the object.</param>
+            <param name="existingValue">The existing value of object being read.</param>
+            <param name="serializer">The calling serializer.</param>
+            <returns>The object value.</returns>
+        </member>
+        <member name="M:Newtonsoft.Json.Converters.KeyValuePairConverter.CanConvert(System.Type)">
+            <summary>
+            Determines whether this instance can convert the specified object type.
+            </summary>
+            <param name="objectType">Type of the object.</param>
+            <returns>
+            	<c>true</c> if this instance can convert the specified object type; otherwise, <c>false</c>.
+            </returns>
+        </member>
+        <member name="T:Newtonsoft.Json.Converters.RegexConverter">
+            <summary>
+            Converts a <see cref="T:System.Text.RegularExpressions.Regex"/> to and from JSON and BSON.
+            </summary>
+        </member>
+        <member name="M:Newtonsoft.Json.Converters.RegexConverter.WriteJson(Newtonsoft.Json.JsonWriter,System.Object,Newtonsoft.Json.JsonSerializer)">
+            <summary>
+            Writes the JSON representation of the object.
+            </summary>
+            <param name="writer">The <see cref="T:Newtonsoft.Json.JsonWriter"/> to write to.</param>
+            <param name="value">The value.</param>
+            <param name="serializer">The calling serializer.</param>
+        </member>
+        <member name="M:Newtonsoft.Json.Converters.RegexConverter.ReadJson(Newtonsoft.Json.JsonReader,System.Type,System.Object,Newtonsoft.Json.JsonSerializer)">
+            <summary>
+            Reads the JSON representation of the object.
+            </summary>
+            <param name="reader">The <see cref="T:Newtonsoft.Json.JsonReader"/> to read from.</param>
+            <param name="objectType">Type of the object.</param>
+            <param name="existingValue">The existing value of object being read.</param>
+            <param name="serializer">The calling serializer.</param>
+            <returns>The object value.</returns>
+        </member>
+        <member name="M:Newtonsoft.Json.Converters.RegexConverter.CanConvert(System.Type)">
+            <summary>
+            Determines whether this instance can convert the specified object type.
+            </summary>
+            <param name="objectType">Type of the object.</param>
+            <returns>
+            	<c>true</c> if this instance can convert the specified object type; otherwise, <c>false</c>.
+            </returns>
+        </member>
+        <member name="T:Newtonsoft.Json.Converters.StringEnumConverter">
+            <summary>
+            Converts an <see cref="T:System.Enum"/> to and from its name string value.
+            </summary>
+        </member>
+        <member name="M:Newtonsoft.Json.Converters.StringEnumConverter.WriteJson(Newtonsoft.Json.JsonWriter,System.Object,Newtonsoft.Json.JsonSerializer)">
+            <summary>
+            Writes the JSON representation of the object.
+            </summary>
+            <param name="writer">The <see cref="T:Newtonsoft.Json.JsonWriter"/> to write to.</param>
+            <param name="value">The value.</param>
+            <param name="serializer">The calling serializer.</param>
+        </member>
+        <member name="M:Newtonsoft.Json.Converters.StringEnumConverter.ReadJson(Newtonsoft.Json.JsonReader,System.Type,System.Object,Newtonsoft.Json.JsonSerializer)">
+            <summary>
+            Reads the JSON representation of the object.
+            </summary>
+            <param name="reader">The <see cref="T:Newtonsoft.Json.JsonReader"/> to read from.</param>
+            <param name="objectType">Type of the object.</param>
+            <param name="existingValue">The existing value of object being read.</param>
+            <param name="serializer">The calling serializer.</param>
+            <returns>The object value.</returns>
+        </member>
+        <member name="M:Newtonsoft.Json.Converters.StringEnumConverter.CanConvert(System.Type)">
+            <summary>
+            Determines whether this instance can convert the specified object type.
+            </summary>
+            <param name="objectType">Type of the object.</param>
+            <returns>
+            <c>true</c> if this instance can convert the specified object type; otherwise, <c>false</c>.
+            </returns>
+        </member>
+        <member name="P:Newtonsoft.Json.Converters.StringEnumConverter.CamelCaseText">
+            <summary>
+            Gets or sets a value indicating whether the written enum text should be camel case.
+            </summary>
+            <value><c>true</c> if the written enum text will be camel case; otherwise, <c>false</c>.</value>
+        </member>
+        <member name="T:Newtonsoft.Json.Converters.VersionConverter">
+            <summary>
+            Converts a <see cref="T:System.Version"/> to and from a string (e.g. "1.2.3.4").
+            </summary>
+        </member>
+        <member name="M:Newtonsoft.Json.Converters.VersionConverter.WriteJson(Newtonsoft.Json.JsonWriter,System.Object,Newtonsoft.Json.JsonSerializer)">
+            <summary>
+            Writes the JSON representation of the object.
+            </summary>
+            <param name="writer">The <see cref="T:Newtonsoft.Json.JsonWriter"/> to write to.</param>
+            <param name="value">The value.</param>
+            <param name="serializer">The calling serializer.</param>
+        </member>
+        <member name="M:Newtonsoft.Json.Converters.VersionConverter.ReadJson(Newtonsoft.Json.JsonReader,System.Type,System.Object,Newtonsoft.Json.JsonSerializer)">
+            <summary>
+            Reads the JSON representation of the object.
+            </summary>
+            <param name="reader">The <see cref="T:Newtonsoft.Json.JsonReader"/> to read from.</param>
+            <param name="objectType">Type of the object.</param>
+            <param name="existingValue">The existing property value of the JSON that is being converted.</param>
+            <param name="serializer">The calling serializer.</param>
+            <returns>The object value.</returns>
+        </member>
+        <member name="M:Newtonsoft.Json.Converters.VersionConverter.CanConvert(System.Type)">
+            <summary>
+            Determines whether this instance can convert the specified object type.
+            </summary>
+            <param name="objectType">Type of the object.</param>
+            <returns>
+            	<c>true</c> if this instance can convert the specified object type; otherwise, <c>false</c>.
+            </returns>
+        </member>
+        <member name="T:Newtonsoft.Json.DateFormatHandling">
+            <summary>
+            Specifies how dates are formatted when writing JSON text.
+            </summary>
+        </member>
+        <member name="F:Newtonsoft.Json.DateFormatHandling.IsoDateFormat">
+            <summary>
+            Dates are written in the ISO 8601 format, e.g. "2012-03-21T05:40Z".
+            </summary>
+        </member>
+        <member name="F:Newtonsoft.Json.DateFormatHandling.MicrosoftDateFormat">
+            <summary>
+            Dates are written in the Microsoft JSON format, e.g. "\/Date(1198908717056)\/".
+            </summary>
+        </member>
+        <member name="T:Newtonsoft.Json.DateParseHandling">
+            <summary>
+            Specifies how date formatted strings, e.g. "\/Date(1198908717056)\/" and "2012-03-21T05:40Z", are parsed when reading JSON text.
+            </summary>
+        </member>
+        <member name="F:Newtonsoft.Json.DateParseHandling.None">
+            <summary>
+            Date formatted strings are not parsed to a date type and are read as strings.
+            </summary>
+        </member>
+        <member name="F:Newtonsoft.Json.DateParseHandling.DateTime">
+            <summary>
+            Date formatted strings, e.g. "\/Date(1198908717056)\/" and "2012-03-21T05:40Z", are parsed to <see cref="F:Newtonsoft.Json.DateParseHandling.DateTime"/>.
+            </summary>
+        </member>
+        <member name="T:Newtonsoft.Json.DateTimeZoneHandling">
+            <summary>
+            Specifies how to treat the time value when converting between string and <see cref="T:System.DateTime"/>.
+            </summary>
+        </member>
+        <member name="F:Newtonsoft.Json.DateTimeZoneHandling.Local">
+            <summary>
+            Treat as local time. If the <see cref="T:System.DateTime"/> object represents a Coordinated Universal Time (UTC), it is converted to the local time.
+            </summary>
+        </member>
+        <member name="F:Newtonsoft.Json.DateTimeZoneHandling.Utc">
+            <summary>
+            Treat as a UTC. If the <see cref="T:System.DateTime"/> object represents a local time, it is converted to a UTC.
+            </summary>
+        </member>
+        <member name="F:Newtonsoft.Json.DateTimeZoneHandling.Unspecified">
+            <summary>
+            Treat as a local time if a <see cref="T:System.DateTime"/> is being converted to a string.
+            If a string is being converted to <see cref="T:System.DateTime"/>, convert to a local time if a time zone is specified.
+            </summary>
+        </member>
+        <member name="F:Newtonsoft.Json.DateTimeZoneHandling.RoundtripKind">
+            <summary>
+            Time zone information should be preserved when converting.
+            </summary>
+        </member>
+        <member name="T:Newtonsoft.Json.FloatFormatHandling">
+            <summary>
+            Specifies float format handling options when writing special floating point numbers, e.g. <see cref="F:System.Double.NaN"/>,
+            <see cref="F:System.Double.PositiveInfinity"/> and <see cref="F:System.Double.NegativeInfinity"/> with <see cref="T:Newtonsoft.Json.JsonWriter"/>.
+            </summary>
+        </member>
+        <member name="F:Newtonsoft.Json.FloatFormatHandling.String">
+            <summary>
+            Write special floating point values as strings in JSON, e.g. "NaN", "Infinity", "-Infinity".
+            </summary>
+        </member>
+        <member name="F:Newtonsoft.Json.FloatFormatHandling.Symbol">
+            <summary>
+            Write special floating point values as symbols in JSON, e.g. NaN, Infinity, -Infinity.
+            Note that this will produce non-valid JSON.
+            </summary>
+        </member>
+        <member name="F:Newtonsoft.Json.FloatFormatHandling.DefaultValue">
+            <summary>
+            Write special floating point values as the property's default value in JSON, e.g. 0.0 for a <see cref="T:System.Double"/> property, null for a <see cref="T:System.Nullable`1"/> property.
+            </summary>
+        </member>
+        <member name="T:Newtonsoft.Json.FloatParseHandling">
+            <summary>
+            Specifies how floating point numbers, e.g. 1.0 and 9.9, are parsed when reading JSON text.
+            </summary>
+        </member>
+        <member name="F:Newtonsoft.Json.FloatParseHandling.Double">
+            <summary>
+            Floating point numbers are parsed to <see cref="F:Newtonsoft.Json.FloatParseHandling.Double"/>.
+            </summary>
+        </member>
+        <member name="F:Newtonsoft.Json.FloatParseHandling.Decimal">
+            <summary>
+            Floating point numbers are parsed to <see cref="F:Newtonsoft.Json.FloatParseHandling.Decimal"/>.
+            </summary>
+        </member>
+        <member name="T:Newtonsoft.Json.Formatting">
+            <summary>
+            Specifies formatting options for the <see cref="T:Newtonsoft.Json.JsonTextWriter"/>.
+            </summary>
+        </member>
+        <member name="F:Newtonsoft.Json.Formatting.None">
+            <summary>
+            No special formatting is applied. This is the default.
+            </summary>
+        </member>
+        <member name="F:Newtonsoft.Json.Formatting.Indented">
+            <summary>
+            Causes child objects to be indented according to the <see cref="P:Newtonsoft.Json.JsonTextWriter.Indentation"/> and <see cref="P:Newtonsoft.Json.JsonTextWriter.IndentChar"/> settings.
+            </summary>
+        </member>
+        <member name="T:Newtonsoft.Json.JsonConstructorAttribute">
+            <summary>
+            Instructs the <see cref="T:Newtonsoft.Json.JsonSerializer"/> to use the specified constructor when deserializing that object.
+            </summary>
+        </member>
+        <member name="T:Newtonsoft.Json.JsonDictionaryAttribute">
+            <summary>
+            Instructs the <see cref="T:Newtonsoft.Json.JsonSerializer"/> how to serialize the collection.
+            </summary>
+        </member>
+        <member name="T:Newtonsoft.Json.JsonContainerAttribute">
+            <summary>
+            Instructs the <see cref="T:Newtonsoft.Json.JsonSerializer"/> how to serialize the object.
+            </summary>
+        </member>
+        <member name="M:Newtonsoft.Json.JsonContainerAttribute.#ctor">
+            <summary>
+            Initializes a new instance of the <see cref="T:Newtonsoft.Json.JsonContainerAttribute"/> class.
+            </summary>
+        </member>
+        <member name="M:Newtonsoft.Json.JsonContainerAttribute.#ctor(System.String)">
+            <summary>
+            Initializes a new instance of the <see cref="T:Newtonsoft.Json.JsonContainerAttribute"/> class with the specified container Id.
+            </summary>
+            <param name="id">The container Id.</param>
+        </member>
+        <member name="P:Newtonsoft.Json.JsonContainerAttribute.Id">
+            <summary>
+            Gets or sets the id.
+            </summary>
+            <value>The id.</value>
+        </member>
+        <member name="P:Newtonsoft.Json.JsonContainerAttribute.Title">
+            <summary>
+            Gets or sets the title.
+            </summary>
+            <value>The title.</value>
+        </member>
+        <member name="P:Newtonsoft.Json.JsonContainerAttribute.Description">
+            <summary>
+            Gets or sets the description.
+            </summary>
+            <value>The description.</value>
+        </member>
+        <member name="P:Newtonsoft.Json.JsonContainerAttribute.ItemConverterType">
+            <summary>
+            Gets the collection's items converter.
+            </summary>
+            <value>The collection's items converter.</value>
+        </member>
+        <member name="P:Newtonsoft.Json.JsonContainerAttribute.IsReference">
+            <summary>
+            Gets or sets a value that indicates whether to preserve object references.
+            </summary>
+            <value>
+            	<c>true</c> to keep object reference; otherwise, <c>false</c>. The default is <c>false</c>.
+            </value>
+        </member>
+        <member name="P:Newtonsoft.Json.JsonContainerAttribute.ItemIsReference">
+            <summary>
+            Gets or sets a value that indicates whether to preserve collection's items references.
+            </summary>
+            <value>
+            	<c>true</c> to keep collection's items object references; otherwise, <c>false</c>. The default is <c>false</c>.
+            </value>
+        </member>
+        <member name="P:Newtonsoft.Json.JsonContainerAttribute.ItemReferenceLoopHandling">
+            <summary>
+            Gets or sets the reference loop handling used when serializing the collection's items.
+            </summary>
+            <value>The reference loop handling.</value>
+        </member>
+        <member name="P:Newtonsoft.Json.JsonContainerAttribute.ItemTypeNameHandling">
+            <summary>
+            Gets or sets the type name handling used when serializing the collection's items.
+            </summary>
+            <value>The type name handling.</value>
+        </member>
+        <member name="M:Newtonsoft.Json.JsonDictionaryAttribute.#ctor">
+            <summary>
+            Initializes a new instance of the <see cref="T:Newtonsoft.Json.JsonDictionaryAttribute"/> class.
+            </summary>
+        </member>
+        <member name="M:Newtonsoft.Json.JsonDictionaryAttribute.#ctor(System.String)">
+            <summary>
+            Initializes a new instance of the <see cref="T:Newtonsoft.Json.JsonDictionaryAttribute"/> class with the specified container Id.
+            </summary>
+            <param name="id">The container Id.</param>
+        </member>
+        <member name="T:Newtonsoft.Json.JsonException">
+            <summary>
+            The exception thrown when an error occurs during Json serialization or deserialization.
+            </summary>
+        </member>
+        <member name="M:Newtonsoft.Json.JsonException.#ctor">
+            <summary>
+            Initializes a new instance of the <see cref="T:Newtonsoft.Json.JsonException"/> class.
+            </summary>
+        </member>
+        <member name="M:Newtonsoft.Json.JsonException.#ctor(System.String)">
+            <summary>
+            Initializes a new instance of the <see cref="T:Newtonsoft.Json.JsonException"/> class
+            with a specified error message.
+            </summary>
+            <param name="message">The error message that explains the reason for the exception.</param>
+        </member>
+        <member name="M:Newtonsoft.Json.JsonException.#ctor(System.String,System.Exception)">
+            <summary>
+            Initializes a new instance of the <see cref="T:Newtonsoft.Json.JsonException"/> class
+            with a specified error message and a reference to the inner exception that is the cause of this exception.
+            </summary>
+            <param name="message">The error message that explains the reason for the exception.</param>
+            <param name="innerException">The exception that is the cause of the current exception, or a null reference (Nothing in Visual Basic) if no inner exception is specified.</param>
+        </member>
+        <member name="M:Newtonsoft.Json.JsonException.#ctor(System.Runtime.Serialization.SerializationInfo,System.Runtime.Serialization.StreamingContext)">
+            <summary>
+            Initializes a new instance of the <see cref="T:Newtonsoft.Json.JsonException"/> class.
+            </summary>
+            <param name="info">The <see cref="T:System.Runtime.Serialization.SerializationInfo"/> that holds the serialized object data about the exception being thrown.</param>
+            <param name="context">The <see cref="T:System.Runtime.Serialization.StreamingContext"/> that contains contextual information about the source or destination.</param>
+            <exception cref="T:System.ArgumentNullException">The <paramref name="info"/> parameter is null. </exception>
+            <exception cref="T:System.Runtime.Serialization.SerializationException">The class name is null or <see cref="P:System.Exception.HResult"/> is zero (0). </exception>
+        </member>
+        <member name="T:Newtonsoft.Json.Serialization.DiagnosticsTraceWriter">
+            <summary>
+            Represents a trace writer that writes to the application's <see cref="T:System.Diagnostics.TraceListener"/> instances.
+            </summary>
+        </member>
+        <member name="T:Newtonsoft.Json.Serialization.ITraceWriter">
+            <summary>
+            Represents a trace writer.
+            </summary>
+        </member>
+        <member name="M:Newtonsoft.Json.Serialization.ITraceWriter.Trace(System.Diagnostics.TraceLevel,System.String,System.Exception)">
+            <summary>
+            Writes the specified trace level, message and optional exception.
+            </summary>
+            <param name="level">The <see cref="T:System.Diagnostics.TraceLevel"/> at which to write this trace.</param>
+            <param name="message">The trace message.</param>
+            <param name="ex">The trace exception. This parameter is optional.</param>
+        </member>
+        <member name="P:Newtonsoft.Json.Serialization.ITraceWriter.LevelFilter">
+            <summary>
+            Gets the <see cref="T:System.Diagnostics.TraceLevel"/> that will be used to filter the trace messages passed to the writer.
+            For example a filter level of <code>Info</code> will exclude <code>Verbose</code> messages and include <code>Info</code>,
+            <code>Warning</code> and <code>Error</code> messages.
+            </summary>
+            <value>The <see cref="T:System.Diagnostics.TraceLevel"/> that will be used to filter the trace messages passed to the writer.</value>
+        </member>
+        <member name="M:Newtonsoft.Json.Serialization.DiagnosticsTraceWriter.Trace(System.Diagnostics.TraceLevel,System.String,System.Exception)">
+            <summary>
+            Writes the specified trace level, message and optional exception.
+            </summary>
+            <param name="level">The <see cref="T:System.Diagnostics.TraceLevel"/> at which to write this trace.</param>
+            <param name="message">The trace message.</param>
+            <param name="ex">The trace exception. This parameter is optional.</param>
+        </member>
+        <member name="P:Newtonsoft.Json.Serialization.DiagnosticsTraceWriter.LevelFilter">
+            <summary>
+            Gets the <see cref="T:System.Diagnostics.TraceLevel"/> that will be used to filter the trace messages passed to the writer.
+            For example a filter level of <code>Info</code> will exclude <code>Verbose</code> messages and include <code>Info</code>,
+            <code>Warning</code> and <code>Error</code> messages.
+            </summary>
+            <value>
+            The <see cref="T:System.Diagnostics.TraceLevel"/> that will be used to filter the trace messages passed to the writer.
+            </value>
+        </member>
+        <member name="T:Newtonsoft.Json.Serialization.JsonContainerContract">
+            <summary>
+            Contract details for a <see cref="T:System.Type"/> used by the <see cref="T:Newtonsoft.Json.JsonSerializer"/>.
+            </summary>
+        </member>
+        <member name="T:Newtonsoft.Json.Serialization.JsonContract">
+            <summary>
+            Contract details for a <see cref="T:System.Type"/> used by the <see cref="T:Newtonsoft.Json.JsonSerializer"/>.
+            </summary>
+        </member>
+        <member name="P:Newtonsoft.Json.Serialization.JsonContract.UnderlyingType">
+            <summary>
+            Gets the underlying type for the contract.
+            </summary>
+            <value>The underlying type for the contract.</value>
+        </member>
+        <member name="P:Newtonsoft.Json.Serialization.JsonContract.CreatedType">
+            <summary>
+            Gets or sets the type created during deserialization.
+            </summary>
+            <value>The type created during deserialization.</value>
+        </member>
+        <member name="P:Newtonsoft.Json.Serialization.JsonContract.IsReference">
+            <summary>
+            Gets or sets whether this type contract is serialized as a reference.
+            </summary>
+            <value>Whether this type contract is serialized as a reference.</value>
+        </member>
+        <member name="P:Newtonsoft.Json.Serialization.JsonContract.Converter">
+            <summary>
+            Gets or sets the default <see cref="T:Newtonsoft.Json.JsonConverter"/> for this contract.
+            </summary>
+            <value>The converter.</value>
+        </member>
+        <member name="P:Newtonsoft.Json.Serialization.JsonContract.OnDeserializedCallbacks">
+            <summary>
+            Gets or sets all methods called immediately after deserialization of the object.
+            </summary>
+            <value>The methods called immediately after deserialization of the object.</value>
+        </member>
+        <member name="P:Newtonsoft.Json.Serialization.JsonContract.OnDeserializingCallbacks">
+            <summary>
+            Gets or sets all methods called during deserialization of the object.
+            </summary>
+            <value>The methods called during deserialization of the object.</value>
+        </member>
+        <member name="P:Newtonsoft.Json.Serialization.JsonContract.OnSerializedCallbacks">
+            <summary>
+            Gets or sets all methods called after serialization of the object graph.
+            </summary>
+            <value>The methods called after serialization of the object graph.</value>
+        </member>
+        <member name="P:Newtonsoft.Json.Serialization.JsonContract.OnSerializingCallbacks">
+            <summary>
+            Gets or sets all methods called before serialization of the object.
+            </summary>
+            <value>The methods called before serialization of the object.</value>
+        </member>
+        <member name="P:Newtonsoft.Json.Serialization.JsonContract.OnErrorCallbacks">
+            <summary>
+            Gets or sets all method called when an error is thrown during the serialization of the object.
+            </summary>
+            <value>The methods called when an error is thrown during the serialization of the object.</value>
+        </member>
+        <member name="P:Newtonsoft.Json.Serialization.JsonContract.OnDeserialized">
+            <summary>
+            Gets or sets the method called immediately after deserialization of the object.
+            </summary>
+            <value>The method called immediately after deserialization of the object.</value>
+        </member>
+        <member name="P:Newtonsoft.Json.Serialization.JsonContract.OnDeserializing">
+            <summary>
+            Gets or sets the method called during deserialization of the object.
+            </summary>
+            <value>The method called during deserialization of the object.</value>
+        </member>
+        <member name="P:Newtonsoft.Json.Serialization.JsonContract.OnSerialized">
+            <summary>
+            Gets or sets the method called after serialization of the object graph.
+            </summary>
+            <value>The method called after serialization of the object graph.</value>
+        </member>
+        <member name="P:Newtonsoft.Json.Serialization.JsonContract.OnSerializing">
+            <summary>
+            Gets or sets the method called before serialization of the object.
+            </summary>
+            <value>The method called before serialization of the object.</value>
+        </member>
+        <member name="P:Newtonsoft.Json.Serialization.JsonContract.OnError">
+            <summary>
+            Gets or sets the method called when an error is thrown during the serialization of the object.
+            </summary>
+            <value>The method called when an error is thrown during the serialization of the object.</value>
+        </member>
+        <member name="P:Newtonsoft.Json.Serialization.JsonContract.DefaultCreator">
+            <summary>
+            Gets or sets the default creator method used to create the object.
+            </summary>
+            <value>The default creator method used to create the object.</value>
+        </member>
+        <member name="P:Newtonsoft.Json.Serialization.JsonContract.DefaultCreatorNonPublic">
+            <summary>
+            Gets or sets a value indicating whether the default creator is non public.
+            </summary>
+            <value><c>true</c> if the default object creator is non-public; otherwise, <c>false</c>.</value>
+        </member>
+        <member name="M:Newtonsoft.Json.Serialization.JsonContainerContract.#ctor(System.Type)">
+            <summary>
+            Initializes a new instance of the <see cref="T:Newtonsoft.Json.Serialization.JsonContainerContract"/> class.
+            </summary>
+            <param name="underlyingType">The underlying type for the contract.</param>
+        </member>
+        <member name="P:Newtonsoft.Json.Serialization.JsonContainerContract.ItemConverter">
+            <summary>
+            Gets or sets the default collection items <see cref="T:Newtonsoft.Json.JsonConverter"/>.
+            </summary>
+            <value>The converter.</value>
+        </member>
+        <member name="P:Newtonsoft.Json.Serialization.JsonContainerContract.ItemIsReference">
+            <summary>
+            Gets or sets a value indicating whether the collection items preserve object references.
+            </summary>
+            <value><c>true</c> if collection items preserve object references; otherwise, <c>false</c>.</value>
+        </member>
+        <member name="P:Newtonsoft.Json.Serialization.JsonContainerContract.ItemReferenceLoopHandling">
+            <summary>
+            Gets or sets the collection item reference loop handling.
+            </summary>
+            <value>The reference loop handling.</value>
+        </member>
+        <member name="P:Newtonsoft.Json.Serialization.JsonContainerContract.ItemTypeNameHandling">
+            <summary>
+            Gets or sets the collection item type name handling.
+            </summary>
+            <value>The type name handling.</value>
+        </member>
+        <member name="T:Newtonsoft.Json.Serialization.MemoryTraceWriter">
+            <summary>
+            Represents a trace writer that writes to memory. When the trace message limit is
+            reached then old trace messages will be removed as new messages are added.
+            </summary>
+        </member>
+        <member name="M:Newtonsoft.Json.Serialization.MemoryTraceWriter.#ctor">
+            <summary>
+            Initializes a new instance of the <see cref="T:Newtonsoft.Json.Serialization.MemoryTraceWriter"/> class.
+            </summary>
+        </member>
+        <member name="M:Newtonsoft.Json.Serialization.MemoryTraceWriter.Trace(System.Diagnostics.TraceLevel,System.String,System.Exception)">
+            <summary>
+            Writes the specified trace level, message and optional exception.
+            </summary>
+            <param name="level">The <see cref="T:System.Diagnostics.TraceLevel"/> at which to write this trace.</param>
+            <param name="message">The trace message.</param>
+            <param name="ex">The trace exception. This parameter is optional.</param>
+        </member>
+        <member name="M:Newtonsoft.Json.Serialization.MemoryTraceWriter.GetTraceMessages">
+            <summary>
+            Returns an enumeration of the most recent trace messages.
+            </summary>
+            <returns>An enumeration of the most recent trace messages.</returns>
+        </member>
+        <member name="M:Newtonsoft.Json.Serialization.MemoryTraceWriter.ToString">
+            <summary>
+            Returns a <see cref="T:System.String"/> of the most recent trace messages.
+            </summary>
+            <returns>
+            A <see cref="T:System.String"/> of the most recent trace messages.
+            </returns>
+        </member>
+        <member name="P:Newtonsoft.Json.Serialization.MemoryTraceWriter.LevelFilter">
+            <summary>
+            Gets the <see cref="T:System.Diagnostics.TraceLevel"/> that will be used to filter the trace messages passed to the writer.
+            For example a filter level of <code>Info</code> will exclude <code>Verbose</code> messages and include <code>Info</code>,
+            <code>Warning</code> and <code>Error</code> messages.
+            </summary>
+            <value>
+            The <see cref="T:System.Diagnostics.TraceLevel"/> that will be used to filter the trace messages passed to the writer.
+            </value>
+        </member>
+        <member name="T:Newtonsoft.Json.StringEscapeHandling">
+            <summary>
+            Specifies how strings are escaped when writing JSON text.
+            </summary>
+        </member>
+        <member name="F:Newtonsoft.Json.StringEscapeHandling.Default">
+            <summary>
+            Only control characters (e.g. newline) are escaped.
+            </summary>
+        </member>
+        <member name="F:Newtonsoft.Json.StringEscapeHandling.EscapeNonAscii">
+            <summary>
+            All non-ASCII and control characters (e.g. newline) are escaped.
+            </summary>
+        </member>
+        <member name="F:Newtonsoft.Json.StringEscapeHandling.EscapeHtml">
+            <summary>
+            HTML (&lt;, &gt;, &amp;, &apos;, &quot;) and control characters (e.g. newline) are escaped.
+            </summary>
+        </member>
+        <member name="T:Newtonsoft.Json.Utilities.LinqBridge.Enumerable">
+            <summary>
+            Provides a set of static (Shared in Visual Basic) methods for 
+            querying objects that implement <see cref="T:System.Collections.Generic.IEnumerable`1"/>.
+            </summary>
+        </member>
+        <member name="M:Newtonsoft.Json.Utilities.LinqBridge.Enumerable.AsEnumerable``1(System.Collections.Generic.IEnumerable{``0})">
+            <summary>
+            Returns the input typed as <see cref="T:System.Collections.Generic.IEnumerable`1"/>.
+            </summary>
+        </member>
+        <member name="M:Newtonsoft.Json.Utilities.LinqBridge.Enumerable.Empty``1">
+            <summary>
+            Returns an empty <see cref="T:System.Collections.Generic.IEnumerable`1"/> that has the 
+            specified type argument.
+            </summary>
+        </member>
+        <member name="M:Newtonsoft.Json.Utilities.LinqBridge.Enumerable.Cast``1(System.Collections.IEnumerable)">
+            <summary>
+            Converts the elements of an <see cref="T:System.Collections.IEnumerable"/> to the 
+            specified type.
+            </summary>
+        </member>
+        <member name="M:Newtonsoft.Json.Utilities.LinqBridge.Enumerable.OfType``1(System.Collections.IEnumerable)">
+            <summary>
+            Filters the elements of an <see cref="T:System.Collections.IEnumerable"/> based on a specified type.
+            </summary>
+        </member>
+        <member name="M:Newtonsoft.Json.Utilities.LinqBridge.Enumerable.Range(System.Int32,System.Int32)">
+            <summary>
+            Generates a sequence of integral numbers within a specified range.
+            </summary>
+            <param name="start">The value of the first integer in the sequence.</param>
+            <param name="count">The number of sequential integers to generate.</param>
+        </member>
+        <member name="M:Newtonsoft.Json.Utilities.LinqBridge.Enumerable.Repeat``1(``0,System.Int32)">
+            <summary>
+            Generates a sequence that contains one repeated value.
+            </summary>
+        </member>
+        <member name="M:Newtonsoft.Json.Utilities.LinqBridge.Enumerable.Where``1(System.Collections.Generic.IEnumerable{``0},Newtonsoft.Json.Serialization.Func{``0,System.Boolean})">
+            <summary>
+            Filters a sequence of values based on a predicate.
+            </summary>
+        </member>
+        <member name="M:Newtonsoft.Json.Utilities.LinqBridge.Enumerable.Where``1(System.Collections.Generic.IEnumerable{``0},Newtonsoft.Json.Serialization.Func{``0,System.Int32,System.Boolean})">
+            <summary>
+            Filters a sequence of values based on a predicate. 
+            Each element's index is used in the logic of the predicate function.
+            </summary>
+        </member>
+        <member name="M:Newtonsoft.Json.Utilities.LinqBridge.Enumerable.Select``2(System.Collections.Generic.IEnumerable{``0},Newtonsoft.Json.Serialization.Func{``0,``1})">
+            <summary>
+            Projects each element of a sequence into a new form.
+            </summary>
+        </member>
+        <member name="M:Newtonsoft.Json.Utilities.LinqBridge.Enumerable.Select``2(System.Collections.Generic.IEnumerable{``0},Newtonsoft.Json.Serialization.Func{``0,System.Int32,``1})">
+            <summary>
+            Projects each element of a sequence into a new form by 
+            incorporating the element's index.
+            </summary>
+        </member>
+        <member name="M:Newtonsoft.Json.Utilities.LinqBridge.Enumerable.SelectMany``2(System.Collections.Generic.IEnumerable{``0},Newtonsoft.Json.Serialization.Func{``0,System.Collections.Generic.IEnumerable{``1}})">
+            <summary>
+            Projects each element of a sequence to an <see cref="T:System.Collections.Generic.IEnumerable`1"/> 
+            and flattens the resulting sequences into one sequence.
+            </summary>
+        </member>
+        <member name="M:Newtonsoft.Json.Utilities.LinqBridge.Enumerable.SelectMany``2(System.Collections.Generic.IEnumerable{``0},Newtonsoft.Json.Serialization.Func{``0,System.Int32,System.Collections.Generic.IEnumerable{``1}})">
+            <summary>
+            Projects each element of a sequence to an <see cref="T:System.Collections.Generic.IEnumerable`1"/>, 
+            and flattens the resulting sequences into one sequence. The 
+            index of each source element is used in the projected form of 
+            that element.
+            </summary>
+        </member>
+        <member name="M:Newtonsoft.Json.Utilities.LinqBridge.Enumerable.SelectMany``3(System.Collections.Generic.IEnumerable{``0},Newtonsoft.Json.Serialization.Func{``0,System.Collections.Generic.IEnumerable{``1}},Newtonsoft.Json.Serialization.Func{``0,``1,``2})">
+            <summary>
+            Projects each element of a sequence to an <see cref="T:System.Collections.Generic.IEnumerable`1"/>, 
+            flattens the resulting sequences into one sequence, and invokes 
+            a result selector function on each element therein.
+            </summary>
+        </member>
+        <member name="M:Newtonsoft.Json.Utilities.LinqBridge.Enumerable.SelectMany``3(System.Collections.Generic.IEnumerable{``0},Newtonsoft.Json.Serialization.Func{``0,System.Int32,System.Collections.Generic.IEnumerable{``1}},Newtonsoft.Json.Serialization.Func{``0,``1,``2})">
+            <summary>
+            Projects each element of a sequence to an <see cref="T:System.Collections.Generic.IEnumerable`1"/>, 
+            flattens the resulting sequences into one sequence, and invokes 
+            a result selector function on each element therein. The index of 
+            each source element is used in the intermediate projected form 
+            of that element.
+            </summary>
+        </member>
+        <member name="M:Newtonsoft.Json.Utilities.LinqBridge.Enumerable.TakeWhile``1(System.Collections.Generic.IEnumerable{``0},Newtonsoft.Json.Serialization.Func{``0,System.Boolean})">
+            <summary>
+            Returns elements from a sequence as long as a specified condition is true.
+            </summary>
+        </member>
+        <member name="M:Newtonsoft.Json.Utilities.LinqBridge.Enumerable.TakeWhile``1(System.Collections.Generic.IEnumerable{``0},Newtonsoft.Json.Serialization.Func{``0,System.Int32,System.Boolean})">
+            <summary>
+            Returns elements from a sequence as long as a specified condition is true.
+            The element's index is used in the logic of the predicate function.
+            </summary>
+        </member>
+        <member name="M:Newtonsoft.Json.Utilities.LinqBridge.Enumerable.FirstImpl``1(System.Collections.Generic.IEnumerable{``0},Newtonsoft.Json.Serialization.Func{``0})">
+            <summary>
+            Base implementation of First operator.
+            </summary>
+        </member>
+        <member name="M:Newtonsoft.Json.Utilities.LinqBridge.Enumerable.First``1(System.Collections.Generic.IEnumerable{``0})">
+            <summary>
+            Returns the first element of a sequence.
+            </summary>
+        </member>
+        <member name="M:Newtonsoft.Json.Utilities.LinqBridge.Enumerable.First``1(System.Collections.Generic.IEnumerable{``0},Newtonsoft.Json.Serialization.Func{``0,System.Boolean})">
+            <summary>
+            Returns the first element in a sequence that satisfies a specified condition.
+            </summary>
+        </member>
+        <member name="M:Newtonsoft.Json.Utilities.LinqBridge.Enumerable.FirstOrDefault``1(System.Collections.Generic.IEnumerable{``0})">
+            <summary>
+            Returns the first element of a sequence, or a default value if 
+            the sequence contains no elements.
+            </summary>
+        </member>
+        <member name="M:Newtonsoft.Json.Utilities.LinqBridge.Enumerable.FirstOrDefault``1(System.Collections.Generic.IEnumerable{``0},Newtonsoft.Json.Serialization.Func{``0,System.Boolean})">
+            <summary>
+            Returns the first element of the sequence that satisfies a 
+            condition or a default value if no such element is found.
+            </summary>
+        </member>
+        <member name="M:Newtonsoft.Json.Utilities.LinqBridge.Enumerable.LastImpl``1(System.Collections.Generic.IEnumerable{``0},Newtonsoft.Json.Serialization.Func{``0})">
+            <summary>
+            Base implementation of Last operator.
+            </summary>
+        </member>
+        <member name="M:Newtonsoft.Json.Utilities.LinqBridge.Enumerable.Last``1(System.Collections.Generic.IEnumerable{``0})">
+            <summary>
+            Returns the last element of a sequence.
+            </summary>
+        </member>
+        <member name="M:Newtonsoft.Json.Utilities.LinqBridge.Enumerable.Last``1(System.Collections.Generic.IEnumerable{``0},Newtonsoft.Json.Serialization.Func{``0,System.Boolean})">
+            <summary>
+            Returns the last element of a sequence that satisfies a 
+            specified condition.
+            </summary>
+        </member>
+        <member name="M:Newtonsoft.Json.Utilities.LinqBridge.Enumerable.LastOrDefault``1(System.Collections.Generic.IEnumerable{``0})">
+            <summary>
+            Returns the last element of a sequence, or a default value if 
+            the sequence contains no elements.
+            </summary>
+        </member>
+        <member name="M:Newtonsoft.Json.Utilities.LinqBridge.Enumerable.LastOrDefault``1(System.Collections.Generic.IEnumerable{``0},Newtonsoft.Json.Serialization.Func{``0,System.Boolean})">
+            <summary>
+            Returns the last element of a sequence that satisfies a 
+            condition or a default value if no such element is found.
+            </summary>
+        </member>
+        <member name="M:Newtonsoft.Json.Utilities.LinqBridge.Enumerable.SingleImpl``1(System.Collections.Generic.IEnumerable{``0},Newtonsoft.Json.Serialization.Func{``0})">
+            <summary>
+            Base implementation of Single operator.
+            </summary>
+        </member>
+        <member name="M:Newtonsoft.Json.Utilities.LinqBridge.Enumerable.Single``1(System.Collections.Generic.IEnumerable{``0})">
+            <summary>
+            Returns the only element of a sequence, and throws an exception 
+            if there is not exactly one element in the sequence.
+            </summary>
+        </member>
+        <member name="M:Newtonsoft.Json.Utilities.LinqBridge.Enumerable.Single``1(System.Collections.Generic.IEnumerable{``0},Newtonsoft.Json.Serialization.Func{``0,System.Boolean})">
+            <summary>
+            Returns the only element of a sequence that satisfies a 
+            specified condition, and throws an exception if more than one 
+            such element exists.
+            </summary>
+        </member>
+        <member name="M:Newtonsoft.Json.Utilities.LinqBridge.Enumerable.SingleOrDefault``1(System.Collections.Generic.IEnumerable{``0})">
+            <summary>
+            Returns the only element of a sequence, or a default value if 
+            the sequence is empty; this method throws an exception if there 
+            is more than one element in the sequence.
+            </summary>
+        </member>
+        <member name="M:Newtonsoft.Json.Utilities.LinqBridge.Enumerable.SingleOrDefault``1(System.Collections.Generic.IEnumerable{``0},Newtonsoft.Json.Serialization.Func{``0,System.Boolean})">
+            <summary>
+            Returns the only element of a sequence that satisfies a 
+            specified condition or a default value if no such element 
+            exists; this method throws an exception if more than one element 
+            satisfies the condition.
+            </summary>
+        </member>
+        <member name="M:Newtonsoft.Json.Utilities.LinqBridge.Enumerable.ElementAt``1(System.Collections.Generic.IEnumerable{``0},System.Int32)">
+            <summary>
+            Returns the element at a specified index in a sequence.
+            </summary>
+        </member>
+        <member name="M:Newtonsoft.Json.Utilities.LinqBridge.Enumerable.ElementAtOrDefault``1(System.Collections.Generic.IEnumerable{``0},System.Int32)">
+            <summary>
+            Returns the element at a specified index in a sequence or a 
+            default value if the index is out of range.
+            </summary>
+        </member>
+        <member name="M:Newtonsoft.Json.Utilities.LinqBridge.Enumerable.Reverse``1(System.Collections.Generic.IEnumerable{``0})">
+            <summary>
+            Inverts the order of the elements in a sequence.
+            </summary>
+        </member>
+        <member name="M:Newtonsoft.Json.Utilities.LinqBridge.Enumerable.Take``1(System.Collections.Generic.IEnumerable{``0},System.Int32)">
+            <summary>
+            Returns a specified number of contiguous elements from the start 
+            of a sequence.
+            </summary>
+        </member>
+        <member name="M:Newtonsoft.Json.Utilities.LinqBridge.Enumerable.Skip``1(System.Collections.Generic.IEnumerable{``0},System.Int32)">
+            <summary>
+            Bypasses a specified number of elements in a sequence and then 
+            returns the remaining elements.
+            </summary>
+        </member>
+        <member name="M:Newtonsoft.Json.Utilities.LinqBridge.Enumerable.SkipWhile``1(System.Collections.Generic.IEnumerable{``0},Newtonsoft.Json.Serialization.Func{``0,System.Boolean})">
+            <summary>
+            Bypasses elements in a sequence as long as a specified condition 
+            is true and then returns the remaining elements.
+            </summary>
+        </member>
+        <member name="M:Newtonsoft.Json.Utilities.LinqBridge.Enumerable.SkipWhile``1(System.Collections.Generic.IEnumerable{``0},Newtonsoft.Json.Serialization.Func{``0,System.Int32,System.Boolean})">
+            <summary>
+            Bypasses elements in a sequence as long as a specified condition 
+            is true and then returns the remaining elements. The element's 
+            index is used in the logic of the predicate function.
+            </summary>
+        </member>
+        <member name="M:Newtonsoft.Json.Utilities.LinqBridge.Enumerable.Count``1(System.Collections.Generic.IEnumerable{``0})">
+            <summary>
+            Returns the number of elements in a sequence.
+            </summary>
+        </member>
+        <member name="M:Newtonsoft.Json.Utilities.LinqBridge.Enumerable.Count``1(System.Collections.Generic.IEnumerable{``0},Newtonsoft.Json.Serialization.Func{``0,System.Boolean})">
+            <summary>
+            Returns a number that represents how many elements in the 
+            specified sequence satisfy a condition.
+            </summary>
+        </member>
+        <member name="M:Newtonsoft.Json.Utilities.LinqBridge.Enumerable.LongCount``1(System.Collections.Generic.IEnumerable{``0})">
+            <summary>
+            Returns an <see cref="T:System.Int64"/> that represents the total number 
+            of elements in a sequence.
+            </summary>
+        </member>
+        <member name="M:Newtonsoft.Json.Utilities.LinqBridge.Enumerable.LongCount``1(System.Collections.Generic.IEnumerable{``0},Newtonsoft.Json.Serialization.Func{``0,System.Boolean})">
+            <summary>
+            Returns an <see cref="T:System.Int64"/> that represents how many elements 
+            in a sequence satisfy a condition.
+            </summary>
+        </member>
+        <member name="M:Newtonsoft.Json.Utilities.LinqBridge.Enumerable.Concat``1(System.Collections.Generic.IEnumerable{``0},System.Collections.Generic.IEnumerable{``0})">
+            <summary>
+            Concatenates two sequences.
+            </summary>
+        </member>
+        <member name="M:Newtonsoft.Json.Utilities.LinqBridge.Enumerable.ToList``1(System.Collections.Generic.IEnumerable{``0})">
+            <summary>
+            Creates a <see cref="T:System.Collections.Generic.List`1"/> from an <see cref="T:System.Collections.Generic.IEnumerable`1"/>.
+            </summary>
+        </member>
+        <member name="M:Newtonsoft.Json.Utilities.LinqBridge.Enumerable.ToArray``1(System.Collections.Generic.IEnumerable{``0})">
+            <summary>
+            Creates an array from an <see cref="T:System.Collections.Generic.IEnumerable`1"/>.
+            </summary>
+        </member>
+        <member name="M:Newtonsoft.Json.Utilities.LinqBridge.Enumerable.Distinct``1(System.Collections.Generic.IEnumerable{``0})">
+            <summary>
+            Returns distinct elements from a sequence by using the default 
+            equality comparer to compare values.
+            </summary>
+        </member>
+        <member name="M:Newtonsoft.Json.Utilities.LinqBridge.Enumerable.Distinct``1(System.Collections.Generic.IEnumerable{``0},System.Collections.Generic.IEqualityComparer{``0})">
+            <summary>
+            Returns distinct elements from a sequence by using a specified 
+            <see cref="T:System.Collections.Generic.IEqualityComparer`1"/> to compare values.
+            </summary>
+        </member>
+        <member name="M:Newtonsoft.Json.Utilities.LinqBridge.Enumerable.ToLookup``2(System.Collections.Generic.IEnumerable{``0},Newtonsoft.Json.Serialization.Func{``0,``1})">
+            <summary>
+            Creates a <see cref="T:Newtonsoft.Json.Utilities.LinqBridge.Lookup`2"/> from an 
+            <see cref="T:System.Collections.Generic.IEnumerable`1"/> according to a specified key 
+            selector function.
+            </summary>
+        </member>
+        <member name="M:Newtonsoft.Json.Utilities.LinqBridge.Enumerable.ToLookup``2(System.Collections.Generic.IEnumerable{``0},Newtonsoft.Json.Serialization.Func{``0,``1},System.Collections.Generic.IEqualityComparer{``1})">
+            <summary>
+            Creates a <see cref="T:Newtonsoft.Json.Utilities.LinqBridge.Lookup`2"/> from an 
+            <see cref="T:System.Collections.Generic.IEnumerable`1"/> according to a specified key 
+            selector function and a key comparer.
+            </summary>
+        </member>
+        <member name="M:Newtonsoft.Json.Utilities.LinqBridge.Enumerable.ToLookup``3(System.Collections.Generic.IEnumerable{``0},Newtonsoft.Json.Serialization.Func{``0,``1},Newtonsoft.Json.Serialization.Func{``0,``2})">
+            <summary>
+            Creates a <see cref="T:Newtonsoft.Json.Utilities.LinqBridge.Lookup`2"/> from an 
+            <see cref="T:System.Collections.Generic.IEnumerable`1"/> according to specified key 
+            and element selector functions.
+            </summary>
+        </member>
+        <member name="M:Newtonsoft.Json.Utilities.LinqBridge.Enumerable.ToLookup``3(System.Collections.Generic.IEnumerable{``0},Newtonsoft.Json.Serialization.Func{``0,``1},Newtonsoft.Json.Serialization.Func{``0,``2},System.Collections.Generic.IEqualityComparer{``1})">
+            <summary>
+            Creates a <see cref="T:Newtonsoft.Json.Utilities.LinqBridge.Lookup`2"/> from an 
+            <see cref="T:System.Collections.Generic.IEnumerable`1"/> according to a specified key 
+            selector function, a comparer and an element selector function.
+            </summary>
+        </member>
+        <member name="M:Newtonsoft.Json.Utilities.LinqBridge.Enumerable.GroupBy``2(System.Collections.Generic.IEnumerable{``0},Newtonsoft.Json.Serialization.Func{``0,``1})">
+            <summary>
+            Groups the elements of a sequence according to a specified key 
+            selector function.
+            </summary>
+        </member>
+        <member name="M:Newtonsoft.Json.Utilities.LinqBridge.Enumerable.GroupBy``2(System.Collections.Generic.IEnumerable{``0},Newtonsoft.Json.Serialization.Func{``0,``1},System.Collections.Generic.IEqualityComparer{``1})">
+            <summary>
+            Groups the elements of a sequence according to a specified key 
+            selector function and compares the keys by using a specified 
+            comparer.
+            </summary>
+        </member>
+        <member name="M:Newtonsoft.Json.Utilities.LinqBridge.Enumerable.GroupBy``3(System.Collections.Generic.IEnumerable{``0},Newtonsoft.Json.Serialization.Func{``0,``1},Newtonsoft.Json.Serialization.Func{``0,``2})">
+            <summary>
+            Groups the elements of a sequence according to a specified key 
+            selector function and projects the elements for each group by 
+            using a specified function.
+            </summary>
+        </member>
+        <member name="M:Newtonsoft.Json.Utilities.LinqBridge.Enumerable.GroupBy``3(System.Collections.Generic.IEnumerable{``0},Newtonsoft.Json.Serialization.Func{``0,``1},Newtonsoft.Json.Serialization.Func{``0,``2},System.Collections.Generic.IEqualityComparer{``1})">
+            <summary>
+            Groups the elements of a sequence according to a specified key 
+            selector function and creates a result value from each group and 
+            its key.
+            </summary>
+        </member>
+        <member name="M:Newtonsoft.Json.Utilities.LinqBridge.Enumerable.GroupBy``3(System.Collections.Generic.IEnumerable{``0},Newtonsoft.Json.Serialization.Func{``0,``1},Newtonsoft.Json.Serialization.Func{``1,System.Collections.Generic.IEnumerable{``0},``2})">
+            <summary>
+            Groups the elements of a sequence according to a key selector 
+            function. The keys are compared by using a comparer and each 
+            group's elements are projected by using a specified function.
+            </summary>
+        </member>
+        <member name="M:Newtonsoft.Json.Utilities.LinqBridge.Enumerable.GroupBy``3(System.Collections.Generic.IEnumerable{``0},Newtonsoft.Json.Serialization.Func{``0,``1},Newtonsoft.Json.Serialization.Func{``1,System.Collections.Generic.IEnumerable{``0},``2},System.Collections.Generic.IEqualityComparer{``1})">
+            <summary>
+            Groups the elements of a sequence according to a specified key 
+            selector function and creates a result value from each group and 
+            its key. The elements of each group are projected by using a 
+            specified function.
+            </summary>
+        </member>
+        <member name="M:Newtonsoft.Json.Utilities.LinqBridge.Enumerable.GroupBy``4(System.Collections.Generic.IEnumerable{``0},Newtonsoft.Json.Serialization.Func{``0,``1},Newtonsoft.Json.Serialization.Func{``0,``2},Newtonsoft.Json.Serialization.Func{``1,System.Collections.Generic.IEnumerable{``2},``3})">
+            <summary>
+            Groups the elements of a sequence according to a specified key 
+            selector function and creates a result value from each group and 
+            its key. The keys are compared by using a specified comparer.
+            </summary>
+        </member>
+        <member name="M:Newtonsoft.Json.Utilities.LinqBridge.Enumerable.GroupBy``4(System.Collections.Generic.IEnumerable{``0},Newtonsoft.Json.Serialization.Func{``0,``1},Newtonsoft.Json.Serialization.Func{``0,``2},Newtonsoft.Json.Serialization.Func{``1,System.Collections.Generic.IEnumerable{``2},``3},System.Collections.Generic.IEqualityComparer{``1})">
+            <summary>
+            Groups the elements of a sequence according to a specified key 
+            selector function and creates a result value from each group and 
+            its key. Key values are compared by using a specified comparer, 
+            and the elements of each group are projected by using a 
+            specified function.
+            </summary>
+        </member>
+        <member name="M:Newtonsoft.Json.Utilities.LinqBridge.Enumerable.Aggregate``1(System.Collections.Generic.IEnumerable{``0},Newtonsoft.Json.Serialization.Func{``0,``0,``0})">
+            <summary>
+            Applies an accumulator function over a sequence.
+            </summary>
+        </member>
+        <member name="M:Newtonsoft.Json.Utilities.LinqBridge.Enumerable.Aggregate``2(System.Collections.Generic.IEnumerable{``0},``1,Newtonsoft.Json.Serialization.Func{``1,``0,``1})">
+            <summary>
+            Applies an accumulator function over a sequence. The specified 
+            seed value is used as the initial accumulator value.
+            </summary>
+        </member>
+        <member name="M:Newtonsoft.Json.Utilities.LinqBridge.Enumerable.Aggregate``3(System.Collections.Generic.IEnumerable{``0},``1,Newtonsoft.Json.Serialization.Func{``1,``0,``1},Newtonsoft.Json.Serialization.Func{``1,``2})">
+            <summary>
+            Applies an accumulator function over a sequence. The specified 
+            seed value is used as the initial accumulator value, and the 
+            specified function is used to select the result value.
+            </summary>
+        </member>
+        <member name="M:Newtonsoft.Json.Utilities.LinqBridge.Enumerable.Union``1(System.Collections.Generic.IEnumerable{``0},System.Collections.Generic.IEnumerable{``0})">
+            <summary>
+            Produces the set union of two sequences by using the default 
+            equality comparer.
+            </summary>
+        </member>
+        <member name="M:Newtonsoft.Json.Utilities.LinqBridge.Enumerable.Union``1(System.Collections.Generic.IEnumerable{``0},System.Collections.Generic.IEnumerable{``0},System.Collections.Generic.IEqualityComparer{``0})">
+            <summary>
+            Produces the set union of two sequences by using a specified 
+            <see cref="T:System.Collections.Generic.IEqualityComparer`1"/>.
+            </summary>
+        </member>
+        <member name="M:Newtonsoft.Json.Utilities.LinqBridge.Enumerable.DefaultIfEmpty``1(System.Collections.Generic.IEnumerable{``0})">
+            <summary>
+            Returns the elements of the specified sequence or the type 
+            parameter's default value in a singleton collection if the 
+            sequence is empty.
+            </summary>
+        </member>
+        <member name="M:Newtonsoft.Json.Utilities.LinqBridge.Enumerable.DefaultIfEmpty``1(System.Collections.Generic.IEnumerable{``0},``0)">
+            <summary>
+            Returns the elements of the specified sequence or the specified 
+            value in a singleton collection if the sequence is empty.
+            </summary>
+        </member>
+        <member name="M:Newtonsoft.Json.Utilities.LinqBridge.Enumerable.All``1(System.Collections.Generic.IEnumerable{``0},Newtonsoft.Json.Serialization.Func{``0,System.Boolean})">
+            <summary>
+            Determines whether all elements of a sequence satisfy a condition.
+            </summary>
+        </member>
+        <member name="M:Newtonsoft.Json.Utilities.LinqBridge.Enumerable.Any``1(System.Collections.Generic.IEnumerable{``0})">
+            <summary>
+            Determines whether a sequence contains any elements.
+            </summary>
+        </member>
+        <member name="M:Newtonsoft.Json.Utilities.LinqBridge.Enumerable.Any``1(System.Collections.Generic.IEnumerable{``0},Newtonsoft.Json.Serialization.Func{``0,System.Boolean})">
+            <summary>
+            Determines whether any element of a sequence satisfies a 
+            condition.
+            </summary>
+        </member>
+        <member name="M:Newtonsoft.Json.Utilities.LinqBridge.Enumerable.Contains``1(System.Collections.Generic.IEnumerable{``0},``0)">
+            <summary>
+            Determines whether a sequence contains a specified element by 
+            using the default equality comparer.
+            </summary>
+        </member>
+        <member name="M:Newtonsoft.Json.Utilities.LinqBridge.Enumerable.Contains``1(System.Collections.Generic.IEnumerable{``0},``0,System.Collections.Generic.IEqualityComparer{``0})">
+            <summary>
+            Determines whether a sequence contains a specified element by 
+            using a specified <see cref="T:System.Collections.Generic.IEqualityComparer`1"/>.
+            </summary>
+        </member>
+        <member name="M:Newtonsoft.Json.Utilities.LinqBridge.Enumerable.SequenceEqual``1(System.Collections.Generic.IEnumerable{``0},System.Collections.Generic.IEnumerable{``0})">
+            <summary>
+            Determines whether two sequences are equal by comparing the 
+            elements by using the default equality comparer for their type.
+            </summary>
+        </member>
+        <member name="M:Newtonsoft.Json.Utilities.LinqBridge.Enumerable.SequenceEqual``1(System.Collections.Generic.IEnumerable{``0},System.Collections.Generic.IEnumerable{``0},System.Collections.Generic.IEqualityComparer{``0})">
+            <summary>
+            Determines whether two sequences are equal by comparing their 
+            elements by using a specified <see cref="T:System.Collections.Generic.IEqualityComparer`1"/>.
+            </summary>
+        </member>
+        <member name="M:Newtonsoft.Json.Utilities.LinqBridge.Enumerable.MinMaxImpl``1(System.Collections.Generic.IEnumerable{``0},Newtonsoft.Json.Serialization.Func{``0,``0,System.Boolean})">
+            <summary>
+            Base implementation for Min/Max operator.
+            </summary>
+        </member>
+        <member name="M:Newtonsoft.Json.Utilities.LinqBridge.Enumerable.MinMaxImpl``1(System.Collections.Generic.IEnumerable{System.Nullable{``0}},System.Nullable{``0},Newtonsoft.Json.Serialization.Func{System.Nullable{``0},System.Nullable{``0},System.Boolean})">
+            <summary>
+            Base implementation for Min/Max operator for nullable types.
+            </summary>
+        </member>
+        <member name="M:Newtonsoft.Json.Utilities.LinqBridge.Enumerable.Min``1(System.Collections.Generic.IEnumerable{``0})">
+            <summary>
+            Returns the minimum value in a generic sequence.
+            </summary>
+        </member>
+        <member name="M:Newtonsoft.Json.Utilities.LinqBridge.Enumerable.Min``2(System.Collections.Generic.IEnumerable{``0},Newtonsoft.Json.Serialization.Func{``0,``1})">
+            <summary>
+            Invokes a transform function on each element of a generic 
+            sequence and returns the minimum resulting value.
+            </summary>
+        </member>
+        <member name="M:Newtonsoft.Json.Utilities.LinqBridge.Enumerable.Max``1(System.Collections.Generic.IEnumerable{``0})">
+            <summary>
+            Returns the maximum value in a generic sequence.
+            </summary>
+        </member>
+        <member name="M:Newtonsoft.Json.Utilities.LinqBridge.Enumerable.Max``2(System.Collections.Generic.IEnumerable{``0},Newtonsoft.Json.Serialization.Func{``0,``1})">
+            <summary>
+            Invokes a transform function on each element of a generic 
+            sequence and returns the maximum resulting value.
+            </summary>
+        </member>
+        <member name="M:Newtonsoft.Json.Utilities.LinqBridge.Enumerable.Renumerable``1(System.Collections.Generic.IEnumerator{``0})">
+            <summary>
+            Makes an enumerator seen as enumerable once more.
+            </summary>
+            <remarks>
+            The supplied enumerator must have been started. The first element
+            returned is the element the enumerator was on when passed in.
+            DO NOT use this method if the caller must be a generator. It is
+            mostly safe among aggregate operations.
+            </remarks>
+        </member>
+        <member name="M:Newtonsoft.Json.Utilities.LinqBridge.Enumerable.OrderBy``2(System.Collections.Generic.IEnumerable{``0},Newtonsoft.Json.Serialization.Func{``0,``1})">
+            <summary>
+            Sorts the elements of a sequence in ascending order according to a key.
+            </summary>
+        </member>
+        <member name="M:Newtonsoft.Json.Utilities.LinqBridge.Enumerable.OrderBy``2(System.Collections.Generic.IEnumerable{``0},Newtonsoft.Json.Serialization.Func{``0,``1},System.Collections.Generic.IComparer{``1})">
+            <summary>
+            Sorts the elements of a sequence in ascending order by using a 
+            specified comparer.
+            </summary>
+        </member>
+        <member name="M:Newtonsoft.Json.Utilities.LinqBridge.Enumerable.OrderByDescending``2(System.Collections.Generic.IEnumerable{``0},Newtonsoft.Json.Serialization.Func{``0,``1})">
+            <summary>
+            Sorts the elements of a sequence in descending order according to a key.
+            </summary>
+        </member>
+        <member name="M:Newtonsoft.Json.Utilities.LinqBridge.Enumerable.OrderByDescending``2(System.Collections.Generic.IEnumerable{``0},Newtonsoft.Json.Serialization.Func{``0,``1},System.Collections.Generic.IComparer{``1})">
+            <summary>
+             Sorts the elements of a sequence in descending order by using a 
+            specified comparer. 
+            </summary>
+        </member>
+        <member name="M:Newtonsoft.Json.Utilities.LinqBridge.Enumerable.ThenBy``2(Newtonsoft.Json.Utilities.LinqBridge.IOrderedEnumerable{``0},Newtonsoft.Json.Serialization.Func{``0,``1})">
+            <summary>
+            Performs a subsequent ordering of the elements in a sequence in 
+            ascending order according to a key.
+            </summary>
+        </member>
+        <member name="M:Newtonsoft.Json.Utilities.LinqBridge.Enumerable.ThenBy``2(Newtonsoft.Json.Utilities.LinqBridge.IOrderedEnumerable{``0},Newtonsoft.Json.Serialization.Func{``0,``1},System.Collections.Generic.IComparer{``1})">
+            <summary>
+            Performs a subsequent ordering of the elements in a sequence in 
+            ascending order by using a specified comparer.
+            </summary>
+        </member>
+        <member name="M:Newtonsoft.Json.Utilities.LinqBridge.Enumerable.ThenByDescending``2(Newtonsoft.Json.Utilities.LinqBridge.IOrderedEnumerable{``0},Newtonsoft.Json.Serialization.Func{``0,``1})">
+            <summary>
+            Performs a subsequent ordering of the elements in a sequence in 
+            descending order, according to a key.
+            </summary>
+        </member>
+        <member name="M:Newtonsoft.Json.Utilities.LinqBridge.Enumerable.ThenByDescending``2(Newtonsoft.Json.Utilities.LinqBridge.IOrderedEnumerable{``0},Newtonsoft.Json.Serialization.Func{``0,``1},System.Collections.Generic.IComparer{``1})">
+            <summary>
+            Performs a subsequent ordering of the elements in a sequence in 
+            descending order by using a specified comparer.
+            </summary>
+        </member>
+        <member name="M:Newtonsoft.Json.Utilities.LinqBridge.Enumerable.IntersectExceptImpl``1(System.Collections.Generic.IEnumerable{``0},System.Collections.Generic.IEnumerable{``0},System.Collections.Generic.IEqualityComparer{``0},System.Boolean)">
+            <summary>
+            Base implementation for Intersect and Except operators.
+            </summary>
+        </member>
+        <member name="M:Newtonsoft.Json.Utilities.LinqBridge.Enumerable.Intersect``1(System.Collections.Generic.IEnumerable{``0},System.Collections.Generic.IEnumerable{``0})">
+            <summary>
+            Produces the set intersection of two sequences by using the 
+            default equality comparer to compare values.
+            </summary>
+        </member>
+        <member name="M:Newtonsoft.Json.Utilities.LinqBridge.Enumerable.Intersect``1(System.Collections.Generic.IEnumerable{``0},System.Collections.Generic.IEnumerable{``0},System.Collections.Generic.IEqualityComparer{``0})">
+            <summary>
+            Produces the set intersection of two sequences by using the 
+            specified <see cref="T:System.Collections.Generic.IEqualityComparer`1"/> to compare values.
+            </summary>
+        </member>
+        <member name="M:Newtonsoft.Json.Utilities.LinqBridge.Enumerable.Except``1(System.Collections.Generic.IEnumerable{``0},System.Collections.Generic.IEnumerable{``0})">
+            <summary>
+            Produces the set difference of two sequences by using the 
+            default equality comparer to compare values.
+            </summary>
+        </member>
+        <member name="M:Newtonsoft.Json.Utilities.LinqBridge.Enumerable.Except``1(System.Collections.Generic.IEnumerable{``0},System.Collections.Generic.IEnumerable{``0},System.Collections.Generic.IEqualityComparer{``0})">
+            <summary>
+            Produces the set difference of two sequences by using the 
+            specified <see cref="T:System.Collections.Generic.IEqualityComparer`1"/> to compare values.
+            </summary>
+        </member>
+        <member name="M:Newtonsoft.Json.Utilities.LinqBridge.Enumerable.ToDictionary``2(System.Collections.Generic.IEnumerable{``0},Newtonsoft.Json.Serialization.Func{``0,``1})">
+            <summary>
+            Creates a <see cref="T:System.Collections.Generic.Dictionary`2"/> from an 
+            <see cref="T:System.Collections.Generic.IEnumerable`1"/> according to a specified key 
+            selector function.
+            </summary>
+        </member>
+        <member name="M:Newtonsoft.Json.Utilities.LinqBridge.Enumerable.ToDictionary``2(System.Collections.Generic.IEnumerable{``0},Newtonsoft.Json.Serialization.Func{``0,``1},System.Collections.Generic.IEqualityComparer{``1})">
+            <summary>
+            Creates a <see cref="T:System.Collections.Generic.Dictionary`2"/> from an 
+            <see cref="T:System.Collections.Generic.IEnumerable`1"/> according to a specified key 
+            selector function and key comparer.
+            </summary>
+        </member>
+        <member name="M:Newtonsoft.Json.Utilities.LinqBridge.Enumerable.ToDictionary``3(System.Collections.Generic.IEnumerable{``0},Newtonsoft.Json.Serialization.Func{``0,``1},Newtonsoft.Json.Serialization.Func{``0,``2})">
+            <summary>
+            Creates a <see cref="T:System.Collections.Generic.Dictionary`2"/> from an 
+            <see cref="T:System.Collections.Generic.IEnumerable`1"/> according to specified key 
+            selector and element selector functions.
+            </summary>
+        </member>
+        <member name="M:Newtonsoft.Json.Utilities.LinqBridge.Enumerable.ToDictionary``3(System.Collections.Generic.IEnumerable{``0},Newtonsoft.Json.Serialization.Func{``0,``1},Newtonsoft.Json.Serialization.Func{``0,``2},System.Collections.Generic.IEqualityComparer{``1})">
+            <summary>
+            Creates a <see cref="T:System.Collections.Generic.Dictionary`2"/> from an 
+            <see cref="T:System.Collections.Generic.IEnumerable`1"/> according to a specified key 
+            selector function, a comparer, and an element selector function.
+            </summary>
+        </member>
+        <member name="M:Newtonsoft.Json.Utilities.LinqBridge.Enumerable.Join``4(System.Collections.Generic.IEnumerable{``0},System.Collections.Generic.IEnumerable{``1},Newtonsoft.Json.Serialization.Func{``0,``2},Newtonsoft.Json.Serialization.Func{``1,``2},Newtonsoft.Json.Serialization.Func{``0,``1,``3})">
+            <summary>
+            Correlates the elements of two sequences based on matching keys. 
+            The default equality comparer is used to compare keys.
+            </summary>
+        </member>
+        <member name="M:Newtonsoft.Json.Utilities.LinqBridge.Enumerable.Join``4(System.Collections.Generic.IEnumerable{``0},System.Collections.Generic.IEnumerable{``1},Newtonsoft.Json.Serialization.Func{``0,``2},Newtonsoft.Json.Serialization.Func{``1,``2},Newtonsoft.Json.Serialization.Func{``0,``1,``3},System.Collections.Generic.IEqualityComparer{``2})">
+            <summary>
+            Correlates the elements of two sequences based on matching keys. 
+            The default equality comparer is used to compare keys. A 
+            specified <see cref="T:System.Collections.Generic.IEqualityComparer`1"/> is used to compare keys.
+            </summary>
+        </member>
+        <member name="M:Newtonsoft.Json.Utilities.LinqBridge.Enumerable.GroupJoin``4(System.Collections.Generic.IEnumerable{``0},System.Collections.Generic.IEnumerable{``1},Newtonsoft.Json.Serialization.Func{``0,``2},Newtonsoft.Json.Serialization.Func{``1,``2},Newtonsoft.Json.Serialization.Func{``0,System.Collections.Generic.IEnumerable{``1},``3})">
+            <summary>
+            Correlates the elements of two sequences based on equality of 
+            keys and groups the results. The default equality comparer is 
+            used to compare keys.
+            </summary>
+        </member>
+        <member name="M:Newtonsoft.Json.Utilities.LinqBridge.Enumerable.GroupJoin``4(System.Collections.Generic.IEnumerable{``0},System.Collections.Generic.IEnumerable{``1},Newtonsoft.Json.Serialization.Func{``0,``2},Newtonsoft.Json.Serialization.Func{``1,``2},Newtonsoft.Json.Serialization.Func{``0,System.Collections.Generic.IEnumerable{``1},``3},System.Collections.Generic.IEqualityComparer{``2})">
+            <summary>
+            Correlates the elements of two sequences based on equality of 
+            keys and groups the results. The default equality comparer is 
+            used to compare keys. A specified <see cref="T:System.Collections.Generic.IEqualityComparer`1"/> 
+            is used to compare keys.
+            </summary>
+        </member>
+        <member name="M:Newtonsoft.Json.Utilities.LinqBridge.Enumerable.Sum(System.Collections.Generic.IEnumerable{System.Int32})">
+            <summary>
+            Computes the sum of a sequence of nullable <see cref="T:System.Int32"/> values.
+            </summary>
+        </member>
+        <member name="M:Newtonsoft.Json.Utilities.LinqBridge.Enumerable.Sum``1(System.Collections.Generic.IEnumerable{``0},Newtonsoft.Json.Serialization.Func{``0,System.Int32})">
+            <summary>
+            Computes the sum of a sequence of nullable <see cref="T:System.Int32"/> 
+            values that are obtained by invoking a transform function on 
+            each element of the input sequence.
+            </summary>
+        </member>
+        <member name="M:Newtonsoft.Json.Utilities.LinqBridge.Enumerable.Average(System.Collections.Generic.IEnumerable{System.Int32})">
+            <summary>
+            Computes the average of a sequence of nullable <see cref="T:System.Int32"/> values.
+            </summary>
+        </member>
+        <member name="M:Newtonsoft.Json.Utilities.LinqBridge.Enumerable.Average``1(System.Collections.Generic.IEnumerable{``0},Newtonsoft.Json.Serialization.Func{``0,System.Int32})">
+            <summary>
+            Computes the average of a sequence of nullable <see cref="T:System.Int32"/> values 
+            that are obtained by invoking a transform function on each 
+            element of the input sequence.
+            </summary>
+        </member>
+        <member name="M:Newtonsoft.Json.Utilities.LinqBridge.Enumerable.Sum(System.Collections.Generic.IEnumerable{System.Nullable{System.Int32}})">
+            <summary>
+            Computes the sum of a sequence of <see cref="T:System.Int32"/> values.
+            </summary>
+        </member>
+        <member name="M:Newtonsoft.Json.Utilities.LinqBridge.Enumerable.Sum``1(System.Collections.Generic.IEnumerable{``0},Newtonsoft.Json.Serialization.Func{``0,System.Nullable{System.Int32}})">
+            <summary>
+            Computes the sum of a sequence of <see cref="T:System.Int32"/> 
+            values that are obtained by invoking a transform function on 
+            each element of the input sequence.
+            </summary>
+        </member>
+        <member name="M:Newtonsoft.Json.Utilities.LinqBridge.Enumerable.Average(System.Collections.Generic.IEnumerable{System.Nullable{System.Int32}})">
+            <summary>
+            Computes the average of a sequence of <see cref="T:System.Int32"/> values.
+            </summary>
+        </member>
+        <member name="M:Newtonsoft.Json.Utilities.LinqBridge.Enumerable.Average``1(System.Collections.Generic.IEnumerable{``0},Newtonsoft.Json.Serialization.Func{``0,System.Nullable{System.Int32}})">
+            <summary>
+            Computes the average of a sequence of <see cref="T:System.Int32"/> values 
+            that are obtained by invoking a transform function on each 
+            element of the input sequence.
+            </summary>
+        </member>
+        <member name="M:Newtonsoft.Json.Utilities.LinqBridge.Enumerable.Min(System.Collections.Generic.IEnumerable{System.Nullable{System.Int32}})">
+            <summary>
+            Returns the minimum value in a sequence of nullable 
+            <see cref="T:System.Int32"/> values.
+            </summary>
+        </member>
+        <member name="M:Newtonsoft.Json.Utilities.LinqBridge.Enumerable.Min``1(System.Collections.Generic.IEnumerable{``0},Newtonsoft.Json.Serialization.Func{``0,System.Nullable{System.Int32}})">
+            <summary>
+            Invokes a transform function on each element of a sequence and 
+            returns the minimum nullable <see cref="T:System.Int32"/> value.
+            </summary>
+        </member>
+        <member name="M:Newtonsoft.Json.Utilities.LinqBridge.Enumerable.Max(System.Collections.Generic.IEnumerable{System.Nullable{System.Int32}})">
+            <summary>
+            Returns the maximum value in a sequence of nullable 
+            <see cref="T:System.Int32"/> values.
+            </summary>
+        </member>
+        <member name="M:Newtonsoft.Json.Utilities.LinqBridge.Enumerable.Max``1(System.Collections.Generic.IEnumerable{``0},Newtonsoft.Json.Serialization.Func{``0,System.Nullable{System.Int32}})">
+            <summary>
+            Invokes a transform function on each element of a sequence and 
+            returns the maximum nullable <see cref="T:System.Int32"/> value.
+            </summary>
+        </member>
+        <member name="M:Newtonsoft.Json.Utilities.LinqBridge.Enumerable.Sum(System.Collections.Generic.IEnumerable{System.Int64})">
+            <summary>
+            Computes the sum of a sequence of nullable <see cref="T:System.Int64"/> values.
+            </summary>
+        </member>
+        <member name="M:Newtonsoft.Json.Utilities.LinqBridge.Enumerable.Sum``1(System.Collections.Generic.IEnumerable{``0},Newtonsoft.Json.Serialization.Func{``0,System.Int64})">
+            <summary>
+            Computes the sum of a sequence of nullable <see cref="T:System.Int64"/> 
+            values that are obtained by invoking a transform function on 
+            each element of the input sequence.
+            </summary>
+        </member>
+        <member name="M:Newtonsoft.Json.Utilities.LinqBridge.Enumerable.Average(System.Collections.Generic.IEnumerable{System.Int64})">
+            <summary>
+            Computes the average of a sequence of nullable <see cref="T:System.Int64"/> values.
+            </summary>
+        </member>
+        <member name="M:Newtonsoft.Json.Utilities.LinqBridge.Enumerable.Average``1(System.Collections.Generic.IEnumerable{``0},Newtonsoft.Json.Serialization.Func{``0,System.Int64})">
+            <summary>
+            Computes the average of a sequence of nullable <see cref="T:System.Int64"/> values 
+            that are obtained by invoking a transform function on each 
+            element of the input sequence.
+            </summary>
+        </member>
+        <member name="M:Newtonsoft.Json.Utilities.LinqBridge.Enumerable.Sum(System.Collections.Generic.IEnumerable{System.Nullable{System.Int64}})">
+            <summary>
+            Computes the sum of a sequence of <see cref="T:System.Int64"/> values.
+            </summary>
+        </member>
+        <member name="M:Newtonsoft.Json.Utilities.LinqBridge.Enumerable.Sum``1(System.Collections.Generic.IEnumerable{``0},Newtonsoft.Json.Serialization.Func{``0,System.Nullable{System.Int64}})">
+            <summary>
+            Computes the sum of a sequence of <see cref="T:System.Int64"/> 
+            values that are obtained by invoking a transform function on 
+            each element of the input sequence.
+            </summary>
+        </member>
+        <member name="M:Newtonsoft.Json.Utilities.LinqBridge.Enumerable.Average(System.Collections.Generic.IEnumerable{System.Nullable{System.Int64}})">
+            <summary>
+            Computes the average of a sequence of <see cref="T:System.Int64"/> values.
+            </summary>
+        </member>
+        <member name="M:Newtonsoft.Json.Utilities.LinqBridge.Enumerable.Average``1(System.Collections.Generic.IEnumerable{``0},Newtonsoft.Json.Serialization.Func{``0,System.Nullable{System.Int64}})">
+            <summary>
+            Computes the average of a sequence of <see cref="T:System.Int64"/> values 
+            that are obtained by invoking a transform function on each 
+            element of the input sequence.
+            </summary>
+        </member>
+        <member name="M:Newtonsoft.Json.Utilities.LinqBridge.Enumerable.Min(System.Collections.Generic.IEnumerable{System.Nullable{System.Int64}})">
+            <summary>
+            Returns the minimum value in a sequence of nullable 
+            <see cref="T:System.Int64"/> values.
+            </summary>
+        </member>
+        <member name="M:Newtonsoft.Json.Utilities.LinqBridge.Enumerable.Min``1(System.Collections.Generic.IEnumerable{``0},Newtonsoft.Json.Serialization.Func{``0,System.Nullable{System.Int64}})">
+            <summary>
+            Invokes a transform function on each element of a sequence and 
+            returns the minimum nullable <see cref="T:System.Int64"/> value.
+            </summary>
+        </member>
+        <member name="M:Newtonsoft.Json.Utilities.LinqBridge.Enumerable.Max(System.Collections.Generic.IEnumerable{System.Nullable{System.Int64}})">
+            <summary>
+            Returns the maximum value in a sequence of nullable 
+            <see cref="T:System.Int64"/> values.
+            </summary>
+        </member>
+        <member name="M:Newtonsoft.Json.Utilities.LinqBridge.Enumerable.Max``1(System.Collections.Generic.IEnumerable{``0},Newtonsoft.Json.Serialization.Func{``0,System.Nullable{System.Int64}})">
+            <summary>
+            Invokes a transform function on each element of a sequence and 
+            returns the maximum nullable <see cref="T:System.Int64"/> value.
+            </summary>
+        </member>
+        <member name="M:Newtonsoft.Json.Utilities.LinqBridge.Enumerable.Sum(System.Collections.Generic.IEnumerable{System.Single})">
+            <summary>
+            Computes the sum of a sequence of nullable <see cref="T:System.Single"/> values.
+            </summary>
+        </member>
+        <member name="M:Newtonsoft.Json.Utilities.LinqBridge.Enumerable.Sum``1(System.Collections.Generic.IEnumerable{``0},Newtonsoft.Json.Serialization.Func{``0,System.Single})">
+            <summary>
+            Computes the sum of a sequence of nullable <see cref="T:System.Single"/> 
+            values that are obtained by invoking a transform function on 
+            each element of the input sequence.
+            </summary>
+        </member>
+        <member name="M:Newtonsoft.Json.Utilities.LinqBridge.Enumerable.Average(System.Collections.Generic.IEnumerable{System.Single})">
+            <summary>
+            Computes the average of a sequence of nullable <see cref="T:System.Single"/> values.
+            </summary>
+        </member>
+        <member name="M:Newtonsoft.Json.Utilities.LinqBridge.Enumerable.Average``1(System.Collections.Generic.IEnumerable{``0},Newtonsoft.Json.Serialization.Func{``0,System.Single})">
+            <summary>
+            Computes the average of a sequence of nullable <see cref="T:System.Single"/> values 
+            that are obtained by invoking a transform function on each 
+            element of the input sequence.
+            </summary>
+        </member>
+        <member name="M:Newtonsoft.Json.Utilities.LinqBridge.Enumerable.Sum(System.Collections.Generic.IEnumerable{System.Nullable{System.Single}})">
+            <summary>
+            Computes the sum of a sequence of <see cref="T:System.Single"/> values.
+            </summary>
+        </member>
+        <member name="M:Newtonsoft.Json.Utilities.LinqBridge.Enumerable.Sum``1(System.Collections.Generic.IEnumerable{``0},Newtonsoft.Json.Serialization.Func{``0,System.Nullable{System.Single}})">
+            <summary>
+            Computes the sum of a sequence of <see cref="T:System.Single"/> 
+            values that are obtained by invoking a transform function on 
+            each element of the input sequence.
+            </summary>
+        </member>
+        <member name="M:Newtonsoft.Json.Utilities.LinqBridge.Enumerable.Average(System.Collections.Generic.IEnumerable{System.Nullable{System.Single}})">
+            <summary>
+            Computes the average of a sequence of <see cref="T:System.Single"/> values.
+            </summary>
+        </member>
+        <member name="M:Newtonsoft.Json.Utilities.LinqBridge.Enumerable.Average``1(System.Collections.Generic.IEnumerable{``0},Newtonsoft.Json.Serialization.Func{``0,System.Nullable{System.Single}})">
+            <summary>
+            Computes the average of a sequence of <see cref="T:System.Single"/> values 
+            that are obtained by invoking a transform function on each 
+            element of the input sequence.
+            </summary>
+        </member>
+        <member name="M:Newtonsoft.Json.Utilities.LinqBridge.Enumerable.Min(System.Collections.Generic.IEnumerable{System.Nullable{System.Single}})">
+            <summary>
+            Returns the minimum value in a sequence of nullable 
+            <see cref="T:System.Single"/> values.
+            </summary>
+        </member>
+        <member name="M:Newtonsoft.Json.Utilities.LinqBridge.Enumerable.Min``1(System.Collections.Generic.IEnumerable{``0},Newtonsoft.Json.Serialization.Func{``0,System.Nullable{System.Single}})">
+            <summary>
+            Invokes a transform function on each element of a sequence and 
+            returns the minimum nullable <see cref="T:System.Single"/> value.
+            </summary>
+        </member>
+        <member name="M:Newtonsoft.Json.Utilities.LinqBridge.Enumerable.Max(System.Collections.Generic.IEnumerable{System.Nullable{System.Single}})">
+            <summary>
+            Returns the maximum value in a sequence of nullable 
+            <see cref="T:System.Single"/> values.
+            </summary>
+        </member>
+        <member name="M:Newtonsoft.Json.Utilities.LinqBridge.Enumerable.Max``1(System.Collections.Generic.IEnumerable{``0},Newtonsoft.Json.Serialization.Func{``0,System.Nullable{System.Single}})">
+            <summary>
+            Invokes a transform function on each element of a sequence and 
+            returns the maximum nullable <see cref="T:System.Single"/> value.
+            </summary>
+        </member>
+        <member name="M:Newtonsoft.Json.Utilities.LinqBridge.Enumerable.Sum(System.Collections.Generic.IEnumerable{System.Double})">
+            <summary>
+            Computes the sum of a sequence of nullable <see cref="T:System.Double"/> values.
+            </summary>
+        </member>
+        <member name="M:Newtonsoft.Json.Utilities.LinqBridge.Enumerable.Sum``1(System.Collections.Generic.IEnumerable{``0},Newtonsoft.Json.Serialization.Func{``0,System.Double})">
+            <summary>
+            Computes the sum of a sequence of nullable <see cref="T:System.Double"/> 
+            values that are obtained by invoking a transform function on 
+            each element of the input sequence.
+            </summary>
+        </member>
+        <member name="M:Newtonsoft.Json.Utilities.LinqBridge.Enumerable.Average(System.Collections.Generic.IEnumerable{System.Double})">
+            <summary>
+            Computes the average of a sequence of nullable <see cref="T:System.Double"/> values.
+            </summary>
+        </member>
+        <member name="M:Newtonsoft.Json.Utilities.LinqBridge.Enumerable.Average``1(System.Collections.Generic.IEnumerable{``0},Newtonsoft.Json.Serialization.Func{``0,System.Double})">
+            <summary>
+            Computes the average of a sequence of nullable <see cref="T:System.Double"/> values 
+            that are obtained by invoking a transform function on each 
+            element of the input sequence.
+            </summary>
+        </member>
+        <member name="M:Newtonsoft.Json.Utilities.LinqBridge.Enumerable.Sum(System.Collections.Generic.IEnumerable{System.Nullable{System.Double}})">
+            <summary>
+            Computes the sum of a sequence of <see cref="T:System.Double"/> values.
+            </summary>
+        </member>
+        <member name="M:Newtonsoft.Json.Utilities.LinqBridge.Enumerable.Sum``1(System.Collections.Generic.IEnumerable{``0},Newtonsoft.Json.Serialization.Func{``0,System.Nullable{System.Double}})">
+            <summary>
+            Computes the sum of a sequence of <see cref="T:System.Double"/> 
+            values that are obtained by invoking a transform function on 
+            each element of the input sequence.
+            </summary>
+        </member>
+        <member name="M:Newtonsoft.Json.Utilities.LinqBridge.Enumerable.Average(System.Collections.Generic.IEnumerable{System.Nullable{System.Double}})">
+            <summary>
+            Computes the average of a sequence of <see cref="T:System.Double"/> values.
+            </summary>
+        </member>
+        <member name="M:Newtonsoft.Json.Utilities.LinqBridge.Enumerable.Average``1(System.Collections.Generic.IEnumerable{``0},Newtonsoft.Json.Serialization.Func{``0,System.Nullable{System.Double}})">
+            <summary>
+            Computes the average of a sequence of <see cref="T:System.Double"/> values 
+            that are obtained by invoking a transform function on each 
+            element of the input sequence.
+            </summary>
+        </member>
+        <member name="M:Newtonsoft.Json.Utilities.LinqBridge.Enumerable.Min(System.Collections.Generic.IEnumerable{System.Nullable{System.Double}})">
+            <summary>
+            Returns the minimum value in a sequence of nullable 
+            <see cref="T:System.Double"/> values.
+            </summary>
+        </member>
+        <member name="M:Newtonsoft.Json.Utilities.LinqBridge.Enumerable.Min``1(System.Collections.Generic.IEnumerable{``0},Newtonsoft.Json.Serialization.Func{``0,System.Nullable{System.Double}})">
+            <summary>
+            Invokes a transform function on each element of a sequence and 
+            returns the minimum nullable <see cref="T:System.Double"/> value.
+            </summary>
+        </member>
+        <member name="M:Newtonsoft.Json.Utilities.LinqBridge.Enumerable.Max(System.Collections.Generic.IEnumerable{System.Nullable{System.Double}})">
+            <summary>
+            Returns the maximum value in a sequence of nullable 
+            <see cref="T:System.Double"/> values.
+            </summary>
+        </member>
+        <member name="M:Newtonsoft.Json.Utilities.LinqBridge.Enumerable.Max``1(System.Collections.Generic.IEnumerable{``0},Newtonsoft.Json.Serialization.Func{``0,System.Nullable{System.Double}})">
+            <summary>
+            Invokes a transform function on each element of a sequence and 
+            returns the maximum nullable <see cref="T:System.Double"/> value.
+            </summary>
+        </member>
+        <member name="M:Newtonsoft.Json.Utilities.LinqBridge.Enumerable.Sum(System.Collections.Generic.IEnumerable{System.Decimal})">
+            <summary>
+            Computes the sum of a sequence of nullable <see cref="T:System.Decimal"/> values.
+            </summary>
+        </member>
+        <member name="M:Newtonsoft.Json.Utilities.LinqBridge.Enumerable.Sum``1(System.Collections.Generic.IEnumerable{``0},Newtonsoft.Json.Serialization.Func{``0,System.Decimal})">
+            <summary>
+            Computes the sum of a sequence of nullable <see cref="T:System.Decimal"/> 
+            values that are obtained by invoking a transform function on 
+            each element of the input sequence.
+            </summary>
+        </member>
+        <member name="M:Newtonsoft.Json.Utilities.LinqBridge.Enumerable.Average(System.Collections.Generic.IEnumerable{System.Decimal})">
+            <summary>
+            Computes the average of a sequence of nullable <see cref="T:System.Decimal"/> values.
+            </summary>
+        </member>
+        <member name="M:Newtonsoft.Json.Utilities.LinqBridge.Enumerable.Average``1(System.Collections.Generic.IEnumerable{``0},Newtonsoft.Json.Serialization.Func{``0,System.Decimal})">
+            <summary>
+            Computes the average of a sequence of nullable <see cref="T:System.Decimal"/> values 
+            that are obtained by invoking a transform function on each 
+            element of the input sequence.
+            </summary>
+        </member>
+        <member name="M:Newtonsoft.Json.Utilities.LinqBridge.Enumerable.Sum(System.Collections.Generic.IEnumerable{System.Nullable{System.Decimal}})">
+            <summary>
+            Computes the sum of a sequence of <see cref="T:System.Decimal"/> values.
+            </summary>
+        </member>
+        <member name="M:Newtonsoft.Json.Utilities.LinqBridge.Enumerable.Sum``1(System.Collections.Generic.IEnumerable{``0},Newtonsoft.Json.Serialization.Func{``0,System.Nullable{System.Decimal}})">
+            <summary>
+            Computes the sum of a sequence of <see cref="T:System.Decimal"/> 
+            values that are obtained by invoking a transform function on 
+            each element of the input sequence.
+            </summary>
+        </member>
+        <member name="M:Newtonsoft.Json.Utilities.LinqBridge.Enumerable.Average(System.Collections.Generic.IEnumerable{System.Nullable{System.Decimal}})">
+            <summary>
+            Computes the average of a sequence of <see cref="T:System.Decimal"/> values.
+            </summary>
+        </member>
+        <member name="M:Newtonsoft.Json.Utilities.LinqBridge.Enumerable.Average``1(System.Collections.Generic.IEnumerable{``0},Newtonsoft.Json.Serialization.Func{``0,System.Nullable{System.Decimal}})">
+            <summary>
+            Computes the average of a sequence of <see cref="T:System.Decimal"/> values 
+            that are obtained by invoking a transform function on each 
+            element of the input sequence.
+            </summary>
+        </member>
+        <member name="M:Newtonsoft.Json.Utilities.LinqBridge.Enumerable.Min(System.Collections.Generic.IEnumerable{System.Nullable{System.Decimal}})">
+            <summary>
+            Returns the minimum value in a sequence of nullable 
+            <see cref="T:System.Decimal"/> values.
+            </summary>
+        </member>
+        <member name="M:Newtonsoft.Json.Utilities.LinqBridge.Enumerable.Min``1(System.Collections.Generic.IEnumerable{``0},Newtonsoft.Json.Serialization.Func{``0,System.Nullable{System.Decimal}})">
+            <summary>
+            Invokes a transform function on each element of a sequence and 
+            returns the minimum nullable <see cref="T:System.Decimal"/> value.
+            </summary>
+        </member>
+        <member name="M:Newtonsoft.Json.Utilities.LinqBridge.Enumerable.Max(System.Collections.Generic.IEnumerable{System.Nullable{System.Decimal}})">
+            <summary>
+            Returns the maximum value in a sequence of nullable 
+            <see cref="T:System.Decimal"/> values.
+            </summary>
+        </member>
+        <member name="M:Newtonsoft.Json.Utilities.LinqBridge.Enumerable.Max``1(System.Collections.Generic.IEnumerable{``0},Newtonsoft.Json.Serialization.Func{``0,System.Nullable{System.Decimal}})">
+            <summary>
+            Invokes a transform function on each element of a sequence and 
+            returns the maximum nullable <see cref="T:System.Decimal"/> value.
+            </summary>
+        </member>
+        <member name="T:Newtonsoft.Json.Utilities.LinqBridge.IGrouping`2">
+            <summary>
+            Represents a collection of objects that have a common key.
+            </summary>
+        </member>
+        <member name="P:Newtonsoft.Json.Utilities.LinqBridge.IGrouping`2.Key">
+            <summary>
+            Gets the key of the <see cref="T:Newtonsoft.Json.Utilities.LinqBridge.IGrouping`2"/>.
+            </summary>
+        </member>
+        <member name="T:Newtonsoft.Json.Utilities.LinqBridge.ILookup`2">
+            <summary>
+            Defines an indexer, size property, and Boolean search method for 
+            data structures that map keys to <see cref="T:System.Collections.Generic.IEnumerable`1"/> 
+            sequences of values.
+            </summary>
+        </member>
+        <member name="T:Newtonsoft.Json.Utilities.LinqBridge.IOrderedEnumerable`1">
+            <summary>
+            Represents a sorted sequence.
+            </summary>
+        </member>
+        <member name="M:Newtonsoft.Json.Utilities.LinqBridge.IOrderedEnumerable`1.CreateOrderedEnumerable``1(Newtonsoft.Json.Serialization.Func{`0,``0},System.Collections.Generic.IComparer{``0},System.Boolean)">
+            <summary>
+            Performs a subsequent ordering on the elements of an 
+            <see cref="T:Newtonsoft.Json.Utilities.LinqBridge.IOrderedEnumerable`1"/> according to a key.
+            </summary>
+        </member>
+        <member name="T:Newtonsoft.Json.Utilities.LinqBridge.Lookup`2">
+            <summary>
+            Represents a collection of keys each mapped to one or more values.
+            </summary>
+        </member>
+        <member name="M:Newtonsoft.Json.Utilities.LinqBridge.Lookup`2.Contains(`0)">
+            <summary>
+            Determines whether a specified key is in the <see cref="T:Newtonsoft.Json.Utilities.LinqBridge.Lookup`2"/>.
+            </summary>
+        </member>
+        <member name="M:Newtonsoft.Json.Utilities.LinqBridge.Lookup`2.ApplyResultSelector``1(Newtonsoft.Json.Serialization.Func{`0,System.Collections.Generic.IEnumerable{`1},``0})">
+            <summary>
+            Applies a transform function to each key and its associated 
+            values and returns the results.
+            </summary>
+        </member>
+        <member name="M:Newtonsoft.Json.Utilities.LinqBridge.Lookup`2.GetEnumerator">
+            <summary>
+            Returns a generic enumerator that iterates through the <see cref="T:Newtonsoft.Json.Utilities.LinqBridge.Lookup`2"/>.
+            </summary>
+        </member>
+        <member name="P:Newtonsoft.Json.Utilities.LinqBridge.Lookup`2.Count">
+            <summary>
+            Gets the number of key/value collection pairs in the <see cref="T:Newtonsoft.Json.Utilities.LinqBridge.Lookup`2"/>.
+            </summary>
+        </member>
+        <member name="P:Newtonsoft.Json.Utilities.LinqBridge.Lookup`2.Item(`0)">
+            <summary>
+            Gets the collection of values indexed by the specified key.
+            </summary>
+        </member>
+        <member name="M:Newtonsoft.Json.Utilities.LinqBridge.OrderedEnumerable`2.TagPosition(`0,System.Int32)">
+            <remarks>
+            See <a href="http://code.google.com/p/linqbridge/issues/detail?id=11">issue #11</a>
+            for why this method is needed and cannot be expressed as a 
+            lambda at the call site.
+            </remarks>
+        </member>
+        <member name="M:Newtonsoft.Json.Utilities.LinqBridge.OrderedEnumerable`2.GetFirst(Newtonsoft.Json.Utilities.LinqBridge.Tuple{`0,System.Int32})">
+            <remarks>
+            See <a href="http://code.google.com/p/linqbridge/issues/detail?id=11">issue #11</a>
+            for why this method is needed and cannot be expressed as a 
+            lambda at the call site.
+            </remarks>
+        </member>
+        <member name="T:System.Runtime.CompilerServices.ExtensionAttribute">
+            <remarks>
+            This attribute allows us to define extension methods without 
+            requiring .NET Framework 3.5. For more information, see the section,
+            <a href="http://msdn.microsoft.com/en-us/magazine/cc163317.aspx#S7">Extension Methods in .NET Framework 2.0 Apps</a>,
+            of <a href="http://msdn.microsoft.com/en-us/magazine/cc163317.aspx">Basic Instincts: Extension Methods</a>
+            column in <a href="http://msdn.microsoft.com/msdnmag/">MSDN Magazine</a>, 
+            issue <a href="http://msdn.microsoft.com/en-us/magazine/cc135410.aspx">Nov 2007</a>.
+            </remarks>
+        </member>
+        <member name="T:Newtonsoft.Json.Linq.JPropertyDescriptor">
+            <summary>
+            Represents a view of a <see cref="T:Newtonsoft.Json.Linq.JProperty"/>.
+            </summary>
+        </member>
+        <member name="M:Newtonsoft.Json.Linq.JPropertyDescriptor.#ctor(System.String,System.Type)">
+            <summary>
+            Initializes a new instance of the <see cref="T:Newtonsoft.Json.Linq.JPropertyDescriptor"/> class.
+            </summary>
+            <param name="name">The name.</param>
+            <param name="propertyType">Type of the property.</param>
+        </member>
+        <member name="M:Newtonsoft.Json.Linq.JPropertyDescriptor.CanResetValue(System.Object)">
+            <summary>
+            When overridden in a derived class, returns whether resetting an object changes its value.
+            </summary>
+            <returns>
+            true if resetting the component changes its value; otherwise, false.
+            </returns>
+            <param name="component">The component to test for reset capability. 
+                            </param>
+        </member>
+        <member name="M:Newtonsoft.Json.Linq.JPropertyDescriptor.GetValue(System.Object)">
+            <summary>
+            When overridden in a derived class, gets the current value of the property on a component.
+            </summary>
+            <returns>
+            The value of a property for a given component.
+            </returns>
+            <param name="component">The component with the property for which to retrieve the value. 
+                            </param>
+        </member>
+        <member name="M:Newtonsoft.Json.Linq.JPropertyDescriptor.ResetValue(System.Object)">
+            <summary>
+            When overridden in a derived class, resets the value for this property of the component to the default value.
+            </summary>
+            <param name="component">The component with the property value that is to be reset to the default value. 
+                            </param>
+        </member>
+        <member name="M:Newtonsoft.Json.Linq.JPropertyDescriptor.SetValue(System.Object,System.Object)">
+            <summary>
+            When overridden in a derived class, sets the value of the component to a different value.
+            </summary>
+            <param name="component">The component with the property value that is to be set. 
+                            </param><param name="value">The new value. 
+                            </param>
+        </member>
+        <member name="M:Newtonsoft.Json.Linq.JPropertyDescriptor.ShouldSerializeValue(System.Object)">
+            <summary>
+            When overridden in a derived class, determines a value indicating whether the value of this property needs to be persisted.
+            </summary>
+            <returns>
+            true if the property should be persisted; otherwise, false.
+            </returns>
+            <param name="component">The component with the property to be examined for persistence. 
+                            </param>
+        </member>
+        <member name="P:Newtonsoft.Json.Linq.JPropertyDescriptor.ComponentType">
+            <summary>
+            When overridden in a derived class, gets the type of the component this property is bound to.
+            </summary>
+            <returns>
+            A <see cref="T:System.Type"/> that represents the type of component this property is bound to. When the <see cref="M:System.ComponentModel.PropertyDescriptor.GetValue(System.Object)"/> or <see cref="M:System.ComponentModel.PropertyDescriptor.SetValue(System.Object,System.Object)"/> methods are invoked, the object specified might be an instance of this type.
+            </returns>
+        </member>
+        <member name="P:Newtonsoft.Json.Linq.JPropertyDescriptor.IsReadOnly">
+            <summary>
+            When overridden in a derived class, gets a value indicating whether this property is read-only.
+            </summary>
+            <returns>
+            true if the property is read-only; otherwise, false.
+            </returns>
+        </member>
+        <member name="P:Newtonsoft.Json.Linq.JPropertyDescriptor.PropertyType">
+            <summary>
+            When overridden in a derived class, gets the type of the property.
+            </summary>
+            <returns>
+            A <see cref="T:System.Type"/> that represents the type of the property.
+            </returns>
+        </member>
+        <member name="P:Newtonsoft.Json.Linq.JPropertyDescriptor.NameHashCode">
+            <summary>
+            Gets the hash code for the name of the member.
+            </summary>
+            <value></value>
+            <returns>
+            The hash code for the name of the member.
+            </returns>
+        </member>
+        <member name="T:Newtonsoft.Json.Linq.JRaw">
+            <summary>
+            Represents a raw JSON string.
+            </summary>
+        </member>
+        <member name="T:Newtonsoft.Json.Linq.JValue">
+            <summary>
+            Represents a value in JSON (string, integer, date, etc).
+            </summary>
+        </member>
+        <member name="T:Newtonsoft.Json.Linq.JToken">
+            <summary>
+            Represents an abstract JSON token.
+            </summary>
+        </member>
+        <member name="T:Newtonsoft.Json.Linq.IJEnumerable`1">
+            <summary>
+            Represents a collection of <see cref="T:Newtonsoft.Json.Linq.JToken"/> objects.
+            </summary>
+            <typeparam name="T">The type of token</typeparam>
+        </member>
+        <member name="P:Newtonsoft.Json.Linq.IJEnumerable`1.Item(System.Object)">
+            <summary>
+            Gets the <see cref="T:Newtonsoft.Json.Linq.IJEnumerable`1"/> with the specified key.
+            </summary>
+            <value></value>
+        </member>
+        <member name="T:Newtonsoft.Json.IJsonLineInfo">
+            <summary>
+            Provides an interface to enable a class to return line and position information.
+            </summary>
+        </member>
+        <member name="M:Newtonsoft.Json.IJsonLineInfo.HasLineInfo">
+            <summary>
+            Gets a value indicating whether the class can return line information.
+            </summary>
+            <returns>
+            	<c>true</c> if LineNumber and LinePosition can be provided; otherwise, <c>false</c>.
+            </returns>
+        </member>
+        <member name="P:Newtonsoft.Json.IJsonLineInfo.LineNumber">
+            <summary>
+            Gets the current line number.
+            </summary>
+            <value>The current line number or 0 if no line information is available (for example, HasLineInfo returns false).</value>
+        </member>
+        <member name="P:Newtonsoft.Json.IJsonLineInfo.LinePosition">
+            <summary>
+            Gets the current line position.
+            </summary>
+            <value>The current line position or 0 if no line information is available (for example, HasLineInfo returns false).</value>
+        </member>
+        <member name="M:Newtonsoft.Json.Linq.JToken.DeepEquals(Newtonsoft.Json.Linq.JToken,Newtonsoft.Json.Linq.JToken)">
+            <summary>
+            Compares the values of two tokens, including the values of all descendant tokens.
+            </summary>
+            <param name="t1">The first <see cref="T:Newtonsoft.Json.Linq.JToken"/> to compare.</param>
+            <param name="t2">The second <see cref="T:Newtonsoft.Json.Linq.JToken"/> to compare.</param>
+            <returns>true if the tokens are equal; otherwise false.</returns>
+        </member>
+        <member name="M:Newtonsoft.Json.Linq.JToken.AddAfterSelf(System.Object)">
+            <summary>
+            Adds the specified content immediately after this token.
+            </summary>
+            <param name="content">A content object that contains simple content or a collection of content objects to be added after this token.</param>
+        </member>
+        <member name="M:Newtonsoft.Json.Linq.JToken.AddBeforeSelf(System.Object)">
+            <summary>
+            Adds the specified content immediately before this token.
+            </summary>
+            <param name="content">A content object that contains simple content or a collection of content objects to be added before this token.</param>
+        </member>
+        <member name="M:Newtonsoft.Json.Linq.JToken.Ancestors">
+            <summary>
+            Returns a collection of the ancestor tokens of this token.
+            </summary>
+            <returns>A collection of the ancestor tokens of this token.</returns>
+        </member>
+        <member name="M:Newtonsoft.Json.Linq.JToken.AfterSelf">
+            <summary>
+            Returns a collection of the sibling tokens after this token, in document order.
+            </summary>
+            <returns>A collection of the sibling tokens after this tokens, in document order.</returns>
+        </member>
+        <member name="M:Newtonsoft.Json.Linq.JToken.BeforeSelf">
+            <summary>
+            Returns a collection of the sibling tokens before this token, in document order.
+            </summary>
+            <returns>A collection of the sibling tokens before this token, in document order.</returns>
+        </member>
+        <member name="M:Newtonsoft.Json.Linq.JToken.Value``1(System.Object)">
+            <summary>
+            Gets the <see cref="T:Newtonsoft.Json.Linq.JToken"/> with the specified key converted to the specified type.
+            </summary>
+            <typeparam name="T">The type to convert the token to.</typeparam>
+            <param name="key">The token key.</param>
+            <returns>The converted token value.</returns>
+        </member>
+        <member name="M:Newtonsoft.Json.Linq.JToken.Children">
+            <summary>
+            Returns a collection of the child tokens of this token, in document order.
+            </summary>
+            <returns>An <see cref="T:System.Collections.Generic.IEnumerable`1"/> of <see cref="T:Newtonsoft.Json.Linq.JToken"/> containing the child tokens of this <see cref="T:Newtonsoft.Json.Linq.JToken"/>, in document order.</returns>
+        </member>
+        <member name="M:Newtonsoft.Json.Linq.JToken.Children``1">
+            <summary>
+            Returns a collection of the child tokens of this token, in document order, filtered by the specified type.
+            </summary>
+            <typeparam name="T">The type to filter the child tokens on.</typeparam>
+            <returns>A <see cref="T:Newtonsoft.Json.Linq.JEnumerable`1"/> containing the child tokens of this <see cref="T:Newtonsoft.Json.Linq.JToken"/>, in document order.</returns>
+        </member>
+        <member name="M:Newtonsoft.Json.Linq.JToken.Values``1">
+            <summary>
+            Returns a collection of the child values of this token, in document order.
+            </summary>
+            <typeparam name="T">The type to convert the values to.</typeparam>
+            <returns>A <see cref="T:System.Collections.Generic.IEnumerable`1"/> containing the child values of this <see cref="T:Newtonsoft.Json.Linq.JToken"/>, in document order.</returns>
+        </member>
+        <member name="M:Newtonsoft.Json.Linq.JToken.Remove">
+            <summary>
+            Removes this token from its parent.
+            </summary>
+        </member>
+        <member name="M:Newtonsoft.Json.Linq.JToken.Replace(Newtonsoft.Json.Linq.JToken)">
+            <summary>
+            Replaces this token with the specified token.
+            </summary>
+            <param name="value">The value.</param>
+        </member>
+        <member name="M:Newtonsoft.Json.Linq.JToken.WriteTo(Newtonsoft.Json.JsonWriter,Newtonsoft.Json.JsonConverter[])">
+            <summary>
+            Writes this token to a <see cref="T:Newtonsoft.Json.JsonWriter"/>.
+            </summary>
+            <param name="writer">A <see cref="T:Newtonsoft.Json.JsonWriter"/> into which this method will write.</param>
+            <param name="converters">A collection of <see cref="T:Newtonsoft.Json.JsonConverter"/> which will be used when writing the token.</param>
+        </member>
+        <member name="M:Newtonsoft.Json.Linq.JToken.ToString">
+            <summary>
+            Returns the indented JSON for this token.
+            </summary>
+            <returns>
+            The indented JSON for this token.
+            </returns>
+        </member>
+        <member name="M:Newtonsoft.Json.Linq.JToken.ToString(Newtonsoft.Json.Formatting,Newtonsoft.Json.JsonConverter[])">
+            <summary>
+            Returns the JSON for this token using the given formatting and converters.
+            </summary>
+            <param name="formatting">Indicates how the output is formatted.</param>
+            <param name="converters">A collection of <see cref="T:Newtonsoft.Json.JsonConverter"/> which will be used when writing the token.</param>
+            <returns>The JSON for this token using the given formatting and converters.</returns>
+        </member>
+        <member name="M:Newtonsoft.Json.Linq.JToken.op_Explicit(Newtonsoft.Json.Linq.JToken)~System.Boolean">
+            <summary>
+            Performs an explicit conversion from <see cref="T:Newtonsoft.Json.Linq.JToken"/> to <see cref="T:System.Boolean"/>.
+            </summary>
+            <param name="value">The value.</param>
+            <returns>The result of the conversion.</returns>
+        </member>
+        <member name="M:Newtonsoft.Json.Linq.JToken.op_Explicit(Newtonsoft.Json.Linq.JToken)~System.Nullable{System.Boolean}">
+            <summary>
+            Performs an explicit conversion from <see cref="T:Newtonsoft.Json.Linq.JToken"/> to <see cref="T:System.Nullable`1"/>.
+            </summary>
+            <param name="value">The value.</param>
+            <returns>The result of the conversion.</returns>
+        </member>
+        <member name="M:Newtonsoft.Json.Linq.JToken.op_Explicit(Newtonsoft.Json.Linq.JToken)~System.Int64">
+            <summary>
+            Performs an explicit conversion from <see cref="T:Newtonsoft.Json.Linq.JToken"/> to <see cref="T:System.Int64"/>.
+            </summary>
+            <param name="value">The value.</param>
+            <returns>The result of the conversion.</returns>
+        </member>
+        <member name="M:Newtonsoft.Json.Linq.JToken.op_Explicit(Newtonsoft.Json.Linq.JToken)~System.Nullable{System.DateTime}">
+            <summary>
+            Performs an explicit conversion from <see cref="T:Newtonsoft.Json.Linq.JToken"/> to <see cref="T:System.Nullable`1"/>.
+            </summary>
+            <param name="value">The value.</param>
+            <returns>The result of the conversion.</returns>
+        </member>
+        <member name="M:Newtonsoft.Json.Linq.JToken.op_Explicit(Newtonsoft.Json.Linq.JToken)~System.Nullable{System.Decimal}">
+            <summary>
+            Performs an explicit conversion from <see cref="T:Newtonsoft.Json.Linq.JToken"/> to <see cref="T:System.Nullable`1"/>.
+            </summary>
+            <param name="value">The value.</param>
+            <returns>The result of the conversion.</returns>
+        </member>
+        <member name="M:Newtonsoft.Json.Linq.JToken.op_Explicit(Newtonsoft.Json.Linq.JToken)~System.Nullable{System.Double}">
+            <summary>
+            Performs an explicit conversion from <see cref="T:Newtonsoft.Json.Linq.JToken"/> to <see cref="T:System.Nullable`1"/>.
+            </summary>
+            <param name="value">The value.</param>
+            <returns>The result of the conversion.</returns>
+        </member>
+        <member name="M:Newtonsoft.Json.Linq.JToken.op_Explicit(Newtonsoft.Json.Linq.JToken)~System.Nullable{System.Char}">
+            <summary>
+            Performs an explicit conversion from <see cref="T:Newtonsoft.Json.Linq.JToken"/> to <see cref="T:System.Nullable`1"/>.
+            </summary>
+            <param name="value">The value.</param>
+            <returns>The result of the conversion.</returns>
+        </member>
+        <member name="M:Newtonsoft.Json.Linq.JToken.op_Explicit(Newtonsoft.Json.Linq.JToken)~System.Int32">
+            <summary>
+            Performs an explicit conversion from <see cref="T:Newtonsoft.Json.Linq.JToken"/> to <see cref="T:System.Int32"/>.
+            </summary>
+            <param name="value">The value.</param>
+            <returns>The result of the conversion.</returns>
+        </member>
+        <member name="M:Newtonsoft.Json.Linq.JToken.op_Explicit(Newtonsoft.Json.Linq.JToken)~System.Int16">
+            <summary>
+            Performs an explicit conversion from <see cref="T:Newtonsoft.Json.Linq.JToken"/> to <see cref="T:System.Int16"/>.
+            </summary>
+            <param name="value">The value.</param>
+            <returns>The result of the conversion.</returns>
+        </member>
+        <member name="M:Newtonsoft.Json.Linq.JToken.op_Explicit(Newtonsoft.Json.Linq.JToken)~System.UInt16">
+            <summary>
+            Performs an explicit conversion from <see cref="T:Newtonsoft.Json.Linq.JToken"/> to <see cref="T:System.UInt16"/>.
+            </summary>
+            <param name="value">The value.</param>
+            <returns>The result of the conversion.</returns>
+        </member>
+        <member name="M:Newtonsoft.Json.Linq.JToken.op_Explicit(Newtonsoft.Json.Linq.JToken)~System.Char">
+            <summary>
+            Performs an explicit conversion from <see cref="T:Newtonsoft.Json.Linq.JToken"/> to <see cref="T:System.Char"/>.
+            </summary>
+            <param name="value">The value.</param>
+            <returns>The result of the conversion.</returns>
+        </member>
+        <member name="M:Newtonsoft.Json.Linq.JToken.op_Explicit(Newtonsoft.Json.Linq.JToken)~System.Byte">
+            <summary>
+            Performs an explicit conversion from <see cref="T:Newtonsoft.Json.Linq.JToken"/> to <see cref="T:System.Byte"/>.
+            </summary>
+            <param name="value">The value.</param>
+            <returns>The result of the conversion.</returns>
+        </member>
+        <member name="M:Newtonsoft.Json.Linq.JToken.op_Explicit(Newtonsoft.Json.Linq.JToken)~System.Nullable{System.Int32}">
+            <summary>
+            Performs an explicit conversion from <see cref="T:Newtonsoft.Json.Linq.JToken"/> to <see cref="T:System.Nullable`1"/>.
+            </summary>
+            <param name="value">The value.</param>
+            <returns>The result of the conversion.</returns>
+        </member>
+        <member name="M:Newtonsoft.Json.Linq.JToken.op_Explicit(Newtonsoft.Json.Linq.JToken)~System.Nullable{System.Int16}">
+            <summary>
+            Performs an explicit conversion from <see cref="T:Newtonsoft.Json.Linq.JToken"/> to <see cref="T:System.Nullable`1"/>.
+            </summary>
+            <param name="value">The value.</param>
+            <returns>The result of the conversion.</returns>
+        </member>
+        <member name="M:Newtonsoft.Json.Linq.JToken.op_Explicit(Newtonsoft.Json.Linq.JToken)~System.Nullable{System.UInt16}">
+            <summary>
+            Performs an explicit conversion from <see cref="T:Newtonsoft.Json.Linq.JToken"/> to <see cref="T:System.Nullable`1"/>.
+            </summary>
+            <param name="value">The value.</param>
+            <returns>The result of the conversion.</returns>
+        </member>
+        <member name="M:Newtonsoft.Json.Linq.JToken.op_Explicit(Newtonsoft.Json.Linq.JToken)~System.Nullable{System.Byte}">
+            <summary>
+            Performs an explicit conversion from <see cref="T:Newtonsoft.Json.Linq.JToken"/> to <see cref="T:System.Nullable`1"/>.
+            </summary>
+            <param name="value">The value.</param>
+            <returns>The result of the conversion.</returns>
+        </member>
+        <member name="M:Newtonsoft.Json.Linq.JToken.op_Explicit(Newtonsoft.Json.Linq.JToken)~System.DateTime">
+            <summary>
+            Performs an explicit conversion from <see cref="T:Newtonsoft.Json.Linq.JToken"/> to <see cref="T:System.DateTime"/>.
+            </summary>
+            <param name="value">The value.</param>
+            <returns>The result of the conversion.</returns>
+        </member>
+        <member name="M:Newtonsoft.Json.Linq.JToken.op_Explicit(Newtonsoft.Json.Linq.JToken)~System.Nullable{System.Int64}">
+            <summary>
+            Performs an explicit conversion from <see cref="T:Newtonsoft.Json.Linq.JToken"/> to <see cref="T:System.Nullable`1"/>.
+            </summary>
+            <param name="value">The value.</param>
+            <returns>The result of the conversion.</returns>
+        </member>
+        <member name="M:Newtonsoft.Json.Linq.JToken.op_Explicit(Newtonsoft.Json.Linq.JToken)~System.Nullable{System.Single}">
+            <summary>
+            Performs an explicit conversion from <see cref="T:Newtonsoft.Json.Linq.JToken"/> to <see cref="T:System.Nullable`1"/>.
+            </summary>
+            <param name="value">The value.</param>
+            <returns>The result of the conversion.</returns>
+        </member>
+        <member name="M:Newtonsoft.Json.Linq.JToken.op_Explicit(Newtonsoft.Json.Linq.JToken)~System.Decimal">
+            <summary>
+            Performs an explicit conversion from <see cref="T:Newtonsoft.Json.Linq.JToken"/> to <see cref="T:System.Decimal"/>.
+            </summary>
+            <param name="value">The value.</param>
+            <returns>The result of the conversion.</returns>
+        </member>
+        <member name="M:Newtonsoft.Json.Linq.JToken.op_Explicit(Newtonsoft.Json.Linq.JToken)~System.Nullable{System.UInt32}">
+            <summary>
+            Performs an explicit conversion from <see cref="T:Newtonsoft.Json.Linq.JToken"/> to <see cref="T:System.Nullable`1"/>.
+            </summary>
+            <param name="value">The value.</param>
+            <returns>The result of the conversion.</returns>
+        </member>
+        <member name="M:Newtonsoft.Json.Linq.JToken.op_Explicit(Newtonsoft.Json.Linq.JToken)~System.Nullable{System.UInt64}">
+            <summary>
+            Performs an explicit conversion from <see cref="T:Newtonsoft.Json.Linq.JToken"/> to <see cref="T:System.Nullable`1"/>.
+            </summary>
+            <param name="value">The value.</param>
+            <returns>The result of the conversion.</returns>
+        </member>
+        <member name="M:Newtonsoft.Json.Linq.JToken.op_Explicit(Newtonsoft.Json.Linq.JToken)~System.Double">
+            <summary>
+            Performs an explicit conversion from <see cref="T:Newtonsoft.Json.Linq.JToken"/> to <see cref="T:System.Double"/>.
+            </summary>
+            <param name="value">The value.</param>
+            <returns>The result of the conversion.</returns>
+        </member>
+        <member name="M:Newtonsoft.Json.Linq.JToken.op_Explicit(Newtonsoft.Json.Linq.JToken)~System.Single">
+            <summary>
+            Performs an explicit conversion from <see cref="T:Newtonsoft.Json.Linq.JToken"/> to <see cref="T:System.Single"/>.
+            </summary>
+            <param name="value">The value.</param>
+            <returns>The result of the conversion.</returns>
+        </member>
+        <member name="M:Newtonsoft.Json.Linq.JToken.op_Explicit(Newtonsoft.Json.Linq.JToken)~System.String">
+            <summary>
+            Performs an explicit conversion from <see cref="T:Newtonsoft.Json.Linq.JToken"/> to <see cref="T:System.String"/>.
+            </summary>
+            <param name="value">The value.</param>
+            <returns>The result of the conversion.</returns>
+        </member>
+        <member name="M:Newtonsoft.Json.Linq.JToken.op_Explicit(Newtonsoft.Json.Linq.JToken)~System.UInt32">
+            <summary>
+            Performs an explicit conversion from <see cref="T:Newtonsoft.Json.Linq.JToken"/> to <see cref="T:System.UInt32"/>.
+            </summary>
+            <param name="value">The value.</param>
+            <returns>The result of the conversion.</returns>
+        </member>
+        <member name="M:Newtonsoft.Json.Linq.JToken.op_Explicit(Newtonsoft.Json.Linq.JToken)~System.UInt64">
+            <summary>
+            Performs an explicit conversion from <see cref="T:Newtonsoft.Json.Linq.JToken"/> to <see cref="T:System.UInt64"/>.
+            </summary>
+            <param name="value">The value.</param>
+            <returns>The result of the conversion.</returns>
+        </member>
+        <member name="M:Newtonsoft.Json.Linq.JToken.op_Explicit(Newtonsoft.Json.Linq.JToken)~System.Byte[]">
+            <summary>
+            Performs an explicit conversion from <see cref="T:Newtonsoft.Json.Linq.JToken"/> to <see cref="T:System.Byte[]"/>.
+            </summary>
+            <param name="value">The value.</param>
+            <returns>The result of the conversion.</returns>
+        </member>
+        <member name="M:Newtonsoft.Json.Linq.JToken.op_Explicit(Newtonsoft.Json.Linq.JToken)~System.Guid">
+            <summary>
+            Performs an explicit conversion from <see cref="T:Newtonsoft.Json.Linq.JToken"/> to <see cref="T:System.Guid"/>.
+            </summary>
+            <param name="value">The value.</param>
+            <returns>The result of the conversion.</returns>
+        </member>
+        <member name="M:Newtonsoft.Json.Linq.JToken.op_Explicit(Newtonsoft.Json.Linq.JToken)~System.Nullable{System.Guid}">
+            <summary>
+            Performs an explicit conversion from <see cref="T:Newtonsoft.Json.Linq.JToken"/> to <see cref="T:System.Guid"/>.
+            </summary>
+            <param name="value">The value.</param>
+            <returns>The result of the conversion.</returns>
+        </member>
+        <member name="M:Newtonsoft.Json.Linq.JToken.op_Explicit(Newtonsoft.Json.Linq.JToken)~System.TimeSpan">
+            <summary>
+            Performs an explicit conversion from <see cref="T:Newtonsoft.Json.Linq.JToken"/> to <see cref="T:System.TimeSpan"/>.
+            </summary>
+            <param name="value">The value.</param>
+            <returns>The result of the conversion.</returns>
+        </member>
+        <member name="M:Newtonsoft.Json.Linq.JToken.op_Explicit(Newtonsoft.Json.Linq.JToken)~System.Nullable{System.TimeSpan}">
+            <summary>
+            Performs an explicit conversion from <see cref="T:Newtonsoft.Json.Linq.JToken"/> to <see cref="T:System.TimeSpan"/>.
+            </summary>
+            <param name="value">The value.</param>
+            <returns>The result of the conversion.</returns>
+        </member>
+        <member name="M:Newtonsoft.Json.Linq.JToken.op_Explicit(Newtonsoft.Json.Linq.JToken)~System.Uri">
+            <summary>
+            Performs an explicit conversion from <see cref="T:Newtonsoft.Json.Linq.JToken"/> to <see cref="T:System.Uri"/>.
+            </summary>
+            <param name="value">The value.</param>
+            <returns>The result of the conversion.</returns>
+        </member>
+        <member name="M:Newtonsoft.Json.Linq.JToken.op_Implicit(System.Boolean)~Newtonsoft.Json.Linq.JToken">
+            <summary>
+            Performs an implicit conversion from <see cref="T:System.Boolean"/> to <see cref="T:Newtonsoft.Json.Linq.JToken"/>.
+            </summary>
+            <param name="value">The value to create a <see cref="T:Newtonsoft.Json.Linq.JValue"/> from.</param>
+            <returns>The <see cref="T:Newtonsoft.Json.Linq.JValue"/> initialized with the specified value.</returns>
+        </member>
+        <member name="M:Newtonsoft.Json.Linq.JToken.op_Implicit(System.Nullable{System.Boolean})~Newtonsoft.Json.Linq.JToken">
+            <summary>
+            Performs an implicit conversion from <see cref="T:System.Nullable`1"/> to <see cref="T:Newtonsoft.Json.Linq.JToken"/>.
+            </summary>
+            <param name="value">The value to create a <see cref="T:Newtonsoft.Json.Linq.JValue"/> from.</param>
+            <returns>The <see cref="T:Newtonsoft.Json.Linq.JValue"/> initialized with the specified value.</returns>
+        </member>
+        <member name="M:Newtonsoft.Json.Linq.JToken.op_Implicit(System.Int64)~Newtonsoft.Json.Linq.JToken">
+            <summary>
+            Performs an implicit conversion from <see cref="T:System.Nullable`1"/> to <see cref="T:Newtonsoft.Json.Linq.JToken"/>.
+            </summary>
+            <param name="value">The value to create a <see cref="T:Newtonsoft.Json.Linq.JValue"/> from.</param>
+            <returns>The <see cref="T:Newtonsoft.Json.Linq.JValue"/> initialized with the specified value.</returns>
+        </member>
+        <member name="M:Newtonsoft.Json.Linq.JToken.op_Implicit(System.Nullable{System.DateTime})~Newtonsoft.Json.Linq.JToken">
+            <summary>
+            Performs an implicit conversion from <see cref="T:System.Nullable`1"/> to <see cref="T:Newtonsoft.Json.Linq.JToken"/>.
+            </summary>
+            <param name="value">The value to create a <see cref="T:Newtonsoft.Json.Linq.JValue"/> from.</param>
+            <returns>The <see cref="T:Newtonsoft.Json.Linq.JValue"/> initialized with the specified value.</returns>
+        </member>
+        <member name="M:Newtonsoft.Json.Linq.JToken.op_Implicit(System.Nullable{System.Decimal})~Newtonsoft.Json.Linq.JToken">
+            <summary>
+            Performs an implicit conversion from <see cref="T:System.Nullable`1"/> to <see cref="T:Newtonsoft.Json.Linq.JToken"/>.
+            </summary>
+            <param name="value">The value to create a <see cref="T:Newtonsoft.Json.Linq.JValue"/> from.</param>
+            <returns>The <see cref="T:Newtonsoft.Json.Linq.JValue"/> initialized with the specified value.</returns>
+        </member>
+        <member name="M:Newtonsoft.Json.Linq.JToken.op_Implicit(System.Nullable{System.Double})~Newtonsoft.Json.Linq.JToken">
+            <summary>
+            Performs an implicit conversion from <see cref="T:System.Nullable`1"/> to <see cref="T:Newtonsoft.Json.Linq.JToken"/>.
+            </summary>
+            <param name="value">The value to create a <see cref="T:Newtonsoft.Json.Linq.JValue"/> from.</param>
+            <returns>The <see cref="T:Newtonsoft.Json.Linq.JValue"/> initialized with the specified value.</returns>
+        </member>
+        <member name="M:Newtonsoft.Json.Linq.JToken.op_Implicit(System.Int16)~Newtonsoft.Json.Linq.JToken">
+            <summary>
+            Performs an implicit conversion from <see cref="T:System.Int16"/> to <see cref="T:Newtonsoft.Json.Linq.JToken"/>.
+            </summary>
+            <param name="value">The value to create a <see cref="T:Newtonsoft.Json.Linq.JValue"/> from.</param>
+            <returns>The <see cref="T:Newtonsoft.Json.Linq.JValue"/> initialized with the specified value.</returns>
+        </member>
+        <member name="M:Newtonsoft.Json.Linq.JToken.op_Implicit(System.UInt16)~Newtonsoft.Json.Linq.JToken">
+            <summary>
+            Performs an implicit conversion from <see cref="T:System.UInt16"/> to <see cref="T:Newtonsoft.Json.Linq.JToken"/>.
+            </summary>
+            <param name="value">The value to create a <see cref="T:Newtonsoft.Json.Linq.JValue"/> from.</param>
+            <returns>The <see cref="T:Newtonsoft.Json.Linq.JValue"/> initialized with the specified value.</returns>
+        </member>
+        <member name="M:Newtonsoft.Json.Linq.JToken.op_Implicit(System.Int32)~Newtonsoft.Json.Linq.JToken">
+            <summary>
+            Performs an implicit conversion from <see cref="T:System.Int32"/> to <see cref="T:Newtonsoft.Json.Linq.JToken"/>.
+            </summary>
+            <param name="value">The value to create a <see cref="T:Newtonsoft.Json.Linq.JValue"/> from.</param>
+            <returns>The <see cref="T:Newtonsoft.Json.Linq.JValue"/> initialized with the specified value.</returns>
+        </member>
+        <member name="M:Newtonsoft.Json.Linq.JToken.op_Implicit(System.Nullable{System.Int32})~Newtonsoft.Json.Linq.JToken">
+            <summary>
+            Performs an implicit conversion from <see cref="T:System.Nullable`1"/> to <see cref="T:Newtonsoft.Json.Linq.JToken"/>.
+            </summary>
+            <param name="value">The value to create a <see cref="T:Newtonsoft.Json.Linq.JValue"/> from.</param>
+            <returns>The <see cref="T:Newtonsoft.Json.Linq.JValue"/> initialized with the specified value.</returns>
+        </member>
+        <member name="M:Newtonsoft.Json.Linq.JToken.op_Implicit(System.DateTime)~Newtonsoft.Json.Linq.JToken">
+            <summary>
+            Performs an implicit conversion from <see cref="T:System.DateTime"/> to <see cref="T:Newtonsoft.Json.Linq.JToken"/>.
+            </summary>
+            <param name="value">The value to create a <see cref="T:Newtonsoft.Json.Linq.JValue"/> from.</param>
+            <returns>The <see cref="T:Newtonsoft.Json.Linq.JValue"/> initialized with the specified value.</returns>
+        </member>
+        <member name="M:Newtonsoft.Json.Linq.JToken.op_Implicit(System.Nullable{System.Int64})~Newtonsoft.Json.Linq.JToken">
+            <summary>
+            Performs an implicit conversion from <see cref="T:System.Nullable`1"/> to <see cref="T:Newtonsoft.Json.Linq.JToken"/>.
+            </summary>
+            <param name="value">The value to create a <see cref="T:Newtonsoft.Json.Linq.JValue"/> from.</param>
+            <returns>The <see cref="T:Newtonsoft.Json.Linq.JValue"/> initialized with the specified value.</returns>
+        </member>
+        <member name="M:Newtonsoft.Json.Linq.JToken.op_Implicit(System.Nullable{System.Single})~Newtonsoft.Json.Linq.JToken">
+            <summary>
+            Performs an implicit conversion from <see cref="T:System.Nullable`1"/> to <see cref="T:Newtonsoft.Json.Linq.JToken"/>.
+            </summary>
+            <param name="value">The value to create a <see cref="T:Newtonsoft.Json.Linq.JValue"/> from.</param>
+            <returns>The <see cref="T:Newtonsoft.Json.Linq.JValue"/> initialized with the specified value.</returns>
+        </member>
+        <member name="M:Newtonsoft.Json.Linq.JToken.op_Implicit(System.Decimal)~Newtonsoft.Json.Linq.JToken">
+            <summary>
+            Performs an implicit conversion from <see cref="T:System.Decimal"/> to <see cref="T:Newtonsoft.Json.Linq.JToken"/>.
+            </summary>
+            <param name="value">The value to create a <see cref="T:Newtonsoft.Json.Linq.JValue"/> from.</param>
+            <returns>The <see cref="T:Newtonsoft.Json.Linq.JValue"/> initialized with the specified value.</returns>
+        </member>
+        <member name="M:Newtonsoft.Json.Linq.JToken.op_Implicit(System.Nullable{System.Int16})~Newtonsoft.Json.Linq.JToken">
+            <summary>
+            Performs an implicit conversion from <see cref="T:System.Nullable`1"/> to <see cref="T:Newtonsoft.Json.Linq.JToken"/>.
+            </summary>
+            <param name="value">The value to create a <see cref="T:Newtonsoft.Json.Linq.JValue"/> from.</param>
+            <returns>The <see cref="T:Newtonsoft.Json.Linq.JValue"/> initialized with the specified value.</returns>
+        </member>
+        <member name="M:Newtonsoft.Json.Linq.JToken.op_Implicit(System.Nullable{System.UInt16})~Newtonsoft.Json.Linq.JToken">
+            <summary>
+            Performs an implicit conversion from <see cref="T:System.Nullable`1"/> to <see cref="T:Newtonsoft.Json.Linq.JToken"/>.
+            </summary>
+            <param name="value">The value to create a <see cref="T:Newtonsoft.Json.Linq.JValue"/> from.</param>
+            <returns>The <see cref="T:Newtonsoft.Json.Linq.JValue"/> initialized with the specified value.</returns>
+        </member>
+        <member name="M:Newtonsoft.Json.Linq.JToken.op_Implicit(System.Nullable{System.UInt32})~Newtonsoft.Json.Linq.JToken">
+            <summary>
+            Performs an implicit conversion from <see cref="T:System.Nullable`1"/> to <see cref="T:Newtonsoft.Json.Linq.JToken"/>.
+            </summary>
+            <param name="value">The value to create a <see cref="T:Newtonsoft.Json.Linq.JValue"/> from.</param>
+            <returns>The <see cref="T:Newtonsoft.Json.Linq.JValue"/> initialized with the specified value.</returns>
+        </member>
+        <member name="M:Newtonsoft.Json.Linq.JToken.op_Implicit(System.Nullable{System.UInt64})~Newtonsoft.Json.Linq.JToken">
+            <summary>
+            Performs an implicit conversion from <see cref="T:System.Nullable`1"/> to <see cref="T:Newtonsoft.Json.Linq.JToken"/>.
+            </summary>
+            <param name="value">The value to create a <see cref="T:Newtonsoft.Json.Linq.JValue"/> from.</param>
+            <returns>The <see cref="T:Newtonsoft.Json.Linq.JValue"/> initialized with the specified value.</returns>
+        </member>
+        <member name="M:Newtonsoft.Json.Linq.JToken.op_Implicit(System.Double)~Newtonsoft.Json.Linq.JToken">
+            <summary>
+            Performs an implicit conversion from <see cref="T:System.Double"/> to <see cref="T:Newtonsoft.Json.Linq.JToken"/>.
+            </summary>
+            <param name="value">The value to create a <see cref="T:Newtonsoft.Json.Linq.JValue"/> from.</param>
+            <returns>The <see cref="T:Newtonsoft.Json.Linq.JValue"/> initialized with the specified value.</returns>
+        </member>
+        <member name="M:Newtonsoft.Json.Linq.JToken.op_Implicit(System.Single)~Newtonsoft.Json.Linq.JToken">
+            <summary>
+            Performs an implicit conversion from <see cref="T:System.Single"/> to <see cref="T:Newtonsoft.Json.Linq.JToken"/>.
+            </summary>
+            <param name="value">The value to create a <see cref="T:Newtonsoft.Json.Linq.JValue"/> from.</param>
+            <returns>The <see cref="T:Newtonsoft.Json.Linq.JValue"/> initialized with the specified value.</returns>
+        </member>
+        <member name="M:Newtonsoft.Json.Linq.JToken.op_Implicit(System.String)~Newtonsoft.Json.Linq.JToken">
+            <summary>
+            Performs an implicit conversion from <see cref="T:System.String"/> to <see cref="T:Newtonsoft.Json.Linq.JToken"/>.
+            </summary>
+            <param name="value">The value to create a <see cref="T:Newtonsoft.Json.Linq.JValue"/> from.</param>
+            <returns>The <see cref="T:Newtonsoft.Json.Linq.JValue"/> initialized with the specified value.</returns>
+        </member>
+        <member name="M:Newtonsoft.Json.Linq.JToken.op_Implicit(System.UInt32)~Newtonsoft.Json.Linq.JToken">
+            <summary>
+            Performs an implicit conversion from <see cref="T:System.UInt32"/> to <see cref="T:Newtonsoft.Json.Linq.JToken"/>.
+            </summary>
+            <param name="value">The value to create a <see cref="T:Newtonsoft.Json.Linq.JValue"/> from.</param>
+            <returns>The <see cref="T:Newtonsoft.Json.Linq.JValue"/> initialized with the specified value.</returns>
+        </member>
+        <member name="M:Newtonsoft.Json.Linq.JToken.op_Implicit(System.UInt64)~Newtonsoft.Json.Linq.JToken">
+            <summary>
+            Performs an implicit conversion from <see cref="T:System.UInt64"/> to <see cref="T:Newtonsoft.Json.Linq.JToken"/>.
+            </summary>
+            <param name="value">The value to create a <see cref="T:Newtonsoft.Json.Linq.JValue"/> from.</param>
+            <returns>The <see cref="T:Newtonsoft.Json.Linq.JValue"/> initialized with the specified value.</returns>
+        </member>
+        <member name="M:Newtonsoft.Json.Linq.JToken.op_Implicit(System.Byte[])~Newtonsoft.Json.Linq.JToken">
+            <summary>
+            Performs an implicit conversion from <see cref="T:System.Byte[]"/> to <see cref="T:Newtonsoft.Json.Linq.JToken"/>.
+            </summary>
+            <param name="value">The value to create a <see cref="T:Newtonsoft.Json.Linq.JValue"/> from.</param>
+            <returns>The <see cref="T:Newtonsoft.Json.Linq.JValue"/> initialized with the specified value.</returns>
+        </member>
+        <member name="M:Newtonsoft.Json.Linq.JToken.op_Implicit(System.Uri)~Newtonsoft.Json.Linq.JToken">
+            <summary>
+            Performs an implicit conversion from <see cref="T:System.Uri"/> to <see cref="T:Newtonsoft.Json.Linq.JToken"/>.
+            </summary>
+            <param name="value">The value to create a <see cref="T:Newtonsoft.Json.Linq.JValue"/> from.</param>
+            <returns>The <see cref="T:Newtonsoft.Json.Linq.JValue"/> initialized with the specified value.</returns>
+        </member>
+        <member name="M:Newtonsoft.Json.Linq.JToken.op_Implicit(System.TimeSpan)~Newtonsoft.Json.Linq.JToken">
+            <summary>
+            Performs an implicit conversion from <see cref="T:System.TimeSpan"/> to <see cref="T:Newtonsoft.Json.Linq.JToken"/>.
+            </summary>
+            <param name="value">The value to create a <see cref="T:Newtonsoft.Json.Linq.JValue"/> from.</param>
+            <returns>The <see cref="T:Newtonsoft.Json.Linq.JValue"/> initialized with the specified value.</returns>
+        </member>
+        <member name="M:Newtonsoft.Json.Linq.JToken.op_Implicit(System.Nullable{System.TimeSpan})~Newtonsoft.Json.Linq.JToken">
+            <summary>
+            Performs an implicit conversion from <see cref="T:System.Nullable`1"/> to <see cref="T:Newtonsoft.Json.Linq.JToken"/>.
+            </summary>
+            <param name="value">The value to create a <see cref="T:Newtonsoft.Json.Linq.JValue"/> from.</param>
+            <returns>The <see cref="T:Newtonsoft.Json.Linq.JValue"/> initialized with the specified value.</returns>
+        </member>
+        <member name="M:Newtonsoft.Json.Linq.JToken.op_Implicit(System.Guid)~Newtonsoft.Json.Linq.JToken">
+            <summary>
+            Performs an implicit conversion from <see cref="T:System.Guid"/> to <see cref="T:Newtonsoft.Json.Linq.JToken"/>.
+            </summary>
+            <param name="value">The value to create a <see cref="T:Newtonsoft.Json.Linq.JValue"/> from.</param>
+            <returns>The <see cref="T:Newtonsoft.Json.Linq.JValue"/> initialized with the specified value.</returns>
+        </member>
+        <member name="M:Newtonsoft.Json.Linq.JToken.op_Implicit(System.Nullable{System.Guid})~Newtonsoft.Json.Linq.JToken">
+            <summary>
+            Performs an implicit conversion from <see cref="T:System.Nullable`1"/> to <see cref="T:Newtonsoft.Json.Linq.JToken"/>.
+            </summary>
+            <param name="value">The value to create a <see cref="T:Newtonsoft.Json.Linq.JValue"/> from.</param>
+            <returns>The <see cref="T:Newtonsoft.Json.Linq.JValue"/> initialized with the specified value.</returns>
+        </member>
+        <member name="M:Newtonsoft.Json.Linq.JToken.CreateReader">
+            <summary>
+            Creates an <see cref="T:Newtonsoft.Json.JsonReader"/> for this token.
+            </summary>
+            <returns>An <see cref="T:Newtonsoft.Json.JsonReader"/> that can be used to read this token and its descendants.</returns>
+        </member>
+        <member name="M:Newtonsoft.Json.Linq.JToken.FromObject(System.Object)">
+            <summary>
+            Creates a <see cref="T:Newtonsoft.Json.Linq.JToken"/> from an object.
+            </summary>
+            <param name="o">The object that will be used to create <see cref="T:Newtonsoft.Json.Linq.JToken"/>.</param>
+            <returns>A <see cref="T:Newtonsoft.Json.Linq.JToken"/> with the value of the specified object</returns>
+        </member>
+        <member name="M:Newtonsoft.Json.Linq.JToken.FromObject(System.Object,Newtonsoft.Json.JsonSerializer)">
+            <summary>
+            Creates a <see cref="T:Newtonsoft.Json.Linq.JToken"/> from an object using the specified <see cref="T:Newtonsoft.Json.JsonSerializer"/>.
+            </summary>
+            <param name="o">The object that will be used to create <see cref="T:Newtonsoft.Json.Linq.JToken"/>.</param>
+            <param name="jsonSerializer">The <see cref="T:Newtonsoft.Json.JsonSerializer"/> that will be used when reading the object.</param>
+            <returns>A <see cref="T:Newtonsoft.Json.Linq.JToken"/> with the value of the specified object</returns>
+        </member>
+        <member name="M:Newtonsoft.Json.Linq.JToken.ToObject``1">
+            <summary>
+            Creates the specified .NET type from the <see cref="T:Newtonsoft.Json.Linq.JToken"/>.
+            </summary>
+            <typeparam name="T">The object type that the token will be deserialized to.</typeparam>
+            <returns>The new object created from the JSON value.</returns>
+        </member>
+        <member name="M:Newtonsoft.Json.Linq.JToken.ToObject(System.Type)">
+            <summary>
+            Creates the specified .NET type from the <see cref="T:Newtonsoft.Json.Linq.JToken"/>.
+            </summary>
+            <param name="objectType">The object type that the token will be deserialized to.</param>
+            <returns>The new object created from the JSON value.</returns>
+        </member>
+        <member name="M:Newtonsoft.Json.Linq.JToken.ToObject``1(Newtonsoft.Json.JsonSerializer)">
+            <summary>
+            Creates the specified .NET type from the <see cref="T:Newtonsoft.Json.Linq.JToken"/> using the specified <see cref="T:Newtonsoft.Json.JsonSerializer"/>.
+            </summary>
+            <typeparam name="T">The object type that the token will be deserialized to.</typeparam>
+            <param name="jsonSerializer">The <see cref="T:Newtonsoft.Json.JsonSerializer"/> that will be used when creating the object.</param>
+            <returns>The new object created from the JSON value.</returns>
+        </member>
+        <member name="M:Newtonsoft.Json.Linq.JToken.ToObject(System.Type,Newtonsoft.Json.JsonSerializer)">
+            <summary>
+            Creates the specified .NET type from the <see cref="T:Newtonsoft.Json.Linq.JToken"/> using the specified <see cref="T:Newtonsoft.Json.JsonSerializer"/>.
+            </summary>
+            <param name="objectType">The object type that the token will be deserialized to.</param>
+            <param name="jsonSerializer">The <see cref="T:Newtonsoft.Json.JsonSerializer"/> that will be used when creating the object.</param>
+            <returns>The new object created from the JSON value.</returns>
+        </member>
+        <member name="M:Newtonsoft.Json.Linq.JToken.ReadFrom(Newtonsoft.Json.JsonReader)">
+            <summary>
+            Creates a <see cref="T:Newtonsoft.Json.Linq.JToken"/> from a <see cref="T:Newtonsoft.Json.JsonReader"/>.
+            </summary>
+            <param name="reader">An <see cref="T:Newtonsoft.Json.JsonReader"/> positioned at the token to read into this <see cref="T:Newtonsoft.Json.Linq.JToken"/>.</param>
+            <returns>
+            An <see cref="T:Newtonsoft.Json.Linq.JToken"/> that contains the token and its descendant tokens
+            that were read from the reader. The runtime type of the token is determined
+            by the token type of the first token encountered in the reader.
+            </returns>
+        </member>
+        <member name="M:Newtonsoft.Json.Linq.JToken.Parse(System.String)">
+            <summary>
+            Load a <see cref="T:Newtonsoft.Json.Linq.JToken"/> from a string that contains JSON.
+            </summary>
+            <param name="json">A <see cref="T:System.String"/> that contains JSON.</param>
+            <returns>A <see cref="T:Newtonsoft.Json.Linq.JToken"/> populated from the string that contains JSON.</returns>
+        </member>
+        <member name="M:Newtonsoft.Json.Linq.JToken.Load(Newtonsoft.Json.JsonReader)">
+            <summary>
+            Creates a <see cref="T:Newtonsoft.Json.Linq.JToken"/> from a <see cref="T:Newtonsoft.Json.JsonReader"/>.
+            </summary>
+            <param name="reader">An <see cref="T:Newtonsoft.Json.JsonReader"/> positioned at the token to read into this <see cref="T:Newtonsoft.Json.Linq.JToken"/>.</param>
+            <returns>
+            An <see cref="T:Newtonsoft.Json.Linq.JToken"/> that contains the token and its descendant tokens
+            that were read from the reader. The runtime type of the token is determined
+            by the token type of the first token encountered in the reader.
+            </returns>
+        </member>
+        <member name="M:Newtonsoft.Json.Linq.JToken.SelectToken(System.String)">
+            <summary>
+            Selects the token that matches the object path.
+            </summary>
+            <param name="path">
+            The object path from the current <see cref="T:Newtonsoft.Json.Linq.JToken"/> to the <see cref="T:Newtonsoft.Json.Linq.JToken"/>
+            to be returned. This must be a string of property names or array indexes separated
+            by periods, such as <code>Tables[0].DefaultView[0].Price</code> in C# or
+            <code>Tables(0).DefaultView(0).Price</code> in Visual Basic.
+            </param>
+            <returns>The <see cref="T:Newtonsoft.Json.Linq.JToken"/> that matches the object path or a null reference if no matching token is found.</returns>
+        </member>
+        <member name="M:Newtonsoft.Json.Linq.JToken.SelectToken(System.String,System.Boolean)">
+            <summary>
+            Selects the token that matches the object path.
+            </summary>
+            <param name="path">
+            The object path from the current <see cref="T:Newtonsoft.Json.Linq.JToken"/> to the <see cref="T:Newtonsoft.Json.Linq.JToken"/>
+            to be returned. This must be a string of property names or array indexes separated
+            by periods, such as <code>Tables[0].DefaultView[0].Price</code> in C# or
+            <code>Tables(0).DefaultView(0).Price</code> in Visual Basic.
+            </param>
+            <param name="errorWhenNoMatch">A flag to indicate whether an error should be thrown if no token is found.</param>
+            <returns>The <see cref="T:Newtonsoft.Json.Linq.JToken"/> that matches the object path.</returns>
+        </member>
+        <member name="M:Newtonsoft.Json.Linq.JToken.DeepClone">
+            <summary>
+            Creates a new instance of the <see cref="T:Newtonsoft.Json.Linq.JToken"/>. All child tokens are recursively cloned.
+            </summary>
+            <returns>A new instance of the <see cref="T:Newtonsoft.Json.Linq.JToken"/>.</returns>
+        </member>
+        <member name="P:Newtonsoft.Json.Linq.JToken.EqualityComparer">
+            <summary>
+            Gets a comparer that can compare two tokens for value equality.
+            </summary>
+            <value>A <see cref="T:Newtonsoft.Json.Linq.JTokenEqualityComparer"/> that can compare two nodes for value equality.</value>
+        </member>
+        <member name="P:Newtonsoft.Json.Linq.JToken.Parent">
+            <summary>
+            Gets or sets the parent.
+            </summary>
+            <value>The parent.</value>
+        </member>
+        <member name="P:Newtonsoft.Json.Linq.JToken.Root">
+            <summary>
+            Gets the root <see cref="T:Newtonsoft.Json.Linq.JToken"/> of this <see cref="T:Newtonsoft.Json.Linq.JToken"/>.
+            </summary>
+            <value>The root <see cref="T:Newtonsoft.Json.Linq.JToken"/> of this <see cref="T:Newtonsoft.Json.Linq.JToken"/>.</value>
+        </member>
+        <member name="P:Newtonsoft.Json.Linq.JToken.Type">
+            <summary>
+            Gets the node type for this <see cref="T:Newtonsoft.Json.Linq.JToken"/>.
+            </summary>
+            <value>The type.</value>
+        </member>
+        <member name="P:Newtonsoft.Json.Linq.JToken.HasValues">
+            <summary>
+            Gets a value indicating whether this token has childen tokens.
+            </summary>
+            <value>
+            	<c>true</c> if this token has child values; otherwise, <c>false</c>.
+            </value>
+        </member>
+        <member name="P:Newtonsoft.Json.Linq.JToken.Next">
+            <summary>
+            Gets the next sibling token of this node.
+            </summary>
+            <value>The <see cref="T:Newtonsoft.Json.Linq.JToken"/> that contains the next sibling token.</value>
+        </member>
+        <member name="P:Newtonsoft.Json.Linq.JToken.Previous">
+            <summary>
+            Gets the previous sibling token of this node.
+            </summary>
+            <value>The <see cref="T:Newtonsoft.Json.Linq.JToken"/> that contains the previous sibling token.</value>
+        </member>
+        <member name="P:Newtonsoft.Json.Linq.JToken.Path">
+            <summary>
+            Gets the path of the JSON token. 
+            </summary>
+        </member>
+        <member name="P:Newtonsoft.Json.Linq.JToken.Item(System.Object)">
+            <summary>
+            Gets the <see cref="T:Newtonsoft.Json.Linq.JToken"/> with the specified key.
+            </summary>
+            <value>The <see cref="T:Newtonsoft.Json.Linq.JToken"/> with the specified key.</value>
+        </member>
+        <member name="P:Newtonsoft.Json.Linq.JToken.First">
+            <summary>
+            Get the first child token of this token.
+            </summary>
+            <value>A <see cref="T:Newtonsoft.Json.Linq.JToken"/> containing the first child token of the <see cref="T:Newtonsoft.Json.Linq.JToken"/>.</value>
+        </member>
+        <member name="P:Newtonsoft.Json.Linq.JToken.Last">
+            <summary>
+            Get the last child token of this token.
+            </summary>
+            <value>A <see cref="T:Newtonsoft.Json.Linq.JToken"/> containing the last child token of the <see cref="T:Newtonsoft.Json.Linq.JToken"/>.</value>
+        </member>
+        <member name="M:Newtonsoft.Json.Linq.JValue.#ctor(Newtonsoft.Json.Linq.JValue)">
+            <summary>
+            Initializes a new instance of the <see cref="T:Newtonsoft.Json.Linq.JValue"/> class from another <see cref="T:Newtonsoft.Json.Linq.JValue"/> object.
+            </summary>
+            <param name="other">A <see cref="T:Newtonsoft.Json.Linq.JValue"/> object to copy from.</param>
+        </member>
+        <member name="M:Newtonsoft.Json.Linq.JValue.#ctor(System.Int64)">
+            <summary>
+            Initializes a new instance of the <see cref="T:Newtonsoft.Json.Linq.JValue"/> class with the given value.
+            </summary>
+            <param name="value">The value.</param>
+        </member>
+        <member name="M:Newtonsoft.Json.Linq.JValue.#ctor(System.Char)">
+            <summary>
+            Initializes a new instance of the <see cref="T:Newtonsoft.Json.Linq.JValue"/> class with the given value.
+            </summary>
+            <param name="value">The value.</param>
+        </member>
+        <member name="M:Newtonsoft.Json.Linq.JValue.#ctor(System.UInt64)">
+            <summary>
+            Initializes a new instance of the <see cref="T:Newtonsoft.Json.Linq.JValue"/> class with the given value.
+            </summary>
+            <param name="value">The value.</param>
+        </member>
+        <member name="M:Newtonsoft.Json.Linq.JValue.#ctor(System.Double)">
+            <summary>
+            Initializes a new instance of the <see cref="T:Newtonsoft.Json.Linq.JValue"/> class with the given value.
+            </summary>
+            <param name="value">The value.</param>
+        </member>
+        <member name="M:Newtonsoft.Json.Linq.JValue.#ctor(System.Single)">
+            <summary>
+            Initializes a new instance of the <see cref="T:Newtonsoft.Json.Linq.JValue"/> class with the given value.
+            </summary>
+            <param name="value">The value.</param>
+        </member>
+        <member name="M:Newtonsoft.Json.Linq.JValue.#ctor(System.DateTime)">
+            <summary>
+            Initializes a new instance of the <see cref="T:Newtonsoft.Json.Linq.JValue"/> class with the given value.
+            </summary>
+            <param name="value">The value.</param>
+        </member>
+        <member name="M:Newtonsoft.Json.Linq.JValue.#ctor(System.Boolean)">
+            <summary>
+            Initializes a new instance of the <see cref="T:Newtonsoft.Json.Linq.JValue"/> class with the given value.
+            </summary>
+            <param name="value">The value.</param>
+        </member>
+        <member name="M:Newtonsoft.Json.Linq.JValue.#ctor(System.String)">
+            <summary>
+            Initializes a new instance of the <see cref="T:Newtonsoft.Json.Linq.JValue"/> class with the given value.
+            </summary>
+            <param name="value">The value.</param>
+        </member>
+        <member name="M:Newtonsoft.Json.Linq.JValue.#ctor(System.Guid)">
+            <summary>
+            Initializes a new instance of the <see cref="T:Newtonsoft.Json.Linq.JValue"/> class with the given value.
+            </summary>
+            <param name="value">The value.</param>
+        </member>
+        <member name="M:Newtonsoft.Json.Linq.JValue.#ctor(System.Uri)">
+            <summary>
+            Initializes a new instance of the <see cref="T:Newtonsoft.Json.Linq.JValue"/> class with the given value.
+            </summary>
+            <param name="value">The value.</param>
+        </member>
+        <member name="M:Newtonsoft.Json.Linq.JValue.#ctor(System.TimeSpan)">
+            <summary>
+            Initializes a new instance of the <see cref="T:Newtonsoft.Json.Linq.JValue"/> class with the given value.
+            </summary>
+            <param name="value">The value.</param>
+        </member>
+        <member name="M:Newtonsoft.Json.Linq.JValue.#ctor(System.Object)">
+            <summary>
+            Initializes a new instance of the <see cref="T:Newtonsoft.Json.Linq.JValue"/> class with the given value.
+            </summary>
+            <param name="value">The value.</param>
+        </member>
+        <member name="M:Newtonsoft.Json.Linq.JValue.CreateComment(System.String)">
+            <summary>
+            Creates a <see cref="T:Newtonsoft.Json.Linq.JValue"/> comment with the given value.
+            </summary>
+            <param name="value">The value.</param>
+            <returns>A <see cref="T:Newtonsoft.Json.Linq.JValue"/> comment with the given value.</returns>
+        </member>
+        <member name="M:Newtonsoft.Json.Linq.JValue.CreateString(System.String)">
+            <summary>
+            Creates a <see cref="T:Newtonsoft.Json.Linq.JValue"/> string with the given value.
+            </summary>
+            <param name="value">The value.</param>
+            <returns>A <see cref="T:Newtonsoft.Json.Linq.JValue"/> string with the given value.</returns>
+        </member>
+        <member name="M:Newtonsoft.Json.Linq.JValue.WriteTo(Newtonsoft.Json.JsonWriter,Newtonsoft.Json.JsonConverter[])">
+            <summary>
+            Writes this token to a <see cref="T:Newtonsoft.Json.JsonWriter"/>.
+            </summary>
+            <param name="writer">A <see cref="T:Newtonsoft.Json.JsonWriter"/> into which this method will write.</param>
+            <param name="converters">A collection of <see cref="T:Newtonsoft.Json.JsonConverter"/> which will be used when writing the token.</param>
+        </member>
+        <member name="M:Newtonsoft.Json.Linq.JValue.Equals(Newtonsoft.Json.Linq.JValue)">
+            <summary>
+            Indicates whether the current object is equal to another object of the same type.
+            </summary>
+            <returns>
+            true if the current object is equal to the <paramref name="other"/> parameter; otherwise, false.
+            </returns>
+            <param name="other">An object to compare with this object.</param>
+        </member>
+        <member name="M:Newtonsoft.Json.Linq.JValue.Equals(System.Object)">
+            <summary>
+            Determines whether the specified <see cref="T:System.Object"/> is equal to the current <see cref="T:System.Object"/>.
+            </summary>
+            <param name="obj">The <see cref="T:System.Object"/> to compare with the current <see cref="T:System.Object"/>.</param>
+            <returns>
+            true if the specified <see cref="T:System.Object"/> is equal to the current <see cref="T:System.Object"/>; otherwise, false.
+            </returns>
+            <exception cref="T:System.NullReferenceException">
+            The <paramref name="obj"/> parameter is null.
+            </exception>
+        </member>
+        <member name="M:Newtonsoft.Json.Linq.JValue.GetHashCode">
+            <summary>
+            Serves as a hash function for a particular type.
+            </summary>
+            <returns>
+            A hash code for the current <see cref="T:System.Object"/>.
+            </returns>
+        </member>
+        <member name="M:Newtonsoft.Json.Linq.JValue.ToString">
+            <summary>
+            Returns a <see cref="T:System.String"/> that represents this instance.
+            </summary>
+            <returns>
+            A <see cref="T:System.String"/> that represents this instance.
+            </returns>
+        </member>
+        <member name="M:Newtonsoft.Json.Linq.JValue.ToString(System.String)">
+            <summary>
+            Returns a <see cref="T:System.String"/> that represents this instance.
+            </summary>
+            <param name="format">The format.</param>
+            <returns>
+            A <see cref="T:System.String"/> that represents this instance.
+            </returns>
+        </member>
+        <member name="M:Newtonsoft.Json.Linq.JValue.ToString(System.IFormatProvider)">
+            <summary>
+            Returns a <see cref="T:System.String"/> that represents this instance.
+            </summary>
+            <param name="formatProvider">The format provider.</param>
+            <returns>
+            A <see cref="T:System.String"/> that represents this instance.
+            </returns>
+        </member>
+        <member name="M:Newtonsoft.Json.Linq.JValue.ToString(System.String,System.IFormatProvider)">
+            <summary>
+            Returns a <see cref="T:System.String"/> that represents this instance.
+            </summary>
+            <param name="format">The format.</param>
+            <param name="formatProvider">The format provider.</param>
+            <returns>
+            A <see cref="T:System.String"/> that represents this instance.
+            </returns>
+        </member>
+        <member name="M:Newtonsoft.Json.Linq.JValue.CompareTo(Newtonsoft.Json.Linq.JValue)">
+            <summary>
+            Compares the current instance with another object of the same type and returns an integer that indicates whether the current instance precedes, follows, or occurs in the same position in the sort order as the other object.
+            </summary>
+            <param name="obj">An object to compare with this instance.</param>
+            <returns>
+            A 32-bit signed integer that indicates the relative order of the objects being compared. The return value has these meanings:
+            Value
+            Meaning
+            Less than zero
+            This instance is less than <paramref name="obj"/>.
+            Zero
+            This instance is equal to <paramref name="obj"/>.
+            Greater than zero
+            This instance is greater than <paramref name="obj"/>.
+            </returns>
+            <exception cref="T:System.ArgumentException">
+            	<paramref name="obj"/> is not the same type as this instance.
+            </exception>
+        </member>
+        <member name="P:Newtonsoft.Json.Linq.JValue.HasValues">
+            <summary>
+            Gets a value indicating whether this token has childen tokens.
+            </summary>
+            <value>
+            	<c>true</c> if this token has child values; otherwise, <c>false</c>.
+            </value>
+        </member>
+        <member name="P:Newtonsoft.Json.Linq.JValue.Type">
+            <summary>
+            Gets the node type for this <see cref="T:Newtonsoft.Json.Linq.JToken"/>.
+            </summary>
+            <value>The type.</value>
+        </member>
+        <member name="P:Newtonsoft.Json.Linq.JValue.Value">
+            <summary>
+            Gets or sets the underlying token value.
+            </summary>
+            <value>The underlying token value.</value>
+        </member>
+        <member name="M:Newtonsoft.Json.Linq.JRaw.#ctor(Newtonsoft.Json.Linq.JRaw)">
+            <summary>
+            Initializes a new instance of the <see cref="T:Newtonsoft.Json.Linq.JRaw"/> class from another <see cref="T:Newtonsoft.Json.Linq.JRaw"/> object.
+            </summary>
+            <param name="other">A <see cref="T:Newtonsoft.Json.Linq.JRaw"/> object to copy from.</param>
+        </member>
+        <member name="M:Newtonsoft.Json.Linq.JRaw.#ctor(System.Object)">
+            <summary>
+            Initializes a new instance of the <see cref="T:Newtonsoft.Json.Linq.JRaw"/> class.
+            </summary>
+            <param name="rawJson">The raw json.</param>
+        </member>
+        <member name="M:Newtonsoft.Json.Linq.JRaw.Create(Newtonsoft.Json.JsonReader)">
+            <summary>
+            Creates an instance of <see cref="T:Newtonsoft.Json.Linq.JRaw"/> with the content of the reader's current token.
+            </summary>
+            <param name="reader">The reader.</param>
+            <returns>An instance of <see cref="T:Newtonsoft.Json.Linq.JRaw"/> with the content of the reader's current token.</returns>
+        </member>
+        <member name="T:Newtonsoft.Json.Required">
+            <summary>
+            Indicating whether a property is required.
+            </summary>
+        </member>
+        <member name="F:Newtonsoft.Json.Required.Default">
+            <summary>
+            The property is not required. The default state.
+            </summary>
+        </member>
+        <member name="F:Newtonsoft.Json.Required.AllowNull">
+            <summary>
+            The property must be defined in JSON but can be a null value.
+            </summary>
+        </member>
+        <member name="F:Newtonsoft.Json.Required.Always">
+            <summary>
+            The property must be defined in JSON and cannot be a null value.
+            </summary>
+        </member>
+        <member name="T:Newtonsoft.Json.Serialization.IReferenceResolver">
+            <summary>
+            Used to resolve references when serializing and deserializing JSON by the <see cref="T:Newtonsoft.Json.JsonSerializer"/>.
+            </summary>
+        </member>
+        <member name="M:Newtonsoft.Json.Serialization.IReferenceResolver.ResolveReference(System.Object,System.String)">
+            <summary>
+            Resolves a reference to its object.
+            </summary>
+            <param name="context">The serialization context.</param>
+            <param name="reference">The reference to resolve.</param>
+            <returns>The object that</returns>
+        </member>
+        <member name="M:Newtonsoft.Json.Serialization.IReferenceResolver.GetReference(System.Object,System.Object)">
+            <summary>
+            Gets the reference for the sepecified object.
+            </summary>
+            <param name="context">The serialization context.</param>
+            <param name="value">The object to get a reference for.</param>
+            <returns>The reference to the object.</returns>
+        </member>
+        <member name="M:Newtonsoft.Json.Serialization.IReferenceResolver.IsReferenced(System.Object,System.Object)">
+            <summary>
+            Determines whether the specified object is referenced.
+            </summary>
+            <param name="context">The serialization context.</param>
+            <param name="value">The object to test for a reference.</param>
+            <returns>
+            	<c>true</c> if the specified object is referenced; otherwise, <c>false</c>.
+            </returns>
+        </member>
+        <member name="M:Newtonsoft.Json.Serialization.IReferenceResolver.AddReference(System.Object,System.String,System.Object)">
+            <summary>
+            Adds a reference to the specified object.
+            </summary>
+            <param name="context">The serialization context.</param>
+            <param name="reference">The reference.</param>
+            <param name="value">The object to reference.</param>
+        </member>
+        <member name="T:Newtonsoft.Json.PreserveReferencesHandling">
+            <summary>
+            Specifies reference handling options for the <see cref="T:Newtonsoft.Json.JsonSerializer"/>.
+            Note that references cannot be preserved when a value is set via a non-default constructor such as types that implement ISerializable.
+            </summary>
+            <example>
+              <code lang="cs" source="..\Src\Newtonsoft.Json.Tests\Documentation\SerializationTests.cs" region="PreservingObjectReferencesOn" title="Preserve Object References"/>       
+            </example>
+        </member>
+        <member name="F:Newtonsoft.Json.PreserveReferencesHandling.None">
+            <summary>
+            Do not preserve references when serializing types.
+            </summary>
+        </member>
+        <member name="F:Newtonsoft.Json.PreserveReferencesHandling.Objects">
+            <summary>
+            Preserve references when serializing into a JSON object structure.
+            </summary>
+        </member>
+        <member name="F:Newtonsoft.Json.PreserveReferencesHandling.Arrays">
+            <summary>
+            Preserve references when serializing into a JSON array structure.
+            </summary>
+        </member>
+        <member name="F:Newtonsoft.Json.PreserveReferencesHandling.All">
+            <summary>
+            Preserve references when serializing.
+            </summary>
+        </member>
+        <member name="T:Newtonsoft.Json.JsonArrayAttribute">
+            <summary>
+            Instructs the <see cref="T:Newtonsoft.Json.JsonSerializer"/> how to serialize the collection.
+            </summary>
+        </member>
+        <member name="M:Newtonsoft.Json.JsonArrayAttribute.#ctor">
+            <summary>
+            Initializes a new instance of the <see cref="T:Newtonsoft.Json.JsonArrayAttribute"/> class.
+            </summary>
+        </member>
+        <member name="M:Newtonsoft.Json.JsonArrayAttribute.#ctor(System.Boolean)">
+            <summary>
+            Initializes a new instance of the <see cref="T:Newtonsoft.Json.JsonObjectAttribute"/> class with a flag indicating whether the array can contain null items
+            </summary>
+            <param name="allowNullItems">A flag indicating whether the array can contain null items.</param>
+        </member>
+        <member name="M:Newtonsoft.Json.JsonArrayAttribute.#ctor(System.String)">
+            <summary>
+            Initializes a new instance of the <see cref="T:Newtonsoft.Json.JsonArrayAttribute"/> class with the specified container Id.
+            </summary>
+            <param name="id">The container Id.</param>
+        </member>
+        <member name="P:Newtonsoft.Json.JsonArrayAttribute.AllowNullItems">
+            <summary>
+            Gets or sets a value indicating whether null items are allowed in the collection.
+            </summary>
+            <value><c>true</c> if null items are allowed in the collection; otherwise, <c>false</c>.</value>
+        </member>
+        <member name="T:Newtonsoft.Json.DefaultValueHandling">
+            <summary>
+            Specifies default value handling options for the <see cref="T:Newtonsoft.Json.JsonSerializer"/>.
+            </summary>
+            <example>
+              <code lang="cs" source="..\Src\Newtonsoft.Json.Tests\Documentation\SerializationTests.cs" region="ReducingSerializedJsonSizeDefaultValueHandlingObject" title="DefaultValueHandling Class"/>
+              <code lang="cs" source="..\Src\Newtonsoft.Json.Tests\Documentation\SerializationTests.cs" region="ReducingSerializedJsonSizeDefaultValueHandlingExample" title="DefaultValueHandling Ignore Example"/>
+            </example>
+        </member>
+        <member name="F:Newtonsoft.Json.DefaultValueHandling.Include">
+            <summary>
+            Include members where the member value is the same as the member's default value when serializing objects.
+            Included members are written to JSON. Has no effect when deserializing.
+            </summary>
+        </member>
+        <member name="F:Newtonsoft.Json.DefaultValueHandling.Ignore">
+            <summary>
+            Ignore members where the member value is the same as the member's default value when serializing objects
+            so that is is not written to JSON.
+            This option will ignore all default values (e.g. <c>null</c> for objects and nullable typesl; <c>0</c> for integers,
+            decimals and floating point numbers; and <c>false</c> for booleans). The default value ignored can be changed by
+            placing the <see cref="T:System.ComponentModel.DefaultValueAttribute"/> on the property.
+            </summary>
+        </member>
+        <member name="F:Newtonsoft.Json.DefaultValueHandling.Populate">
+            <summary>
+            Members with a default value but no JSON will be set to their default value when deserializing.
+            </summary>
+        </member>
+        <member name="F:Newtonsoft.Json.DefaultValueHandling.IgnoreAndPopulate">
+            <summary>
+            Ignore members where the member value is the same as the member's default value when serializing objects
+            and sets members to their default value when deserializing.
+            </summary>
+        </member>
+        <member name="T:Newtonsoft.Json.JsonConverterAttribute">
+            <summary>
+            Instructs the <see cref="T:Newtonsoft.Json.JsonSerializer"/> to use the specified <see cref="T:Newtonsoft.Json.JsonConverter"/> when serializing the member or class.
+            </summary>
+        </member>
+        <member name="M:Newtonsoft.Json.JsonConverterAttribute.#ctor(System.Type)">
+            <summary>
+            Initializes a new instance of the <see cref="T:Newtonsoft.Json.JsonConverterAttribute"/> class.
+            </summary>
+            <param name="converterType">Type of the converter.</param>
+        </member>
+        <member name="P:Newtonsoft.Json.JsonConverterAttribute.ConverterType">
+            <summary>
+            Gets the type of the converter.
+            </summary>
+            <value>The type of the converter.</value>
+        </member>
+        <member name="T:Newtonsoft.Json.JsonObjectAttribute">
+            <summary>
+            Instructs the <see cref="T:Newtonsoft.Json.JsonSerializer"/> how to serialize the object.
+            </summary>
+        </member>
+        <member name="M:Newtonsoft.Json.JsonObjectAttribute.#ctor">
+            <summary>
+            Initializes a new instance of the <see cref="T:Newtonsoft.Json.JsonObjectAttribute"/> class.
+            </summary>
+        </member>
+        <member name="M:Newtonsoft.Json.JsonObjectAttribute.#ctor(Newtonsoft.Json.MemberSerialization)">
+            <summary>
+            Initializes a new instance of the <see cref="T:Newtonsoft.Json.JsonObjectAttribute"/> class with the specified member serialization.
+            </summary>
+            <param name="memberSerialization">The member serialization.</param>
+        </member>
+        <member name="M:Newtonsoft.Json.JsonObjectAttribute.#ctor(System.String)">
+            <summary>
+            Initializes a new instance of the <see cref="T:Newtonsoft.Json.JsonObjectAttribute"/> class with the specified container Id.
+            </summary>
+            <param name="id">The container Id.</param>
+        </member>
+        <member name="P:Newtonsoft.Json.JsonObjectAttribute.MemberSerialization">
+            <summary>
+            Gets or sets the member serialization.
+            </summary>
+            <value>The member serialization.</value>
+        </member>
+        <member name="P:Newtonsoft.Json.JsonObjectAttribute.ItemRequired">
+            <summary>
+            Gets or sets a value that indicates whether the object's properties are required.
+            </summary>
+            <value>
+            	A value indicating whether the object's properties are required.
+            </value>
+        </member>
+        <member name="T:Newtonsoft.Json.JsonSerializerSettings">
+            <summary>
+            Specifies the settings on a <see cref="T:Newtonsoft.Json.JsonSerializer"/> object.
+            </summary>
+        </member>
+        <member name="M:Newtonsoft.Json.JsonSerializerSettings.#ctor">
+            <summary>
+            Initializes a new instance of the <see cref="T:Newtonsoft.Json.JsonSerializerSettings"/> class.
+            </summary>
+        </member>
+        <member name="P:Newtonsoft.Json.JsonSerializerSettings.ReferenceLoopHandling">
+            <summary>
+            Gets or sets how reference loops (e.g. a class referencing itself) is handled.
+            </summary>
+            <value>Reference loop handling.</value>
+        </member>
+        <member name="P:Newtonsoft.Json.JsonSerializerSettings.MissingMemberHandling">
+            <summary>
+            Gets or sets how missing members (e.g. JSON contains a property that isn't a member on the object) are handled during deserialization.
+            </summary>
+            <value>Missing member handling.</value>
+        </member>
+        <member name="P:Newtonsoft.Json.JsonSerializerSettings.ObjectCreationHandling">
+            <summary>
+            Gets or sets how objects are created during deserialization.
+            </summary>
+            <value>The object creation handling.</value>
+        </member>
+        <member name="P:Newtonsoft.Json.JsonSerializerSettings.NullValueHandling">
+            <summary>
+            Gets or sets how null values are handled during serialization and deserialization.
+            </summary>
+            <value>Null value handling.</value>
+        </member>
+        <member name="P:Newtonsoft.Json.JsonSerializerSettings.DefaultValueHandling">
+            <summary>
+            Gets or sets how null default are handled during serialization and deserialization.
+            </summary>
+            <value>The default value handling.</value>
+        </member>
+        <member name="P:Newtonsoft.Json.JsonSerializerSettings.Converters">
+            <summary>
+            Gets or sets a collection <see cref="T:Newtonsoft.Json.JsonConverter"/> that will be used during serialization.
+            </summary>
+            <value>The converters.</value>
+        </member>
+        <member name="P:Newtonsoft.Json.JsonSerializerSettings.PreserveReferencesHandling">
+            <summary>
+            Gets or sets how object references are preserved by the serializer.
+            </summary>
+            <value>The preserve references handling.</value>
+        </member>
+        <member name="P:Newtonsoft.Json.JsonSerializerSettings.TypeNameHandling">
+            <summary>
+            Gets or sets how type name writing and reading is handled by the serializer.
+            </summary>
+            <value>The type name handling.</value>
+        </member>
+        <member name="P:Newtonsoft.Json.JsonSerializerSettings.TypeNameAssemblyFormat">
+            <summary>
+            Gets or sets how a type name assembly is written and resolved by the serializer.
+            </summary>
+            <value>The type name assembly format.</value>
+        </member>
+        <member name="P:Newtonsoft.Json.JsonSerializerSettings.ConstructorHandling">
+            <summary>
+            Gets or sets how constructors are used during deserialization.
+            </summary>
+            <value>The constructor handling.</value>
+        </member>
+        <member name="P:Newtonsoft.Json.JsonSerializerSettings.ContractResolver">
+            <summary>
+            Gets or sets the contract resolver used by the serializer when
+            serializing .NET objects to JSON and vice versa.
+            </summary>
+            <value>The contract resolver.</value>
+        </member>
+        <member name="P:Newtonsoft.Json.JsonSerializerSettings.ReferenceResolver">
+            <summary>
+            Gets or sets the <see cref="T:Newtonsoft.Json.Serialization.IReferenceResolver"/> used by the serializer when resolving references.
+            </summary>
+            <value>The reference resolver.</value>
+        </member>
+        <member name="P:Newtonsoft.Json.JsonSerializerSettings.TraceWriter">
+            <summary>
+            Gets or sets the <see cref="T:Newtonsoft.Json.Serialization.ITraceWriter"/> used by the serializer when writing trace messages.
+            </summary>
+            <value>The trace writer.</value>
+        </member>
+        <member name="P:Newtonsoft.Json.JsonSerializerSettings.Binder">
+            <summary>
+            Gets or sets the <see cref="T:System.Runtime.Serialization.SerializationBinder"/> used by the serializer when resolving type names.
+            </summary>
+            <value>The binder.</value>
+        </member>
+        <member name="P:Newtonsoft.Json.JsonSerializerSettings.Error">
+            <summary>
+            Gets or sets the error handler called during serialization and deserialization.
+            </summary>
+            <value>The error handler called during serialization and deserialization.</value>
+        </member>
+        <member name="P:Newtonsoft.Json.JsonSerializerSettings.Context">
+            <summary>
+            Gets or sets the <see cref="T:System.Runtime.Serialization.StreamingContext"/> used by the serializer when invoking serialization callback methods.
+            </summary>
+            <value>The context.</value>
+        </member>
+        <member name="P:Newtonsoft.Json.JsonSerializerSettings.DateFormatString">
+            <summary>
+            Get or set how <see cref="T:System.DateTime"/> and <see cref="T:System.DateTimeOffset"/> values are formatting when writing JSON text.
+            </summary>
+        </member>
+        <member name="P:Newtonsoft.Json.JsonSerializerSettings.MaxDepth">
+            <summary>
+            Gets or sets the maximum depth allowed when reading JSON. Reading past this depth will throw a <see cref="T:Newtonsoft.Json.JsonReaderException"/>.
+            </summary>
+        </member>
+        <member name="P:Newtonsoft.Json.JsonSerializerSettings.Formatting">
+            <summary>
+            Indicates how JSON text output is formatted.
+            </summary>
+        </member>
+        <member name="P:Newtonsoft.Json.JsonSerializerSettings.DateFormatHandling">
+            <summary>
+            Get or set how dates are written to JSON text.
+            </summary>
+        </member>
+        <member name="P:Newtonsoft.Json.JsonSerializerSettings.DateTimeZoneHandling">
+            <summary>
+            Get or set how <see cref="T:System.DateTime"/> time zones are handling during serialization and deserialization.
+            </summary>
+        </member>
+        <member name="P:Newtonsoft.Json.JsonSerializerSettings.DateParseHandling">
+            <summary>
+            Get or set how date formatted strings, e.g. "\/Date(1198908717056)\/" and "2012-03-21T05:40Z", are parsed when reading JSON.
+            </summary>
+        </member>
+        <member name="P:Newtonsoft.Json.JsonSerializerSettings.FloatFormatHandling">
+            <summary>
+            Get or set how special floating point numbers, e.g. <see cref="F:System.Double.NaN"/>,
+            <see cref="F:System.Double.PositiveInfinity"/> and <see cref="F:System.Double.NegativeInfinity"/>,
+            are written as JSON.
+            </summary>
+        </member>
+        <member name="P:Newtonsoft.Json.JsonSerializerSettings.FloatParseHandling">
+            <summary>
+            Get or set how floating point numbers, e.g. 1.0 and 9.9, are parsed when reading JSON text.
+            </summary>
+        </member>
+        <member name="P:Newtonsoft.Json.JsonSerializerSettings.StringEscapeHandling">
+            <summary>
+            Get or set how strings are escaped when writing JSON text.
+            </summary>
+        </member>
+        <member name="P:Newtonsoft.Json.JsonSerializerSettings.Culture">
+            <summary>
+            Gets or sets the culture used when reading JSON. Defaults to <see cref="P:System.Globalization.CultureInfo.InvariantCulture"/>.
+            </summary>
+        </member>
+        <member name="P:Newtonsoft.Json.JsonSerializerSettings.CheckAdditionalContent">
+            <summary>
+            Gets a value indicating whether there will be a check for additional content after deserializing an object.
+            </summary>
+            <value>
+            	<c>true</c> if there will be a check for additional content after deserializing an object; otherwise, <c>false</c>.
+            </value>
+        </member>
+        <member name="T:Newtonsoft.Json.JsonValidatingReader">
+            <summary>
+            Represents a reader that provides <see cref="T:Newtonsoft.Json.Schema.JsonSchema"/> validation.
+            </summary>
+        </member>
+        <member name="M:Newtonsoft.Json.JsonValidatingReader.#ctor(Newtonsoft.Json.JsonReader)">
+            <summary>
+            Initializes a new instance of the <see cref="T:Newtonsoft.Json.JsonValidatingReader"/> class that
+            validates the content returned from the given <see cref="T:Newtonsoft.Json.JsonReader"/>.
+            </summary>
+            <param name="reader">The <see cref="T:Newtonsoft.Json.JsonReader"/> to read from while validating.</param>
+        </member>
+        <member name="M:Newtonsoft.Json.JsonValidatingReader.ReadAsInt32">
+            <summary>
+            Reads the next JSON token from the stream as a <see cref="T:System.Nullable`1"/>.
+            </summary>
+            <returns>A <see cref="T:System.Nullable`1"/>.</returns>
+        </member>
+        <member name="M:Newtonsoft.Json.JsonValidatingReader.ReadAsBytes">
+            <summary>
+            Reads the next JSON token from the stream as a <see cref="T:Byte[]"/>.
+            </summary>
+            <returns>
+            A <see cref="T:Byte[]"/> or a null reference if the next JSON token is null.
+            </returns>
+        </member>
+        <member name="M:Newtonsoft.Json.JsonValidatingReader.ReadAsDecimal">
+            <summary>
+            Reads the next JSON token from the stream as a <see cref="T:System.Nullable`1"/>.
+            </summary>
+            <returns>A <see cref="T:System.Nullable`1"/>.</returns>
+        </member>
+        <member name="M:Newtonsoft.Json.JsonValidatingReader.ReadAsString">
+            <summary>
+            Reads the next JSON token from the stream as a <see cref="T:System.String"/>.
+            </summary>
+            <returns>A <see cref="T:System.String"/>. This method will return <c>null</c> at the end of an array.</returns>
+        </member>
+        <member name="M:Newtonsoft.Json.JsonValidatingReader.ReadAsDateTime">
+            <summary>
+            Reads the next JSON token from the stream as a <see cref="T:System.Nullable`1"/>.
+            </summary>
+            <returns>A <see cref="T:System.String"/>. This method will return <c>null</c> at the end of an array.</returns>
+        </member>
+        <member name="M:Newtonsoft.Json.JsonValidatingReader.Read">
+            <summary>
+            Reads the next JSON token from the stream.
+            </summary>
+            <returns>
+            true if the next token was read successfully; false if there are no more tokens to read.
+            </returns>
+        </member>
+        <member name="E:Newtonsoft.Json.JsonValidatingReader.ValidationEventHandler">
+            <summary>
+            Sets an event handler for receiving schema validation errors.
+            </summary>
+        </member>
+        <member name="P:Newtonsoft.Json.JsonValidatingReader.Value">
+            <summary>
+            Gets the text value of the current JSON token.
+            </summary>
+            <value></value>
+        </member>
+        <member name="P:Newtonsoft.Json.JsonValidatingReader.Depth">
+            <summary>
+            Gets the depth of the current token in the JSON document.
+            </summary>
+            <value>The depth of the current token in the JSON document.</value>
+        </member>
+        <member name="P:Newtonsoft.Json.JsonValidatingReader.Path">
+            <summary>
+            Gets the path of the current JSON token. 
+            </summary>
+        </member>
+        <member name="P:Newtonsoft.Json.JsonValidatingReader.QuoteChar">
+            <summary>
+            Gets the quotation mark character used to enclose the value of a string.
+            </summary>
+            <value></value>
+        </member>
+        <member name="P:Newtonsoft.Json.JsonValidatingReader.TokenType">
+            <summary>
+            Gets the type of the current JSON token.
+            </summary>
+            <value></value>
+        </member>
+        <member name="P:Newtonsoft.Json.JsonValidatingReader.ValueType">
+            <summary>
+            Gets the Common Language Runtime (CLR) type for the current JSON token.
+            </summary>
+            <value></value>
+        </member>
+        <member name="P:Newtonsoft.Json.JsonValidatingReader.Schema">
+            <summary>
+            Gets or sets the schema.
+            </summary>
+            <value>The schema.</value>
+        </member>
+        <member name="P:Newtonsoft.Json.JsonValidatingReader.Reader">
+            <summary>
+            Gets the <see cref="T:Newtonsoft.Json.JsonReader"/> used to construct this <see cref="T:Newtonsoft.Json.JsonValidatingReader"/>.
+            </summary>
+            <value>The <see cref="T:Newtonsoft.Json.JsonReader"/> specified in the constructor.</value>
+        </member>
+        <member name="T:Newtonsoft.Json.Linq.JTokenEqualityComparer">
+            <summary>
+            Compares tokens to determine whether they are equal.
+            </summary>
+        </member>
+        <member name="M:Newtonsoft.Json.Linq.JTokenEqualityComparer.Equals(Newtonsoft.Json.Linq.JToken,Newtonsoft.Json.Linq.JToken)">
+            <summary>
+            Determines whether the specified objects are equal.
+            </summary>
+            <param name="x">The first object of type <see cref="T:Newtonsoft.Json.Linq.JToken"/> to compare.</param>
+            <param name="y">The second object of type <see cref="T:Newtonsoft.Json.Linq.JToken"/> to compare.</param>
+            <returns>
+            true if the specified objects are equal; otherwise, false.
+            </returns>
+        </member>
+        <member name="M:Newtonsoft.Json.Linq.JTokenEqualityComparer.GetHashCode(Newtonsoft.Json.Linq.JToken)">
+            <summary>
+            Returns a hash code for the specified object.
+            </summary>
+            <param name="obj">The <see cref="T:System.Object"/> for which a hash code is to be returned.</param>
+            <returns>A hash code for the specified object.</returns>
+            <exception cref="T:System.ArgumentNullException">The type of <paramref name="obj"/> is a reference type and <paramref name="obj"/> is null.</exception>
+        </member>
+        <member name="T:Newtonsoft.Json.MemberSerialization">
+            <summary>
+            Specifies the member serialization options for the <see cref="T:Newtonsoft.Json.JsonSerializer"/>.
+            </summary>
+        </member>
+        <member name="F:Newtonsoft.Json.MemberSerialization.OptOut">
+            <summary>
+            All public members are serialized by default. Members can be excluded using <see cref="T:Newtonsoft.Json.JsonIgnoreAttribute"/> or <see cref="T:System.NonSerializedAttribute"/>.
+            This is the default member serialization mode.
+            </summary>
+        </member>
+        <member name="F:Newtonsoft.Json.MemberSerialization.OptIn">
+            <summary>
+            Only members must be marked with <see cref="T:Newtonsoft.Json.JsonPropertyAttribute"/> or <see cref="!:DataMemberAttribute"/> are serialized.
+            This member serialization mode can also be set by marking the class with <see cref="!:DataContractAttribute"/>.
+            </summary>
+        </member>
+        <member name="F:Newtonsoft.Json.MemberSerialization.Fields">
+            <summary>
+            All public and private fields are serialized. Members can be excluded using <see cref="T:Newtonsoft.Json.JsonIgnoreAttribute"/> or <see cref="T:System.NonSerializedAttribute"/>.
+            This member serialization mode can also be set by marking the class with <see cref="T:System.SerializableAttribute"/>
+            and setting IgnoreSerializableAttribute on <see cref="T:Newtonsoft.Json.Serialization.DefaultContractResolver"/> to false.
+            </summary>
+        </member>
+        <member name="T:Newtonsoft.Json.ObjectCreationHandling">
+            <summary>
+            Specifies how object creation is handled by the <see cref="T:Newtonsoft.Json.JsonSerializer"/>.
+            </summary>
+        </member>
+        <member name="F:Newtonsoft.Json.ObjectCreationHandling.Auto">
+            <summary>
+            Reuse existing objects, create new objects when needed.
+            </summary>
+        </member>
+        <member name="F:Newtonsoft.Json.ObjectCreationHandling.Reuse">
+            <summary>
+            Only reuse existing objects.
+            </summary>
+        </member>
+        <member name="F:Newtonsoft.Json.ObjectCreationHandling.Replace">
+            <summary>
+            Always create new objects.
+            </summary>
+        </member>
+        <member name="T:Newtonsoft.Json.Converters.IsoDateTimeConverter">
+            <summary>
+            Converts a <see cref="T:System.DateTime"/> to and from the ISO 8601 date format (e.g. 2008-04-12T12:53Z).
+            </summary>
+        </member>
+        <member name="M:Newtonsoft.Json.Converters.IsoDateTimeConverter.WriteJson(Newtonsoft.Json.JsonWriter,System.Object,Newtonsoft.Json.JsonSerializer)">
+            <summary>
+            Writes the JSON representation of the object.
+            </summary>
+            <param name="writer">The <see cref="T:Newtonsoft.Json.JsonWriter"/> to write to.</param>
+            <param name="value">The value.</param>
+            <param name="serializer">The calling serializer.</param>
+        </member>
+        <member name="M:Newtonsoft.Json.Converters.IsoDateTimeConverter.ReadJson(Newtonsoft.Json.JsonReader,System.Type,System.Object,Newtonsoft.Json.JsonSerializer)">
+            <summary>
+            Reads the JSON representation of the object.
+            </summary>
+            <param name="reader">The <see cref="T:Newtonsoft.Json.JsonReader"/> to read from.</param>
+            <param name="objectType">Type of the object.</param>
+            <param name="existingValue">The existing value of object being read.</param>
+            <param name="serializer">The calling serializer.</param>
+            <returns>The object value.</returns>
+        </member>
+        <member name="P:Newtonsoft.Json.Converters.IsoDateTimeConverter.DateTimeStyles">
+            <summary>
+            Gets or sets the date time styles used when converting a date to and from JSON.
+            </summary>
+            <value>The date time styles used when converting a date to and from JSON.</value>
+        </member>
+        <member name="P:Newtonsoft.Json.Converters.IsoDateTimeConverter.DateTimeFormat">
+            <summary>
+            Gets or sets the date time format used when converting a date to and from JSON.
+            </summary>
+            <value>The date time format used when converting a date to and from JSON.</value>
+        </member>
+        <member name="P:Newtonsoft.Json.Converters.IsoDateTimeConverter.Culture">
+            <summary>
+            Gets or sets the culture used when converting a date to and from JSON.
+            </summary>
+            <value>The culture used when converting a date to and from JSON.</value>
+        </member>
+        <member name="T:Newtonsoft.Json.Converters.JavaScriptDateTimeConverter">
+            <summary>
+            Converts a <see cref="T:System.DateTime"/> to and from a JavaScript date constructor (e.g. new Date(52231943)).
+            </summary>
+        </member>
+        <member name="M:Newtonsoft.Json.Converters.JavaScriptDateTimeConverter.WriteJson(Newtonsoft.Json.JsonWriter,System.Object,Newtonsoft.Json.JsonSerializer)">
+            <summary>
+            Writes the JSON representation of the object.
+            </summary>
+            <param name="writer">The <see cref="T:Newtonsoft.Json.JsonWriter"/> to write to.</param>
+            <param name="value">The value.</param>
+            <param name="serializer">The calling serializer.</param>
+        </member>
+        <member name="M:Newtonsoft.Json.Converters.JavaScriptDateTimeConverter.ReadJson(Newtonsoft.Json.JsonReader,System.Type,System.Object,Newtonsoft.Json.JsonSerializer)">
+            <summary>
+            Reads the JSON representation of the object.
+            </summary>
+            <param name="reader">The <see cref="T:Newtonsoft.Json.JsonReader"/> to read from.</param>
+            <param name="objectType">Type of the object.</param>
+            <param name="existingValue">The existing property value of the JSON that is being converted.</param>
+            <param name="serializer">The calling serializer.</param>
+            <returns>The object value.</returns>
+        </member>
+        <member name="T:Newtonsoft.Json.Converters.XmlNodeConverter">
+            <summary>
+            Converts XML to and from JSON.
+            </summary>
+        </member>
+        <member name="M:Newtonsoft.Json.Converters.XmlNodeConverter.WriteJson(Newtonsoft.Json.JsonWriter,System.Object,Newtonsoft.Json.JsonSerializer)">
+            <summary>
+            Writes the JSON representation of the object.
+            </summary>
+            <param name="writer">The <see cref="T:Newtonsoft.Json.JsonWriter"/> to write to.</param>
+            <param name="serializer">The calling serializer.</param>
+            <param name="value">The value.</param>
+        </member>
+        <member name="M:Newtonsoft.Json.Converters.XmlNodeConverter.ReadJson(Newtonsoft.Json.JsonReader,System.Type,System.Object,Newtonsoft.Json.JsonSerializer)">
+            <summary>
+            Reads the JSON representation of the object.
+            </summary>
+            <param name="reader">The <see cref="T:Newtonsoft.Json.JsonReader"/> to read from.</param>
+            <param name="objectType">Type of the object.</param>
+            <param name="existingValue">The existing value of object being read.</param>
+            <param name="serializer">The calling serializer.</param>
+            <returns>The object value.</returns>
+        </member>
+        <member name="M:Newtonsoft.Json.Converters.XmlNodeConverter.IsNamespaceAttribute(System.String,System.String@)">
+            <summary>
+            Checks if the attributeName is a namespace attribute.
+            </summary>
+            <param name="attributeName">Attribute name to test.</param>
+            <param name="prefix">The attribute name prefix if it has one, otherwise an empty string.</param>
+            <returns>True if attribute name is for a namespace attribute, otherwise false.</returns>
+        </member>
+        <member name="M:Newtonsoft.Json.Converters.XmlNodeConverter.CanConvert(System.Type)">
+            <summary>
+            Determines whether this instance can convert the specified value type.
+            </summary>
+            <param name="valueType">Type of the value.</param>
+            <returns>
+            	<c>true</c> if this instance can convert the specified value type; otherwise, <c>false</c>.
+            </returns>
+        </member>
+        <member name="P:Newtonsoft.Json.Converters.XmlNodeConverter.DeserializeRootElementName">
+            <summary>
+            Gets or sets the name of the root element to insert when deserializing to XML if the JSON structure has produces multiple root elements.
+            </summary>
+            <value>The name of the deserialize root element.</value>
+        </member>
+        <member name="P:Newtonsoft.Json.Converters.XmlNodeConverter.WriteArrayAttribute">
+            <summary>
+            Gets or sets a flag to indicate whether to write the Json.NET array attribute.
+            This attribute helps preserve arrays when converting the written XML back to JSON.
+            </summary>
+            <value><c>true</c> if the array attibute is written to the XML; otherwise, <c>false</c>.</value>
+        </member>
+        <member name="P:Newtonsoft.Json.Converters.XmlNodeConverter.OmitRootObject">
+            <summary>
+            Gets or sets a value indicating whether to write the root JSON object.
+            </summary>
+            <value><c>true</c> if the JSON root object is omitted; otherwise, <c>false</c>.</value>
+        </member>
+        <member name="T:Newtonsoft.Json.JsonTextReader">
+            <summary>
+            Represents a reader that provides fast, non-cached, forward-only access to JSON text data.
+            </summary>
+        </member>
+        <member name="M:Newtonsoft.Json.JsonTextReader.#ctor(System.IO.TextReader)">
+            <summary>
+            Initializes a new instance of the <see cref="T:Newtonsoft.Json.JsonReader"/> class with the specified <see cref="T:System.IO.TextReader"/>.
+            </summary>
+            <param name="reader">The <c>TextReader</c> containing the XML data to read.</param>
+        </member>
+        <member name="M:Newtonsoft.Json.JsonTextReader.Read">
+            <summary>
+            Reads the next JSON token from the stream.
+            </summary>
+            <returns>
+            true if the next token was read successfully; false if there are no more tokens to read.
+            </returns>
+        </member>
+        <member name="M:Newtonsoft.Json.JsonTextReader.ReadAsBytes">
+            <summary>
+            Reads the next JSON token from the stream as a <see cref="T:Byte[]"/>.
+            </summary>
+            <returns>
+            A <see cref="T:Byte[]"/> or a null reference if the next JSON token is null. This method will return <c>null</c> at the end of an array.
+            </returns>
+        </member>
+        <member name="M:Newtonsoft.Json.JsonTextReader.ReadAsDecimal">
+            <summary>
+            Reads the next JSON token from the stream as a <see cref="T:System.Nullable`1"/>.
+            </summary>
+            <returns>A <see cref="T:System.Nullable`1"/>. This method will return <c>null</c> at the end of an array.</returns>
+        </member>
+        <member name="M:Newtonsoft.Json.JsonTextReader.ReadAsInt32">
+            <summary>
+            Reads the next JSON token from the stream as a <see cref="T:System.Nullable`1"/>.
+            </summary>
+            <returns>A <see cref="T:System.Nullable`1"/>. This method will return <c>null</c> at the end of an array.</returns>
+        </member>
+        <member name="M:Newtonsoft.Json.JsonTextReader.ReadAsString">
+            <summary>
+            Reads the next JSON token from the stream as a <see cref="T:System.String"/>.
+            </summary>
+            <returns>A <see cref="T:System.String"/>. This method will return <c>null</c> at the end of an array.</returns>
+        </member>
+        <member name="M:Newtonsoft.Json.JsonTextReader.ReadAsDateTime">
+            <summary>
+            Reads the next JSON token from the stream as a <see cref="T:System.Nullable`1"/>.
+            </summary>
+            <returns>A <see cref="T:System.String"/>. This method will return <c>null</c> at the end of an array.</returns>
+        </member>
+        <member name="M:Newtonsoft.Json.JsonTextReader.Close">
+            <summary>
+            Changes the state to closed. 
+            </summary>
+        </member>
+        <member name="M:Newtonsoft.Json.JsonTextReader.HasLineInfo">
+            <summary>
+            Gets a value indicating whether the class can return line information.
+            </summary>
+            <returns>
+            	<c>true</c> if LineNumber and LinePosition can be provided; otherwise, <c>false</c>.
+            </returns>
+        </member>
+        <member name="P:Newtonsoft.Json.JsonTextReader.LineNumber">
+            <summary>
+            Gets the current line number.
+            </summary>
+            <value>
+            The current line number or 0 if no line information is available (for example, HasLineInfo returns false).
+            </value>
+        </member>
+        <member name="P:Newtonsoft.Json.JsonTextReader.LinePosition">
+            <summary>
+            Gets the current line position.
+            </summary>
+            <value>
+            The current line position or 0 if no line information is available (for example, HasLineInfo returns false).
+            </value>
+        </member>
+        <member name="T:Newtonsoft.Json.JsonPropertyAttribute">
+            <summary>
+            Instructs the <see cref="T:Newtonsoft.Json.JsonSerializer"/> to always serialize the member with the specified name.
+            </summary>
+        </member>
+        <member name="M:Newtonsoft.Json.JsonPropertyAttribute.#ctor">
+            <summary>
+            Initializes a new instance of the <see cref="T:Newtonsoft.Json.JsonPropertyAttribute"/> class.
+            </summary>
+        </member>
+        <member name="M:Newtonsoft.Json.JsonPropertyAttribute.#ctor(System.String)">
+            <summary>
+            Initializes a new instance of the <see cref="T:Newtonsoft.Json.JsonPropertyAttribute"/> class with the specified name.
+            </summary>
+            <param name="propertyName">Name of the property.</param>
+        </member>
+        <member name="P:Newtonsoft.Json.JsonPropertyAttribute.ItemConverterType">
+            <summary>
+            Gets or sets the converter used when serializing the property's collection items.
+            </summary>
+            <value>The collection's items converter.</value>
+        </member>
+        <member name="P:Newtonsoft.Json.JsonPropertyAttribute.NullValueHandling">
+            <summary>
+            Gets or sets the null value handling used when serializing this property.
+            </summary>
+            <value>The null value handling.</value>
+        </member>
+        <member name="P:Newtonsoft.Json.JsonPropertyAttribute.DefaultValueHandling">
+            <summary>
+            Gets or sets the default value handling used when serializing this property.
+            </summary>
+            <value>The default value handling.</value>
+        </member>
+        <member name="P:Newtonsoft.Json.JsonPropertyAttribute.ReferenceLoopHandling">
+            <summary>
+            Gets or sets the reference loop handling used when serializing this property.
+            </summary>
+            <value>The reference loop handling.</value>
+        </member>
+        <member name="P:Newtonsoft.Json.JsonPropertyAttribute.ObjectCreationHandling">
+            <summary>
+            Gets or sets the object creation handling used when deserializing this property.
+            </summary>
+            <value>The object creation handling.</value>
+        </member>
+        <member name="P:Newtonsoft.Json.JsonPropertyAttribute.TypeNameHandling">
+            <summary>
+            Gets or sets the type name handling used when serializing this property.
+            </summary>
+            <value>The type name handling.</value>
+        </member>
+        <member name="P:Newtonsoft.Json.JsonPropertyAttribute.IsReference">
+            <summary>
+            Gets or sets whether this property's value is serialized as a reference.
+            </summary>
+            <value>Whether this property's value is serialized as a reference.</value>
+        </member>
+        <member name="P:Newtonsoft.Json.JsonPropertyAttribute.Order">
+            <summary>
+            Gets or sets the order of serialization and deserialization of a member.
+            </summary>
+            <value>The numeric order of serialization or deserialization.</value>
+        </member>
+        <member name="P:Newtonsoft.Json.JsonPropertyAttribute.Required">
+            <summary>
+            Gets or sets a value indicating whether this property is required.
+            </summary>
+            <value>
+            	A value indicating whether this property is required.
+            </value>
+        </member>
+        <member name="P:Newtonsoft.Json.JsonPropertyAttribute.PropertyName">
+            <summary>
+            Gets or sets the name of the property.
+            </summary>
+            <value>The name of the property.</value>
+        </member>
+        <member name="P:Newtonsoft.Json.JsonPropertyAttribute.ItemReferenceLoopHandling">
+            <summary>
+            Gets or sets the the reference loop handling used when serializing the property's collection items.
+            </summary>
+            <value>The collection's items reference loop handling.</value>
+        </member>
+        <member name="P:Newtonsoft.Json.JsonPropertyAttribute.ItemTypeNameHandling">
+            <summary>
+            Gets or sets the the type name handling used when serializing the property's collection items.
+            </summary>
+            <value>The collection's items type name handling.</value>
+        </member>
+        <member name="P:Newtonsoft.Json.JsonPropertyAttribute.ItemIsReference">
+            <summary>
+            Gets or sets whether this property's collection items are serialized as a reference.
+            </summary>
+            <value>Whether this property's collection items are serialized as a reference.</value>
+        </member>
+        <member name="T:Newtonsoft.Json.JsonIgnoreAttribute">
+            <summary>
+            Instructs the <see cref="T:Newtonsoft.Json.JsonSerializer"/> not to serialize the public field or public read/write property value.
+            </summary>
+        </member>
+        <member name="T:Newtonsoft.Json.JsonTextWriter">
+            <summary>
+            Represents a writer that provides a fast, non-cached, forward-only way of generating Json data.
+            </summary>
+        </member>
+        <member name="M:Newtonsoft.Json.JsonTextWriter.#ctor(System.IO.TextWriter)">
+            <summary>
+            Creates an instance of the <c>JsonWriter</c> class using the specified <see cref="T:System.IO.TextWriter"/>. 
+            </summary>
+            <param name="textWriter">The <c>TextWriter</c> to write to.</param>
+        </member>
+        <member name="M:Newtonsoft.Json.JsonTextWriter.Flush">
+            <summary>
+            Flushes whatever is in the buffer to the underlying streams and also flushes the underlying stream.
+            </summary>
+        </member>
+        <member name="M:Newtonsoft.Json.JsonTextWriter.Close">
+            <summary>
+            Closes this stream and the underlying stream.
+            </summary>
+        </member>
+        <member name="M:Newtonsoft.Json.JsonTextWriter.WriteStartObject">
+            <summary>
+            Writes the beginning of a Json object.
+            </summary>
+        </member>
+        <member name="M:Newtonsoft.Json.JsonTextWriter.WriteStartArray">
+            <summary>
+            Writes the beginning of a Json array.
+            </summary>
+        </member>
+        <member name="M:Newtonsoft.Json.JsonTextWriter.WriteStartConstructor(System.String)">
+            <summary>
+            Writes the start of a constructor with the given name.
+            </summary>
+            <param name="name">The name of the constructor.</param>
+        </member>
+        <member name="M:Newtonsoft.Json.JsonTextWriter.WriteEnd(Newtonsoft.Json.JsonToken)">
+            <summary>
+            Writes the specified end token.
+            </summary>
+            <param name="token">The end token to write.</param>
+        </member>
+        <member name="M:Newtonsoft.Json.JsonTextWriter.WritePropertyName(System.String)">
+            <summary>
+            Writes the property name of a name/value pair on a Json object.
+            </summary>
+            <param name="name">The name of the property.</param>
+        </member>
+        <member name="M:Newtonsoft.Json.JsonTextWriter.WritePropertyName(System.String,System.Boolean)">
+            <summary>
+            Writes the property name of a name/value pair on a JSON object.
+            </summary>
+            <param name="name">The name of the property.</param>
+            <param name="escape">A flag to indicate whether the text should be escaped when it is written as a JSON property name.</param>
+        </member>
+        <member name="M:Newtonsoft.Json.JsonTextWriter.WriteIndent">
+            <summary>
+            Writes indent characters.
+            </summary>
+        </member>
+        <member name="M:Newtonsoft.Json.JsonTextWriter.WriteValueDelimiter">
+            <summary>
+            Writes the JSON value delimiter.
+            </summary>
+        </member>
+        <member name="M:Newtonsoft.Json.JsonTextWriter.WriteIndentSpace">
+            <summary>
+            Writes an indent space.
+            </summary>
+        </member>
+        <member name="M:Newtonsoft.Json.JsonTextWriter.WriteNull">
+            <summary>
+            Writes a null value.
+            </summary>
+        </member>
+        <member name="M:Newtonsoft.Json.JsonTextWriter.WriteUndefined">
+            <summary>
+            Writes an undefined value.
+            </summary>
+        </member>
+        <member name="M:Newtonsoft.Json.JsonTextWriter.WriteRaw(System.String)">
+            <summary>
+            Writes raw JSON.
+            </summary>
+            <param name="json">The raw JSON to write.</param>
+        </member>
+        <member name="M:Newtonsoft.Json.JsonTextWriter.WriteValue(System.String)">
+            <summary>
+            Writes a <see cref="T:System.String"/> value.
+            </summary>
+            <param name="value">The <see cref="T:System.String"/> value to write.</param>
+        </member>
+        <member name="M:Newtonsoft.Json.JsonTextWriter.WriteValue(System.Int32)">
+            <summary>
+            Writes a <see cref="T:System.Int32"/> value.
+            </summary>
+            <param name="value">The <see cref="T:System.Int32"/> value to write.</param>
+        </member>
+        <member name="M:Newtonsoft.Json.JsonTextWriter.WriteValue(System.UInt32)">
+            <summary>
+            Writes a <see cref="T:System.UInt32"/> value.
+            </summary>
+            <param name="value">The <see cref="T:System.UInt32"/> value to write.</param>
+        </member>
+        <member name="M:Newtonsoft.Json.JsonTextWriter.WriteValue(System.Int64)">
+            <summary>
+            Writes a <see cref="T:System.Int64"/> value.
+            </summary>
+            <param name="value">The <see cref="T:System.Int64"/> value to write.</param>
+        </member>
+        <member name="M:Newtonsoft.Json.JsonTextWriter.WriteValue(System.UInt64)">
+            <summary>
+            Writes a <see cref="T:System.UInt64"/> value.
+            </summary>
+            <param name="value">The <see cref="T:System.UInt64"/> value to write.</param>
+        </member>
+        <member name="M:Newtonsoft.Json.JsonTextWriter.WriteValue(System.Single)">
+            <summary>
+            Writes a <see cref="T:System.Single"/> value.
+            </summary>
+            <param name="value">The <see cref="T:System.Single"/> value to write.</param>
+        </member>
+        <member name="M:Newtonsoft.Json.JsonTextWriter.WriteValue(System.Nullable{System.Single})">
+            <summary>
+            Writes a <see cref="T:System.Nullable`1"/> value.
+            </summary>
+            <param name="value">The <see cref="T:System.Nullable`1"/> value to write.</param>
+        </member>
+        <member name="M:Newtonsoft.Json.JsonTextWriter.WriteValue(System.Double)">
+            <summary>
+            Writes a <see cref="T:System.Double"/> value.
+            </summary>
+            <param name="value">The <see cref="T:System.Double"/> value to write.</param>
+        </member>
+        <member name="M:Newtonsoft.Json.JsonTextWriter.WriteValue(System.Nullable{System.Double})">
+            <summary>
+            Writes a <see cref="T:System.Nullable`1"/> value.
+            </summary>
+            <param name="value">The <see cref="T:System.Nullable`1"/> value to write.</param>
+        </member>
+        <member name="M:Newtonsoft.Json.JsonTextWriter.WriteValue(System.Boolean)">
+            <summary>
+            Writes a <see cref="T:System.Boolean"/> value.
+            </summary>
+            <param name="value">The <see cref="T:System.Boolean"/> value to write.</param>
+        </member>
+        <member name="M:Newtonsoft.Json.JsonTextWriter.WriteValue(System.Int16)">
+            <summary>
+            Writes a <see cref="T:System.Int16"/> value.
+            </summary>
+            <param name="value">The <see cref="T:System.Int16"/> value to write.</param>
+        </member>
+        <member name="M:Newtonsoft.Json.JsonTextWriter.WriteValue(System.UInt16)">
+            <summary>
+            Writes a <see cref="T:System.UInt16"/> value.
+            </summary>
+            <param name="value">The <see cref="T:System.UInt16"/> value to write.</param>
+        </member>
+        <member name="M:Newtonsoft.Json.JsonTextWriter.WriteValue(System.Char)">
+            <summary>
+            Writes a <see cref="T:System.Char"/> value.
+            </summary>
+            <param name="value">The <see cref="T:System.Char"/> value to write.</param>
+        </member>
+        <member name="M:Newtonsoft.Json.JsonTextWriter.WriteValue(System.Byte)">
+            <summary>
+            Writes a <see cref="T:System.Byte"/> value.
+            </summary>
+            <param name="value">The <see cref="T:System.Byte"/> value to write.</param>
+        </member>
+        <member name="M:Newtonsoft.Json.JsonTextWriter.WriteValue(System.SByte)">
+            <summary>
+            Writes a <see cref="T:System.SByte"/> value.
+            </summary>
+            <param name="value">The <see cref="T:System.SByte"/> value to write.</param>
+        </member>
+        <member name="M:Newtonsoft.Json.JsonTextWriter.WriteValue(System.Decimal)">
+            <summary>
+            Writes a <see cref="T:System.Decimal"/> value.
+            </summary>
+            <param name="value">The <see cref="T:System.Decimal"/> value to write.</param>
+        </member>
+        <member name="M:Newtonsoft.Json.JsonTextWriter.WriteValue(System.DateTime)">
+            <summary>
+            Writes a <see cref="T:System.DateTime"/> value.
+            </summary>
+            <param name="value">The <see cref="T:System.DateTime"/> value to write.</param>
+        </member>
+        <member name="M:Newtonsoft.Json.JsonTextWriter.WriteValue(System.Byte[])">
+            <summary>
+            Writes a <see cref="T:Byte[]"/> value.
+            </summary>
+            <param name="value">The <see cref="T:Byte[]"/> value to write.</param>
+        </member>
+        <member name="M:Newtonsoft.Json.JsonTextWriter.WriteValue(System.Guid)">
+            <summary>
+            Writes a <see cref="T:System.Guid"/> value.
+            </summary>
+            <param name="value">The <see cref="T:System.Guid"/> value to write.</param>
+        </member>
+        <member name="M:Newtonsoft.Json.JsonTextWriter.WriteValue(System.TimeSpan)">
+            <summary>
+            Writes a <see cref="T:System.TimeSpan"/> value.
+            </summary>
+            <param name="value">The <see cref="T:System.TimeSpan"/> value to write.</param>
+        </member>
+        <member name="M:Newtonsoft.Json.JsonTextWriter.WriteValue(System.Uri)">
+            <summary>
+            Writes a <see cref="T:System.Uri"/> value.
+            </summary>
+            <param name="value">The <see cref="T:System.Uri"/> value to write.</param>
+        </member>
+        <member name="M:Newtonsoft.Json.JsonTextWriter.WriteComment(System.String)">
+            <summary>
+            Writes out a comment <code>/*...*/</code> containing the specified text. 
+            </summary>
+            <param name="text">Text to place inside the comment.</param>
+        </member>
+        <member name="M:Newtonsoft.Json.JsonTextWriter.WriteWhitespace(System.String)">
+            <summary>
+            Writes out the given white space.
+            </summary>
+            <param name="ws">The string of white space characters.</param>
+        </member>
+        <member name="P:Newtonsoft.Json.JsonTextWriter.Indentation">
+            <summary>
+            Gets or sets how many IndentChars to write for each level in the hierarchy when <see cref="T:Newtonsoft.Json.Formatting"/> is set to <c>Formatting.Indented</c>.
+            </summary>
+        </member>
+        <member name="P:Newtonsoft.Json.JsonTextWriter.QuoteChar">
+            <summary>
+            Gets or sets which character to use to quote attribute values.
+            </summary>
+        </member>
+        <member name="P:Newtonsoft.Json.JsonTextWriter.IndentChar">
+            <summary>
+            Gets or sets which character to use for indenting when <see cref="T:Newtonsoft.Json.Formatting"/> is set to <c>Formatting.Indented</c>.
+            </summary>
+        </member>
+        <member name="P:Newtonsoft.Json.JsonTextWriter.QuoteName">
+            <summary>
+            Gets or sets a value indicating whether object names will be surrounded with quotes.
+            </summary>
+        </member>
+        <member name="T:Newtonsoft.Json.JsonWriterException">
+            <summary>
+            The exception thrown when an error occurs while reading Json text.
+            </summary>
+        </member>
+        <member name="M:Newtonsoft.Json.JsonWriterException.#ctor">
+            <summary>
+            Initializes a new instance of the <see cref="T:Newtonsoft.Json.JsonWriterException"/> class.
+            </summary>
+        </member>
+        <member name="M:Newtonsoft.Json.JsonWriterException.#ctor(System.String)">
+            <summary>
+            Initializes a new instance of the <see cref="T:Newtonsoft.Json.JsonWriterException"/> class
+            with a specified error message.
+            </summary>
+            <param name="message">The error message that explains the reason for the exception.</param>
+        </member>
+        <member name="M:Newtonsoft.Json.JsonWriterException.#ctor(System.String,System.Exception)">
+            <summary>
+            Initializes a new instance of the <see cref="T:Newtonsoft.Json.JsonWriterException"/> class
+            with a specified error message and a reference to the inner exception that is the cause of this exception.
+            </summary>
+            <param name="message">The error message that explains the reason for the exception.</param>
+            <param name="innerException">The exception that is the cause of the current exception, or a null reference (Nothing in Visual Basic) if no inner exception is specified.</param>
+        </member>
+        <member name="M:Newtonsoft.Json.JsonWriterException.#ctor(System.Runtime.Serialization.SerializationInfo,System.Runtime.Serialization.StreamingContext)">
+            <summary>
+            Initializes a new instance of the <see cref="T:Newtonsoft.Json.JsonWriterException"/> class.
+            </summary>
+            <param name="info">The <see cref="T:System.Runtime.Serialization.SerializationInfo"/> that holds the serialized object data about the exception being thrown.</param>
+            <param name="context">The <see cref="T:System.Runtime.Serialization.StreamingContext"/> that contains contextual information about the source or destination.</param>
+            <exception cref="T:System.ArgumentNullException">The <paramref name="info"/> parameter is null. </exception>
+            <exception cref="T:System.Runtime.Serialization.SerializationException">The class name is null or <see cref="P:System.Exception.HResult"/> is zero (0). </exception>
+        </member>
+        <member name="P:Newtonsoft.Json.JsonWriterException.Path">
+            <summary>
+            Gets the path to the JSON where the error occurred.
+            </summary>
+            <value>The path to the JSON where the error occurred.</value>
+        </member>
+        <member name="T:Newtonsoft.Json.JsonReaderException">
+            <summary>
+            The exception thrown when an error occurs while reading Json text.
+            </summary>
+        </member>
+        <member name="M:Newtonsoft.Json.JsonReaderException.#ctor">
+            <summary>
+            Initializes a new instance of the <see cref="T:Newtonsoft.Json.JsonReaderException"/> class.
+            </summary>
+        </member>
+        <member name="M:Newtonsoft.Json.JsonReaderException.#ctor(System.String)">
+            <summary>
+            Initializes a new instance of the <see cref="T:Newtonsoft.Json.JsonReaderException"/> class
+            with a specified error message.
+            </summary>
+            <param name="message">The error message that explains the reason for the exception.</param>
+        </member>
+        <member name="M:Newtonsoft.Json.JsonReaderException.#ctor(System.String,System.Exception)">
+            <summary>
+            Initializes a new instance of the <see cref="T:Newtonsoft.Json.JsonReaderException"/> class
+            with a specified error message and a reference to the inner exception that is the cause of this exception.
+            </summary>
+            <param name="message">The error message that explains the reason for the exception.</param>
+            <param name="innerException">The exception that is the cause of the current exception, or a null reference (Nothing in Visual Basic) if no inner exception is specified.</param>
+        </member>
+        <member name="M:Newtonsoft.Json.JsonReaderException.#ctor(System.Runtime.Serialization.SerializationInfo,System.Runtime.Serialization.StreamingContext)">
+            <summary>
+            Initializes a new instance of the <see cref="T:Newtonsoft.Json.JsonReaderException"/> class.
+            </summary>
+            <param name="info">The <see cref="T:System.Runtime.Serialization.SerializationInfo"/> that holds the serialized object data about the exception being thrown.</param>
+            <param name="context">The <see cref="T:System.Runtime.Serialization.StreamingContext"/> that contains contextual information about the source or destination.</param>
+            <exception cref="T:System.ArgumentNullException">The <paramref name="info"/> parameter is null. </exception>
+            <exception cref="T:System.Runtime.Serialization.SerializationException">The class name is null or <see cref="P:System.Exception.HResult"/> is zero (0). </exception>
+        </member>
+        <member name="P:Newtonsoft.Json.JsonReaderException.LineNumber">
+            <summary>
+            Gets the line number indicating where the error occurred.
+            </summary>
+            <value>The line number indicating where the error occurred.</value>
+        </member>
+        <member name="P:Newtonsoft.Json.JsonReaderException.LinePosition">
+            <summary>
+            Gets the line position indicating where the error occurred.
+            </summary>
+            <value>The line position indicating where the error occurred.</value>
+        </member>
+        <member name="P:Newtonsoft.Json.JsonReaderException.Path">
+            <summary>
+            Gets the path to the JSON where the error occurred.
+            </summary>
+            <value>The path to the JSON where the error occurred.</value>
+        </member>
+        <member name="T:Newtonsoft.Json.JsonConverterCollection">
+            <summary>
+            Represents a collection of <see cref="T:Newtonsoft.Json.JsonConverter"/>.
+            </summary>
+        </member>
+        <member name="T:Newtonsoft.Json.JsonConvert">
+            <summary>
+            Provides methods for converting between common language runtime types and JSON types.
+            </summary>
+            <example>
+              <code lang="cs" source="..\Src\Newtonsoft.Json.Tests\Documentation\SerializationTests.cs" region="SerializeObject" title="Serializing and Deserializing JSON with JsonConvert" />
+            </example>
+        </member>
+        <member name="F:Newtonsoft.Json.JsonConvert.True">
+            <summary>
+            Represents JavaScript's boolean value true as a string. This field is read-only.
+            </summary>
+        </member>
+        <member name="F:Newtonsoft.Json.JsonConvert.False">
+            <summary>
+            Represents JavaScript's boolean value false as a string. This field is read-only.
+            </summary>
+        </member>
+        <member name="F:Newtonsoft.Json.JsonConvert.Null">
+            <summary>
+            Represents JavaScript's null as a string. This field is read-only.
+            </summary>
+        </member>
+        <member name="F:Newtonsoft.Json.JsonConvert.Undefined">
+            <summary>
+            Represents JavaScript's undefined as a string. This field is read-only.
+            </summary>
+        </member>
+        <member name="F:Newtonsoft.Json.JsonConvert.PositiveInfinity">
+            <summary>
+            Represents JavaScript's positive infinity as a string. This field is read-only.
+            </summary>
+        </member>
+        <member name="F:Newtonsoft.Json.JsonConvert.NegativeInfinity">
+            <summary>
+            Represents JavaScript's negative infinity as a string. This field is read-only.
+            </summary>
+        </member>
+        <member name="F:Newtonsoft.Json.JsonConvert.NaN">
+            <summary>
+            Represents JavaScript's NaN as a string. This field is read-only.
+            </summary>
+        </member>
+        <member name="M:Newtonsoft.Json.JsonConvert.ToString(System.DateTime)">
+            <summary>
+            Converts the <see cref="T:System.DateTime"/> to its JSON string representation.
+            </summary>
+            <param name="value">The value to convert.</param>
+            <returns>A JSON string representation of the <see cref="T:System.DateTime"/>.</returns>
+        </member>
+        <member name="M:Newtonsoft.Json.JsonConvert.ToString(System.DateTime,Newtonsoft.Json.DateFormatHandling,Newtonsoft.Json.DateTimeZoneHandling)">
+            <summary>
+            Converts the <see cref="T:System.DateTime"/> to its JSON string representation using the <see cref="T:Newtonsoft.Json.DateFormatHandling"/> specified.
+            </summary>
+            <param name="value">The value to convert.</param>
+            <param name="format">The format the date will be converted to.</param>
+            <param name="timeZoneHandling">The time zone handling when the date is converted to a string.</param>
+            <returns>A JSON string representation of the <see cref="T:System.DateTime"/>.</returns>
+        </member>
+        <member name="M:Newtonsoft.Json.JsonConvert.ToString(System.Boolean)">
+            <summary>
+            Converts the <see cref="T:System.Boolean"/> to its JSON string representation.
+            </summary>
+            <param name="value">The value to convert.</param>
+            <returns>A JSON string representation of the <see cref="T:System.Boolean"/>.</returns>
+        </member>
+        <member name="M:Newtonsoft.Json.JsonConvert.ToString(System.Char)">
+            <summary>
+            Converts the <see cref="T:System.Char"/> to its JSON string representation.
+            </summary>
+            <param name="value">The value to convert.</param>
+            <returns>A JSON string representation of the <see cref="T:System.Char"/>.</returns>
+        </member>
+        <member name="M:Newtonsoft.Json.JsonConvert.ToString(System.Enum)">
+            <summary>
+            Converts the <see cref="T:System.Enum"/> to its JSON string representation.
+            </summary>
+            <param name="value">The value to convert.</param>
+            <returns>A JSON string representation of the <see cref="T:System.Enum"/>.</returns>
+        </member>
+        <member name="M:Newtonsoft.Json.JsonConvert.ToString(System.Int32)">
+            <summary>
+            Converts the <see cref="T:System.Int32"/> to its JSON string representation.
+            </summary>
+            <param name="value">The value to convert.</param>
+            <returns>A JSON string representation of the <see cref="T:System.Int32"/>.</returns>
+        </member>
+        <member name="M:Newtonsoft.Json.JsonConvert.ToString(System.Int16)">
+            <summary>
+            Converts the <see cref="T:System.Int16"/> to its JSON string representation.
+            </summary>
+            <param name="value">The value to convert.</param>
+            <returns>A JSON string representation of the <see cref="T:System.Int16"/>.</returns>
+        </member>
+        <member name="M:Newtonsoft.Json.JsonConvert.ToString(System.UInt16)">
+            <summary>
+            Converts the <see cref="T:System.UInt16"/> to its JSON string representation.
+            </summary>
+            <param name="value">The value to convert.</param>
+            <returns>A JSON string representation of the <see cref="T:System.UInt16"/>.</returns>
+        </member>
+        <member name="M:Newtonsoft.Json.JsonConvert.ToString(System.UInt32)">
+            <summary>
+            Converts the <see cref="T:System.UInt32"/> to its JSON string representation.
+            </summary>
+            <param name="value">The value to convert.</param>
+            <returns>A JSON string representation of the <see cref="T:System.UInt32"/>.</returns>
+        </member>
+        <member name="M:Newtonsoft.Json.JsonConvert.ToString(System.Int64)">
+            <summary>
+            Converts the <see cref="T:System.Int64"/>  to its JSON string representation.
+            </summary>
+            <param name="value">The value to convert.</param>
+            <returns>A JSON string representation of the <see cref="T:System.Int64"/>.</returns>
+        </member>
+        <member name="M:Newtonsoft.Json.JsonConvert.ToString(System.UInt64)">
+            <summary>
+            Converts the <see cref="T:System.UInt64"/> to its JSON string representation.
+            </summary>
+            <param name="value">The value to convert.</param>
+            <returns>A JSON string representation of the <see cref="T:System.UInt64"/>.</returns>
+        </member>
+        <member name="M:Newtonsoft.Json.JsonConvert.ToString(System.Single)">
+            <summary>
+            Converts the <see cref="T:System.Single"/> to its JSON string representation.
+            </summary>
+            <param name="value">The value to convert.</param>
+            <returns>A JSON string representation of the <see cref="T:System.Single"/>.</returns>
+        </member>
+        <member name="M:Newtonsoft.Json.JsonConvert.ToString(System.Double)">
+            <summary>
+            Converts the <see cref="T:System.Double"/> to its JSON string representation.
+            </summary>
+            <param name="value">The value to convert.</param>
+            <returns>A JSON string representation of the <see cref="T:System.Double"/>.</returns>
+        </member>
+        <member name="M:Newtonsoft.Json.JsonConvert.ToString(System.Byte)">
+            <summary>
+            Converts the <see cref="T:System.Byte"/> to its JSON string representation.
+            </summary>
+            <param name="value">The value to convert.</param>
+            <returns>A JSON string representation of the <see cref="T:System.Byte"/>.</returns>
+        </member>
+        <member name="M:Newtonsoft.Json.JsonConvert.ToString(System.SByte)">
+            <summary>
+            Converts the <see cref="T:System.SByte"/> to its JSON string representation.
+            </summary>
+            <param name="value">The value to convert.</param>
+            <returns>A JSON string representation of the <see cref="T:System.SByte"/>.</returns>
+        </member>
+        <member name="M:Newtonsoft.Json.JsonConvert.ToString(System.Decimal)">
+            <summary>
+            Converts the <see cref="T:System.Decimal"/> to its JSON string representation.
+            </summary>
+            <param name="value">The value to convert.</param>
+            <returns>A JSON string representation of the <see cref="T:System.SByte"/>.</returns>
+        </member>
+        <member name="M:Newtonsoft.Json.JsonConvert.ToString(System.Guid)">
+            <summary>
+            Converts the <see cref="T:System.Guid"/> to its JSON string representation.
+            </summary>
+            <param name="value">The value to convert.</param>
+            <returns>A JSON string representation of the <see cref="T:System.Guid"/>.</returns>
+        </member>
+        <member name="M:Newtonsoft.Json.JsonConvert.ToString(System.TimeSpan)">
+            <summary>
+            Converts the <see cref="T:System.TimeSpan"/> to its JSON string representation.
+            </summary>
+            <param name="value">The value to convert.</param>
+            <returns>A JSON string representation of the <see cref="T:System.TimeSpan"/>.</returns>
+        </member>
+        <member name="M:Newtonsoft.Json.JsonConvert.ToString(System.Uri)">
+            <summary>
+            Converts the <see cref="T:System.Uri"/> to its JSON string representation.
+            </summary>
+            <param name="value">The value to convert.</param>
+            <returns>A JSON string representation of the <see cref="T:System.Uri"/>.</returns>
+        </member>
+        <member name="M:Newtonsoft.Json.JsonConvert.ToString(System.String)">
+            <summary>
+            Converts the <see cref="T:System.String"/> to its JSON string representation.
+            </summary>
+            <param name="value">The value to convert.</param>
+            <returns>A JSON string representation of the <see cref="T:System.String"/>.</returns>
+        </member>
+        <member name="M:Newtonsoft.Json.JsonConvert.ToString(System.String,System.Char)">
+            <summary>
+            Converts the <see cref="T:System.String"/> to its JSON string representation.
+            </summary>
+            <param name="value">The value to convert.</param>
+            <param name="delimiter">The string delimiter character.</param>
+            <returns>A JSON string representation of the <see cref="T:System.String"/>.</returns>
+        </member>
+        <member name="M:Newtonsoft.Json.JsonConvert.ToString(System.Object)">
+            <summary>
+            Converts the <see cref="T:System.Object"/> to its JSON string representation.
+            </summary>
+            <param name="value">The value to convert.</param>
+            <returns>A JSON string representation of the <see cref="T:System.Object"/>.</returns>
+        </member>
+        <member name="M:Newtonsoft.Json.JsonConvert.SerializeObject(System.Object)">
+            <summary>
+            Serializes the specified object to a JSON string.
+            </summary>
+            <param name="value">The object to serialize.</param>
+            <returns>A JSON string representation of the object.</returns>
+        </member>
+        <member name="M:Newtonsoft.Json.JsonConvert.SerializeObject(System.Object,Newtonsoft.Json.Formatting)">
+            <summary>
+            Serializes the specified object to a JSON string.
+            </summary>
+            <param name="value">The object to serialize.</param>
+            <param name="formatting">Indicates how the output is formatted.</param>
+            <returns>
+            A JSON string representation of the object.
+            </returns>
+        </member>
+        <member name="M:Newtonsoft.Json.JsonConvert.SerializeObject(System.Object,Newtonsoft.Json.JsonConverter[])">
+            <summary>
+            Serializes the specified object to a JSON string using a collection of <see cref="T:Newtonsoft.Json.JsonConverter"/>.
+            </summary>
+            <param name="value">The object to serialize.</param>
+            <param name="converters">A collection converters used while serializing.</param>
+            <returns>A JSON string representation of the object.</returns>
+        </member>
+        <member name="M:Newtonsoft.Json.JsonConvert.SerializeObject(System.Object,Newtonsoft.Json.Formatting,Newtonsoft.Json.JsonConverter[])">
+            <summary>
+            Serializes the specified object to a JSON string using a collection of <see cref="T:Newtonsoft.Json.JsonConverter"/>.
+            </summary>
+            <param name="value">The object to serialize.</param>
+            <param name="formatting">Indicates how the output is formatted.</param>
+            <param name="converters">A collection converters used while serializing.</param>
+            <returns>A JSON string representation of the object.</returns>
+        </member>
+        <member name="M:Newtonsoft.Json.JsonConvert.SerializeObject(System.Object,Newtonsoft.Json.JsonSerializerSettings)">
+            <summary>
+            Serializes the specified object to a JSON string using a collection of <see cref="T:Newtonsoft.Json.JsonConverter"/>.
+            </summary>
+            <param name="value">The object to serialize.</param>
+            <param name="settings">The <see cref="T:Newtonsoft.Json.JsonSerializerSettings"/> used to serialize the object.
+            If this is null, default serialization settings will be is used.</param>
+            <returns>
+            A JSON string representation of the object.
+            </returns>
+        </member>
+        <member name="M:Newtonsoft.Json.JsonConvert.SerializeObject(System.Object,Newtonsoft.Json.Formatting,Newtonsoft.Json.JsonSerializerSettings)">
+            <summary>
+            Serializes the specified object to a JSON string using a collection of <see cref="T:Newtonsoft.Json.JsonConverter"/>.
+            </summary>
+            <param name="value">The object to serialize.</param>
+            <param name="formatting">Indicates how the output is formatted.</param>
+            <param name="settings">The <see cref="T:Newtonsoft.Json.JsonSerializerSettings"/> used to serialize the object.
+            If this is null, default serialization settings will be is used.</param>
+            <returns>
+            A JSON string representation of the object.
+            </returns>
+        </member>
+        <member name="M:Newtonsoft.Json.JsonConvert.SerializeObject(System.Object,System.Type,Newtonsoft.Json.Formatting,Newtonsoft.Json.JsonSerializerSettings)">
+            <summary>
+            Serializes the specified object to a JSON string using a collection of <see cref="T:Newtonsoft.Json.JsonConverter"/>.
+            </summary>
+            <param name="value">The object to serialize.</param>
+            <param name="formatting">Indicates how the output is formatted.</param>
+            <param name="settings">The <see cref="T:Newtonsoft.Json.JsonSerializerSettings"/> used to serialize the object.
+            If this is null, default serialization settings will be is used.</param>
+            <param name="type">
+            The type of the value being serialized.
+            This parameter is used when <see cref="T:Newtonsoft.Json.TypeNameHandling"/> is Auto to write out the type name if the type of the value does not match.
+            Specifing the type is optional.
+            </param>
+            <returns>
+            A JSON string representation of the object.
+            </returns>
+        </member>
+        <member name="M:Newtonsoft.Json.JsonConvert.DeserializeObject(System.String)">
+            <summary>
+            Deserializes the JSON to a .NET object.
+            </summary>
+            <param name="value">The JSON to deserialize.</param>
+            <returns>The deserialized object from the Json string.</returns>
+        </member>
+        <member name="M:Newtonsoft.Json.JsonConvert.DeserializeObject(System.String,Newtonsoft.Json.JsonSerializerSettings)">
+            <summary>
+            Deserializes the JSON to a .NET object.
+            </summary>
+            <param name="value">The JSON to deserialize.</param>
+            <param name="settings">
+            The <see cref="T:Newtonsoft.Json.JsonSerializerSettings"/> used to deserialize the object.
+            If this is null, default serialization settings will be is used.
+            </param>
+            <returns>The deserialized object from the JSON string.</returns>
+        </member>
+        <member name="M:Newtonsoft.Json.JsonConvert.DeserializeObject(System.String,System.Type)">
+            <summary>
+            Deserializes the JSON to the specified .NET type.
+            </summary>
+            <param name="value">The JSON to deserialize.</param>
+            <param name="type">The <see cref="T:System.Type"/> of object being deserialized.</param>
+            <returns>The deserialized object from the Json string.</returns>
+        </member>
+        <member name="M:Newtonsoft.Json.JsonConvert.DeserializeObject``1(System.String)">
+            <summary>
+            Deserializes the JSON to the specified .NET type.
+            </summary>
+            <typeparam name="T">The type of the object to deserialize to.</typeparam>
+            <param name="value">The JSON to deserialize.</param>
+            <returns>The deserialized object from the Json string.</returns>
+        </member>
+        <member name="M:Newtonsoft.Json.JsonConvert.DeserializeAnonymousType``1(System.String,``0)">
+            <summary>
+            Deserializes the JSON to the given anonymous type.
+            </summary>
+            <typeparam name="T">
+            The anonymous type to deserialize to. This can't be specified
+            traditionally and must be infered from the anonymous type passed
+            as a parameter.
+            </typeparam>
+            <param name="value">The JSON to deserialize.</param>
+            <param name="anonymousTypeObject">The anonymous type object.</param>
+            <returns>The deserialized anonymous type from the JSON string.</returns>
+        </member>
+        <member name="M:Newtonsoft.Json.JsonConvert.DeserializeAnonymousType``1(System.String,``0,Newtonsoft.Json.JsonSerializerSettings)">
+            <summary>
+            Deserializes the JSON to the given anonymous type.
+            </summary>
+            <typeparam name="T">
+            The anonymous type to deserialize to. This can't be specified
+            traditionally and must be infered from the anonymous type passed
+            as a parameter.
+            </typeparam>
+            <param name="value">The JSON to deserialize.</param>
+            <param name="anonymousTypeObject">The anonymous type object.</param>
+            <param name="settings">
+            The <see cref="T:Newtonsoft.Json.JsonSerializerSettings"/> used to deserialize the object.
+            If this is null, default serialization settings will be is used.
+            </param>
+            <returns>The deserialized anonymous type from the JSON string.</returns>
+        </member>
+        <member name="M:Newtonsoft.Json.JsonConvert.DeserializeObject``1(System.String,Newtonsoft.Json.JsonConverter[])">
+            <summary>
+            Deserializes the JSON to the specified .NET type.
+            </summary>
+            <typeparam name="T">The type of the object to deserialize to.</typeparam>
+            <param name="value">The JSON to deserialize.</param>
+            <param name="converters">Converters to use while deserializing.</param>
+            <returns>The deserialized object from the JSON string.</returns>
+        </member>
+        <member name="M:Newtonsoft.Json.JsonConvert.DeserializeObject``1(System.String,Newtonsoft.Json.JsonSerializerSettings)">
+            <summary>
+            Deserializes the JSON to the specified .NET type.
+            </summary>
+            <typeparam name="T">The type of the object to deserialize to.</typeparam>
+            <param name="value">The object to deserialize.</param>
+            <param name="settings">
+            The <see cref="T:Newtonsoft.Json.JsonSerializerSettings"/> used to deserialize the object.
+            If this is null, default serialization settings will be is used.
+            </param>
+            <returns>The deserialized object from the JSON string.</returns>
+        </member>
+        <member name="M:Newtonsoft.Json.JsonConvert.DeserializeObject(System.String,System.Type,Newtonsoft.Json.JsonConverter[])">
+            <summary>
+            Deserializes the JSON to the specified .NET type.
+            </summary>
+            <param name="value">The JSON to deserialize.</param>
+            <param name="type">The type of the object to deserialize.</param>
+            <param name="converters">Converters to use while deserializing.</param>
+            <returns>The deserialized object from the JSON string.</returns>
+        </member>
+        <member name="M:Newtonsoft.Json.JsonConvert.DeserializeObject(System.String,System.Type,Newtonsoft.Json.JsonSerializerSettings)">
+            <summary>
+            Deserializes the JSON to the specified .NET type.
+            </summary>
+            <param name="value">The JSON to deserialize.</param>
+            <param name="type">The type of the object to deserialize to.</param>
+            <param name="settings">
+            The <see cref="T:Newtonsoft.Json.JsonSerializerSettings"/> used to deserialize the object.
+            If this is null, default serialization settings will be is used.
+            </param>
+            <returns>The deserialized object from the JSON string.</returns>
+        </member>
+        <member name="M:Newtonsoft.Json.JsonConvert.PopulateObject(System.String,System.Object)">
+            <summary>
+            Populates the object with values from the JSON string.
+            </summary>
+            <param name="value">The JSON to populate values from.</param>
+            <param name="target">The target object to populate values onto.</param>
+        </member>
+        <member name="M:Newtonsoft.Json.JsonConvert.PopulateObject(System.String,System.Object,Newtonsoft.Json.JsonSerializerSettings)">
+            <summary>
+            Populates the object with values from the JSON string.
+            </summary>
+            <param name="value">The JSON to populate values from.</param>
+            <param name="target">The target object to populate values onto.</param>
+            <param name="settings">
+            The <see cref="T:Newtonsoft.Json.JsonSerializerSettings"/> used to deserialize the object.
+            If this is null, default serialization settings will be is used.
+            </param>
+        </member>
+        <member name="M:Newtonsoft.Json.JsonConvert.SerializeXmlNode(System.Xml.XmlNode)">
+            <summary>
+            Serializes the XML node to a JSON string.
+            </summary>
+            <param name="node">The node to serialize.</param>
+            <returns>A JSON string of the XmlNode.</returns>
+        </member>
+        <member name="M:Newtonsoft.Json.JsonConvert.SerializeXmlNode(System.Xml.XmlNode,Newtonsoft.Json.Formatting)">
+            <summary>
+            Serializes the XML node to a JSON string.
+            </summary>
+            <param name="node">The node to serialize.</param>
+            <param name="formatting">Indicates how the output is formatted.</param>
+            <returns>A JSON string of the XmlNode.</returns>
+        </member>
+        <member name="M:Newtonsoft.Json.JsonConvert.SerializeXmlNode(System.Xml.XmlNode,Newtonsoft.Json.Formatting,System.Boolean)">
+            <summary>
+            Serializes the XML node to a JSON string.
+            </summary>
+            <param name="node">The node to serialize.</param>
+            <param name="formatting">Indicates how the output is formatted.</param>
+            <param name="omitRootObject">Omits writing the root object.</param>
+            <returns>A JSON string of the XmlNode.</returns>
+        </member>
+        <member name="M:Newtonsoft.Json.JsonConvert.DeserializeXmlNode(System.String)">
+            <summary>
+            Deserializes the XmlNode from a JSON string.
+            </summary>
+            <param name="value">The JSON string.</param>
+            <returns>The deserialized XmlNode</returns>
+        </member>
+        <member name="M:Newtonsoft.Json.JsonConvert.DeserializeXmlNode(System.String,System.String)">
+            <summary>
+            Deserializes the XmlNode from a JSON string nested in a root elment.
+            </summary>
+            <param name="value">The JSON string.</param>
+            <param name="deserializeRootElementName">The name of the root element to append when deserializing.</param>
+            <returns>The deserialized XmlNode</returns>
+        </member>
+        <member name="M:Newtonsoft.Json.JsonConvert.DeserializeXmlNode(System.String,System.String,System.Boolean)">
+            <summary>
+            Deserializes the XmlNode from a JSON string nested in a root elment.
+            </summary>
+            <param name="value">The JSON string.</param>
+            <param name="deserializeRootElementName">The name of the root element to append when deserializing.</param>
+            <param name="writeArrayAttribute">
+            A flag to indicate whether to write the Json.NET array attribute.
+            This attribute helps preserve arrays when converting the written XML back to JSON.
+            </param>
+            <returns>The deserialized XmlNode</returns>
+        </member>
+        <member name="T:Newtonsoft.Json.JsonSerializationException">
+            <summary>
+            The exception thrown when an error occurs during Json serialization or deserialization.
+            </summary>
+        </member>
+        <member name="M:Newtonsoft.Json.JsonSerializationException.#ctor">
+            <summary>
+            Initializes a new instance of the <see cref="T:Newtonsoft.Json.JsonSerializationException"/> class.
+            </summary>
+        </member>
+        <member name="M:Newtonsoft.Json.JsonSerializationException.#ctor(System.String)">
+            <summary>
+            Initializes a new instance of the <see cref="T:Newtonsoft.Json.JsonSerializationException"/> class
+            with a specified error message.
+            </summary>
+            <param name="message">The error message that explains the reason for the exception.</param>
+        </member>
+        <member name="M:Newtonsoft.Json.JsonSerializationException.#ctor(System.String,System.Exception)">
+            <summary>
+            Initializes a new instance of the <see cref="T:Newtonsoft.Json.JsonSerializationException"/> class
+            with a specified error message and a reference to the inner exception that is the cause of this exception.
+            </summary>
+            <param name="message">The error message that explains the reason for the exception.</param>
+            <param name="innerException">The exception that is the cause of the current exception, or a null reference (Nothing in Visual Basic) if no inner exception is specified.</param>
+        </member>
+        <member name="M:Newtonsoft.Json.JsonSerializationException.#ctor(System.Runtime.Serialization.SerializationInfo,System.Runtime.Serialization.StreamingContext)">
+            <summary>
+            Initializes a new instance of the <see cref="T:Newtonsoft.Json.JsonSerializationException"/> class.
+            </summary>
+            <param name="info">The <see cref="T:System.Runtime.Serialization.SerializationInfo"/> that holds the serialized object data about the exception being thrown.</param>
+            <param name="context">The <see cref="T:System.Runtime.Serialization.StreamingContext"/> that contains contextual information about the source or destination.</param>
+            <exception cref="T:System.ArgumentNullException">The <paramref name="info"/> parameter is null. </exception>
+            <exception cref="T:System.Runtime.Serialization.SerializationException">The class name is null or <see cref="P:System.Exception.HResult"/> is zero (0). </exception>
+        </member>
+        <member name="T:Newtonsoft.Json.JsonSerializer">
+            <summary>
+            Serializes and deserializes objects into and from the JSON format.
+            The <see cref="T:Newtonsoft.Json.JsonSerializer"/> enables you to control how objects are encoded into JSON.
+            </summary>
+        </member>
+        <member name="M:Newtonsoft.Json.JsonSerializer.#ctor">
+            <summary>
+            Initializes a new instance of the <see cref="T:Newtonsoft.Json.JsonSerializer"/> class.
+            </summary>
+        </member>
+        <member name="M:Newtonsoft.Json.JsonSerializer.Create(Newtonsoft.Json.JsonSerializerSettings)">
+            <summary>
+            Creates a new <see cref="T:Newtonsoft.Json.JsonSerializer"/> instance using the specified <see cref="T:Newtonsoft.Json.JsonSerializerSettings"/>.
+            </summary>
+            <param name="settings">The settings to be applied to the <see cref="T:Newtonsoft.Json.JsonSerializer"/>.</param>
+            <returns>A new <see cref="T:Newtonsoft.Json.JsonSerializer"/> instance using the specified <see cref="T:Newtonsoft.Json.JsonSerializerSettings"/>.</returns>
+        </member>
+        <member name="M:Newtonsoft.Json.JsonSerializer.Populate(System.IO.TextReader,System.Object)">
+            <summary>
+            Populates the JSON values onto the target object.
+            </summary>
+            <param name="reader">The <see cref="T:System.IO.TextReader"/> that contains the JSON structure to reader values from.</param>
+            <param name="target">The target object to populate values onto.</param>
+        </member>
+        <member name="M:Newtonsoft.Json.JsonSerializer.Populate(Newtonsoft.Json.JsonReader,System.Object)">
+            <summary>
+            Populates the JSON values onto the target object.
+            </summary>
+            <param name="reader">The <see cref="T:Newtonsoft.Json.JsonReader"/> that contains the JSON structure to reader values from.</param>
+            <param name="target">The target object to populate values onto.</param>
+        </member>
+        <member name="M:Newtonsoft.Json.JsonSerializer.Deserialize(Newtonsoft.Json.JsonReader)">
+            <summary>
+            Deserializes the Json structure contained by the specified <see cref="T:Newtonsoft.Json.JsonReader"/>.
+            </summary>
+            <param name="reader">The <see cref="T:Newtonsoft.Json.JsonReader"/> that contains the JSON structure to deserialize.</param>
+            <returns>The <see cref="T:System.Object"/> being deserialized.</returns>
+        </member>
+        <member name="M:Newtonsoft.Json.JsonSerializer.Deserialize(System.IO.TextReader,System.Type)">
+            <summary>
+            Deserializes the Json structure contained by the specified <see cref="T:System.IO.StringReader"/>
+            into an instance of the specified type.
+            </summary>
+            <param name="reader">The <see cref="T:System.IO.TextReader"/> containing the object.</param>
+            <param name="objectType">The <see cref="T:System.Type"/> of object being deserialized.</param>
+            <returns>The instance of <paramref name="objectType"/> being deserialized.</returns>
+        </member>
+        <member name="M:Newtonsoft.Json.JsonSerializer.Deserialize``1(Newtonsoft.Json.JsonReader)">
+            <summary>
+            Deserializes the Json structure contained by the specified <see cref="T:Newtonsoft.Json.JsonReader"/>
+            into an instance of the specified type.
+            </summary>
+            <param name="reader">The <see cref="T:Newtonsoft.Json.JsonReader"/> containing the object.</param>
+            <typeparam name="T">The type of the object to deserialize.</typeparam>
+            <returns>The instance of <typeparamref name="T"/> being deserialized.</returns>
+        </member>
+        <member name="M:Newtonsoft.Json.JsonSerializer.Deserialize(Newtonsoft.Json.JsonReader,System.Type)">
+            <summary>
+            Deserializes the Json structure contained by the specified <see cref="T:Newtonsoft.Json.JsonReader"/>
+            into an instance of the specified type.
+            </summary>
+            <param name="reader">The <see cref="T:Newtonsoft.Json.JsonReader"/> containing the object.</param>
+            <param name="objectType">The <see cref="T:System.Type"/> of object being deserialized.</param>
+            <returns>The instance of <paramref name="objectType"/> being deserialized.</returns>
+        </member>
+        <member name="M:Newtonsoft.Json.JsonSerializer.Serialize(System.IO.TextWriter,System.Object)">
+            <summary>
+            Serializes the specified <see cref="T:System.Object"/> and writes the Json structure
+            to a <c>Stream</c> using the specified <see cref="T:System.IO.TextWriter"/>. 
+            </summary>
+            <param name="textWriter">The <see cref="T:System.IO.TextWriter"/> used to write the Json structure.</param>
+            <param name="value">The <see cref="T:System.Object"/> to serialize.</param>
+        </member>
+        <member name="M:Newtonsoft.Json.JsonSerializer.Serialize(Newtonsoft.Json.JsonWriter,System.Object,System.Type)">
+            <summary>
+            Serializes the specified <see cref="T:System.Object"/> and writes the Json structure
+            to a <c>Stream</c> using the specified <see cref="T:System.IO.TextWriter"/>. 
+            </summary>
+            <param name="jsonWriter">The <see cref="T:Newtonsoft.Json.JsonWriter"/> used to write the Json structure.</param>
+            <param name="value">The <see cref="T:System.Object"/> to serialize.</param>
+            <param name="objectType">
+            The type of the value being serialized.
+            This parameter is used when <see cref="P:Newtonsoft.Json.JsonSerializer.TypeNameHandling"/> is Auto to write out the type name if the type of the value does not match.
+            Specifing the type is optional.
+            </param>
+        </member>
+        <member name="M:Newtonsoft.Json.JsonSerializer.Serialize(System.IO.TextWriter,System.Object,System.Type)">
+            <summary>
+            Serializes the specified <see cref="T:System.Object"/> and writes the Json structure
+            to a <c>Stream</c> using the specified <see cref="T:System.IO.TextWriter"/>. 
+            </summary>
+            <param name="textWriter">The <see cref="T:System.IO.TextWriter"/> used to write the Json structure.</param>
+            <param name="value">The <see cref="T:System.Object"/> to serialize.</param>
+            <param name="objectType">
+            The type of the value being serialized.
+            This parameter is used when <see cref="P:Newtonsoft.Json.JsonSerializer.TypeNameHandling"/> is Auto to write out the type name if the type of the value does not match.
+            Specifing the type is optional.
+            </param>
+        </member>
+        <member name="M:Newtonsoft.Json.JsonSerializer.Serialize(Newtonsoft.Json.JsonWriter,System.Object)">
+            <summary>
+            Serializes the specified <see cref="T:System.Object"/> and writes the Json structure
+            to a <c>Stream</c> using the specified <see cref="T:Newtonsoft.Json.JsonWriter"/>. 
+            </summary>
+            <param name="jsonWriter">The <see cref="T:Newtonsoft.Json.JsonWriter"/> used to write the Json structure.</param>
+            <param name="value">The <see cref="T:System.Object"/> to serialize.</param>
+        </member>
+        <member name="E:Newtonsoft.Json.JsonSerializer.Error">
+            <summary>
+            Occurs when the <see cref="T:Newtonsoft.Json.JsonSerializer"/> errors during serialization and deserialization.
+            </summary>
+        </member>
+        <member name="P:Newtonsoft.Json.JsonSerializer.ReferenceResolver">
+            <summary>
+            Gets or sets the <see cref="T:Newtonsoft.Json.Serialization.IReferenceResolver"/> used by the serializer when resolving references.
+            </summary>
+        </member>
+        <member name="P:Newtonsoft.Json.JsonSerializer.Binder">
+            <summary>
+            Gets or sets the <see cref="T:System.Runtime.Serialization.SerializationBinder"/> used by the serializer when resolving type names.
+            </summary>
+        </member>
+        <member name="P:Newtonsoft.Json.JsonSerializer.TraceWriter">
+            <summary>
+            Gets or sets the <see cref="T:Newtonsoft.Json.Serialization.ITraceWriter"/> used by the serializer when writing trace messages.
+            </summary>
+            <value>The trace writer.</value>
+        </member>
+        <member name="P:Newtonsoft.Json.JsonSerializer.TypeNameHandling">
+            <summary>
+            Gets or sets how type name writing and reading is handled by the serializer.
+            </summary>
+        </member>
+        <member name="P:Newtonsoft.Json.JsonSerializer.TypeNameAssemblyFormat">
+            <summary>
+            Gets or sets how a type name assembly is written and resolved by the serializer.
+            </summary>
+            <value>The type name assembly format.</value>
+        </member>
+        <member name="P:Newtonsoft.Json.JsonSerializer.PreserveReferencesHandling">
+            <summary>
+            Gets or sets how object references are preserved by the serializer.
+            </summary>
+        </member>
+        <member name="P:Newtonsoft.Json.JsonSerializer.ReferenceLoopHandling">
+            <summary>
+            Get or set how reference loops (e.g. a class referencing itself) is handled.
+            </summary>
+        </member>
+        <member name="P:Newtonsoft.Json.JsonSerializer.MissingMemberHandling">
+            <summary>
+            Get or set how missing members (e.g. JSON contains a property that isn't a member on the object) are handled during deserialization.
+            </summary>
+        </member>
+        <member name="P:Newtonsoft.Json.JsonSerializer.NullValueHandling">
+            <summary>
+            Get or set how null values are handled during serialization and deserialization.
+            </summary>
+        </member>
+        <member name="P:Newtonsoft.Json.JsonSerializer.DefaultValueHandling">
+            <summary>
+            Get or set how null default are handled during serialization and deserialization.
+            </summary>
+        </member>
+        <member name="P:Newtonsoft.Json.JsonSerializer.ObjectCreationHandling">
+            <summary>
+            Gets or sets how objects are created during deserialization.
+            </summary>
+            <value>The object creation handling.</value>
+        </member>
+        <member name="P:Newtonsoft.Json.JsonSerializer.ConstructorHandling">
+            <summary>
+            Gets or sets how constructors are used during deserialization.
+            </summary>
+            <value>The constructor handling.</value>
+        </member>
+        <member name="P:Newtonsoft.Json.JsonSerializer.Converters">
+            <summary>
+            Gets a collection <see cref="T:Newtonsoft.Json.JsonConverter"/> that will be used during serialization.
+            </summary>
+            <value>Collection <see cref="T:Newtonsoft.Json.JsonConverter"/> that will be used during serialization.</value>
+        </member>
+        <member name="P:Newtonsoft.Json.JsonSerializer.ContractResolver">
+            <summary>
+            Gets or sets the contract resolver used by the serializer when
+            serializing .NET objects to JSON and vice versa.
+            </summary>
+        </member>
+        <member name="P:Newtonsoft.Json.JsonSerializer.Context">
+            <summary>
+            Gets or sets the <see cref="T:System.Runtime.Serialization.StreamingContext"/> used by the serializer when invoking serialization callback methods.
+            </summary>
+            <value>The context.</value>
+        </member>
+        <member name="P:Newtonsoft.Json.JsonSerializer.Formatting">
+            <summary>
+            Indicates how JSON text output is formatted.
+            </summary>
+        </member>
+        <member name="P:Newtonsoft.Json.JsonSerializer.DateFormatHandling">
+            <summary>
+            Get or set how dates are written to JSON text.
+            </summary>
+        </member>
+        <member name="P:Newtonsoft.Json.JsonSerializer.DateTimeZoneHandling">
+            <summary>
+            Get or set how <see cref="T:System.DateTime"/> time zones are handling during serialization and deserialization.
+            </summary>
+        </member>
+        <member name="P:Newtonsoft.Json.JsonSerializer.DateParseHandling">
+            <summary>
+            Get or set how date formatted strings, e.g. "\/Date(1198908717056)\/" and "2012-03-21T05:40Z", are parsed when reading JSON.
+            </summary>
+        </member>
+        <member name="P:Newtonsoft.Json.JsonSerializer.FloatParseHandling">
+            <summary>
+            Get or set how floating point numbers, e.g. 1.0 and 9.9, are parsed when reading JSON text.
+            </summary>
+        </member>
+        <member name="P:Newtonsoft.Json.JsonSerializer.FloatFormatHandling">
+            <summary>
+            Get or set how special floating point numbers, e.g. <see cref="F:System.Double.NaN"/>,
+            <see cref="F:System.Double.PositiveInfinity"/> and <see cref="F:System.Double.NegativeInfinity"/>,
+            are written as JSON text.
+            </summary>
+        </member>
+        <member name="P:Newtonsoft.Json.JsonSerializer.StringEscapeHandling">
+            <summary>
+            Get or set how strings are escaped when writing JSON text.
+            </summary>
+        </member>
+        <member name="P:Newtonsoft.Json.JsonSerializer.DateFormatString">
+            <summary>
+            Get or set how <see cref="T:System.DateTime"/> and <see cref="T:System.DateTimeOffset"/> values are formatting when writing JSON text.
+            </summary>
+        </member>
+        <member name="P:Newtonsoft.Json.JsonSerializer.Culture">
+            <summary>
+            Gets or sets the culture used when reading JSON. Defaults to <see cref="P:System.Globalization.CultureInfo.InvariantCulture"/>.
+            </summary>
+        </member>
+        <member name="P:Newtonsoft.Json.JsonSerializer.MaxDepth">
+            <summary>
+            Gets or sets the maximum depth allowed when reading JSON. Reading past this depth will throw a <see cref="T:Newtonsoft.Json.JsonReaderException"/>.
+            </summary>
+        </member>
+        <member name="P:Newtonsoft.Json.JsonSerializer.CheckAdditionalContent">
+            <summary>
+            Gets a value indicating whether there will be a check for additional JSON content after deserializing an object.
+            </summary>
+            <value>
+            	<c>true</c> if there will be a check for additional JSON content after deserializing an object; otherwise, <c>false</c>.
+            </value>
+        </member>
+        <member name="T:Newtonsoft.Json.Linq.Extensions">
+            <summary>
+            Contains the LINQ to JSON extension methods.
+            </summary>
+        </member>
+        <member name="M:Newtonsoft.Json.Linq.Extensions.Ancestors``1(System.Collections.Generic.IEnumerable{``0})">
+            <summary>
+            Returns a collection of tokens that contains the ancestors of every token in the source collection.
+            </summary>
+            <typeparam name="T">The type of the objects in source, constrained to <see cref="T:Newtonsoft.Json.Linq.JToken"/>.</typeparam>
+            <param name="source">An <see cref="T:System.Collections.Generic.IEnumerable`1"/> of <see cref="T:Newtonsoft.Json.Linq.JToken"/> that contains the source collection.</param>
+            <returns>An <see cref="T:System.Collections.Generic.IEnumerable`1"/> of <see cref="T:Newtonsoft.Json.Linq.JToken"/> that contains the ancestors of every node in the source collection.</returns>
+        </member>
+        <member name="M:Newtonsoft.Json.Linq.Extensions.Descendants``1(System.Collections.Generic.IEnumerable{``0})">
+            <summary>
+            Returns a collection of tokens that contains the descendants of every token in the source collection.
+            </summary>
+            <typeparam name="T">The type of the objects in source, constrained to <see cref="T:Newtonsoft.Json.Linq.JContainer"/>.</typeparam>
+            <param name="source">An <see cref="T:System.Collections.Generic.IEnumerable`1"/> of <see cref="T:Newtonsoft.Json.Linq.JToken"/> that contains the source collection.</param>
+            <returns>An <see cref="T:System.Collections.Generic.IEnumerable`1"/> of <see cref="T:Newtonsoft.Json.Linq.JToken"/> that contains the descendants of every node in the source collection.</returns>
+        </member>
+        <member name="M:Newtonsoft.Json.Linq.Extensions.Properties(System.Collections.Generic.IEnumerable{Newtonsoft.Json.Linq.JObject})">
+            <summary>
+            Returns a collection of child properties of every object in the source collection.
+            </summary>
+            <param name="source">An <see cref="T:System.Collections.Generic.IEnumerable`1"/> of <see cref="T:Newtonsoft.Json.Linq.JObject"/> that contains the source collection.</param>
+            <returns>An <see cref="T:System.Collections.Generic.IEnumerable`1"/> of <see cref="T:Newtonsoft.Json.Linq.JProperty"/> that contains the properties of every object in the source collection.</returns>
+        </member>
+        <member name="M:Newtonsoft.Json.Linq.Extensions.Values(System.Collections.Generic.IEnumerable{Newtonsoft.Json.Linq.JToken},System.Object)">
+            <summary>
+            Returns a collection of child values of every object in the source collection with the given key.
+            </summary>
+            <param name="source">An <see cref="T:System.Collections.Generic.IEnumerable`1"/> of <see cref="T:Newtonsoft.Json.Linq.JToken"/> that contains the source collection.</param>
+            <param name="key">The token key.</param>
+            <returns>An <see cref="T:System.Collections.Generic.IEnumerable`1"/> of <see cref="T:Newtonsoft.Json.Linq.JToken"/> that contains the values of every node in the source collection with the given key.</returns>
+        </member>
+        <member name="M:Newtonsoft.Json.Linq.Extensions.Values(System.Collections.Generic.IEnumerable{Newtonsoft.Json.Linq.JToken})">
+            <summary>
+            Returns a collection of child values of every object in the source collection.
+            </summary>
+            <param name="source">An <see cref="T:System.Collections.Generic.IEnumerable`1"/> of <see cref="T:Newtonsoft.Json.Linq.JToken"/> that contains the source collection.</param>
+            <returns>An <see cref="T:System.Collections.Generic.IEnumerable`1"/> of <see cref="T:Newtonsoft.Json.Linq.JToken"/> that contains the values of every node in the source collection.</returns>
+        </member>
+        <member name="M:Newtonsoft.Json.Linq.Extensions.Values``1(System.Collections.Generic.IEnumerable{Newtonsoft.Json.Linq.JToken},System.Object)">
+            <summary>
+            Returns a collection of converted child values of every object in the source collection with the given key.
+            </summary>
+            <typeparam name="U">The type to convert the values to.</typeparam>
+            <param name="source">An <see cref="T:System.Collections.Generic.IEnumerable`1"/> of <see cref="T:Newtonsoft.Json.Linq.JToken"/> that contains the source collection.</param>
+            <param name="key">The token key.</param>
+            <returns>An <see cref="T:System.Collections.Generic.IEnumerable`1"/> that contains the converted values of every node in the source collection with the given key.</returns>
+        </member>
+        <member name="M:Newtonsoft.Json.Linq.Extensions.Values``1(System.Collections.Generic.IEnumerable{Newtonsoft.Json.Linq.JToken})">
+            <summary>
+            Returns a collection of converted child values of every object in the source collection.
+            </summary>
+            <typeparam name="U">The type to convert the values to.</typeparam>
+            <param name="source">An <see cref="T:System.Collections.Generic.IEnumerable`1"/> of <see cref="T:Newtonsoft.Json.Linq.JToken"/> that contains the source collection.</param>
+            <returns>An <see cref="T:System.Collections.Generic.IEnumerable`1"/> that contains the converted values of every node in the source collection.</returns>
+        </member>
+        <member name="M:Newtonsoft.Json.Linq.Extensions.Value``1(System.Collections.Generic.IEnumerable{Newtonsoft.Json.Linq.JToken})">
+            <summary>
+            Converts the value.
+            </summary>
+            <typeparam name="U">The type to convert the value to.</typeparam>
+            <param name="value">A <see cref="T:Newtonsoft.Json.Linq.JToken"/> cast as a <see cref="T:System.Collections.Generic.IEnumerable`1"/> of <see cref="T:Newtonsoft.Json.Linq.JToken"/>.</param>
+            <returns>A converted value.</returns>
+        </member>
+        <member name="M:Newtonsoft.Json.Linq.Extensions.Value``2(System.Collections.Generic.IEnumerable{``0})">
+            <summary>
+            Converts the value.
+            </summary>
+            <typeparam name="T">The source collection type.</typeparam>
+            <typeparam name="U">The type to convert the value to.</typeparam>
+            <param name="value">A <see cref="T:Newtonsoft.Json.Linq.JToken"/> cast as a <see cref="T:System.Collections.Generic.IEnumerable`1"/> of <see cref="T:Newtonsoft.Json.Linq.JToken"/>.</param>
+            <returns>A converted value.</returns>
+        </member>
+        <member name="M:Newtonsoft.Json.Linq.Extensions.Children``1(System.Collections.Generic.IEnumerable{``0})">
+            <summary>
+            Returns a collection of child tokens of every array in the source collection.
+            </summary>
+            <typeparam name="T">The source collection type.</typeparam>
+            <param name="source">An <see cref="T:System.Collections.Generic.IEnumerable`1"/> of <see cref="T:Newtonsoft.Json.Linq.JToken"/> that contains the source collection.</param>
+            <returns>An <see cref="T:System.Collections.Generic.IEnumerable`1"/> of <see cref="T:Newtonsoft.Json.Linq.JToken"/> that contains the values of every node in the source collection.</returns>
+        </member>
+        <member name="M:Newtonsoft.Json.Linq.Extensions.Children``2(System.Collections.Generic.IEnumerable{``0})">
+            <summary>
+            Returns a collection of converted child tokens of every array in the source collection.
+            </summary>
+            <param name="source">An <see cref="T:System.Collections.Generic.IEnumerable`1"/> of <see cref="T:Newtonsoft.Json.Linq.JToken"/> that contains the source collection.</param>
+            <typeparam name="U">The type to convert the values to.</typeparam>
+            <typeparam name="T">The source collection type.</typeparam>
+            <returns>An <see cref="T:System.Collections.Generic.IEnumerable`1"/> that contains the converted values of every node in the source collection.</returns>
+        </member>
+        <member name="M:Newtonsoft.Json.Linq.Extensions.AsJEnumerable(System.Collections.Generic.IEnumerable{Newtonsoft.Json.Linq.JToken})">
+            <summary>
+            Returns the input typed as <see cref="T:Newtonsoft.Json.Linq.IJEnumerable`1"/>.
+            </summary>
+            <param name="source">An <see cref="T:System.Collections.Generic.IEnumerable`1"/> of <see cref="T:Newtonsoft.Json.Linq.JToken"/> that contains the source collection.</param>
+            <returns>The input typed as <see cref="T:Newtonsoft.Json.Linq.IJEnumerable`1"/>.</returns>
+        </member>
+        <member name="M:Newtonsoft.Json.Linq.Extensions.AsJEnumerable``1(System.Collections.Generic.IEnumerable{``0})">
+            <summary>
+            Returns the input typed as <see cref="T:Newtonsoft.Json.Linq.IJEnumerable`1"/>.
+            </summary>
+            <typeparam name="T">The source collection type.</typeparam>
+            <param name="source">An <see cref="T:System.Collections.Generic.IEnumerable`1"/> of <see cref="T:Newtonsoft.Json.Linq.JToken"/> that contains the source collection.</param>
+            <returns>The input typed as <see cref="T:Newtonsoft.Json.Linq.IJEnumerable`1"/>.</returns>
+        </member>
+        <member name="T:Newtonsoft.Json.Linq.JConstructor">
+            <summary>
+            Represents a JSON constructor.
+            </summary>
+        </member>
+        <member name="T:Newtonsoft.Json.Linq.JContainer">
+            <summary>
+            Represents a token that can contain other tokens.
+            </summary>
+        </member>
+        <member name="M:Newtonsoft.Json.Linq.JContainer.OnAddingNew(System.ComponentModel.AddingNewEventArgs)">
+            <summary>
+            Raises the <see cref="E:Newtonsoft.Json.Linq.JContainer.AddingNew"/> event.
+            </summary>
+            <param name="e">The <see cref="T:System.ComponentModel.AddingNewEventArgs"/> instance containing the event data.</param>
+        </member>
+        <member name="M:Newtonsoft.Json.Linq.JContainer.OnListChanged(System.ComponentModel.ListChangedEventArgs)">
+            <summary>
+            Raises the <see cref="E:Newtonsoft.Json.Linq.JContainer.ListChanged"/> event.
+            </summary>
+            <param name="e">The <see cref="T:System.ComponentModel.ListChangedEventArgs"/> instance containing the event data.</param>
+        </member>
+        <member name="M:Newtonsoft.Json.Linq.JContainer.Children">
+            <summary>
+            Returns a collection of the child tokens of this token, in document order.
+            </summary>
+            <returns>
+            An <see cref="T:System.Collections.Generic.IEnumerable`1"/> of <see cref="T:Newtonsoft.Json.Linq.JToken"/> containing the child tokens of this <see cref="T:Newtonsoft.Json.Linq.JToken"/>, in document order.
+            </returns>
+        </member>
+        <member name="M:Newtonsoft.Json.Linq.JContainer.Values``1">
+            <summary>
+            Returns a collection of the child values of this token, in document order.
+            </summary>
+            <typeparam name="T">The type to convert the values to.</typeparam>
+            <returns>
+            A <see cref="T:System.Collections.Generic.IEnumerable`1"/> containing the child values of this <see cref="T:Newtonsoft.Json.Linq.JToken"/>, in document order.
+            </returns>
+        </member>
+        <member name="M:Newtonsoft.Json.Linq.JContainer.Descendants">
+            <summary>
+            Returns a collection of the descendant tokens for this token in document order.
+            </summary>
+            <returns>An <see cref="T:System.Collections.Generic.IEnumerable`1"/> containing the descendant tokens of the <see cref="T:Newtonsoft.Json.Linq.JToken"/>.</returns>
+        </member>
+        <member name="M:Newtonsoft.Json.Linq.JContainer.Add(System.Object)">
+            <summary>
+            Adds the specified content as children of this <see cref="T:Newtonsoft.Json.Linq.JToken"/>.
+            </summary>
+            <param name="content">The content to be added.</param>
+        </member>
+        <member name="M:Newtonsoft.Json.Linq.JContainer.AddFirst(System.Object)">
+            <summary>
+            Adds the specified content as the first children of this <see cref="T:Newtonsoft.Json.Linq.JToken"/>.
+            </summary>
+            <param name="content">The content to be added.</param>
+        </member>
+        <member name="M:Newtonsoft.Json.Linq.JContainer.CreateWriter">
+            <summary>
+            Creates an <see cref="T:Newtonsoft.Json.JsonWriter"/> that can be used to add tokens to the <see cref="T:Newtonsoft.Json.Linq.JToken"/>.
+            </summary>
+            <returns>An <see cref="T:Newtonsoft.Json.JsonWriter"/> that is ready to have content written to it.</returns>
+        </member>
+        <member name="M:Newtonsoft.Json.Linq.JContainer.ReplaceAll(System.Object)">
+            <summary>
+            Replaces the children nodes of this token with the specified content.
+            </summary>
+            <param name="content">The content.</param>
+        </member>
+        <member name="M:Newtonsoft.Json.Linq.JContainer.RemoveAll">
+            <summary>
+            Removes the child nodes from this token.
+            </summary>
+        </member>
+        <member name="E:Newtonsoft.Json.Linq.JContainer.ListChanged">
+            <summary>
+            Occurs when the list changes or an item in the list changes.
+            </summary>
+        </member>
+        <member name="E:Newtonsoft.Json.Linq.JContainer.AddingNew">
+            <summary>
+            Occurs before an item is added to the collection.
+            </summary>
+        </member>
+        <member name="P:Newtonsoft.Json.Linq.JContainer.ChildrenTokens">
+            <summary>
+            Gets the container's children tokens.
+            </summary>
+            <value>The container's children tokens.</value>
+        </member>
+        <member name="P:Newtonsoft.Json.Linq.JContainer.HasValues">
+            <summary>
+            Gets a value indicating whether this token has childen tokens.
+            </summary>
+            <value>
+            	<c>true</c> if this token has child values; otherwise, <c>false</c>.
+            </value>
+        </member>
+        <member name="P:Newtonsoft.Json.Linq.JContainer.First">
+            <summary>
+            Get the first child token of this token.
+            </summary>
+            <value>
+            A <see cref="T:Newtonsoft.Json.Linq.JToken"/> containing the first child token of the <see cref="T:Newtonsoft.Json.Linq.JToken"/>.
+            </value>
+        </member>
+        <member name="P:Newtonsoft.Json.Linq.JContainer.Last">
+            <summary>
+            Get the last child token of this token.
+            </summary>
+            <value>
+            A <see cref="T:Newtonsoft.Json.Linq.JToken"/> containing the last child token of the <see cref="T:Newtonsoft.Json.Linq.JToken"/>.
+            </value>
+        </member>
+        <member name="P:Newtonsoft.Json.Linq.JContainer.Count">
+            <summary>
+            Gets the count of child JSON tokens.
+            </summary>
+            <value>The count of child JSON tokens</value>
+        </member>
+        <member name="M:Newtonsoft.Json.Linq.JConstructor.#ctor">
+            <summary>
+            Initializes a new instance of the <see cref="T:Newtonsoft.Json.Linq.JConstructor"/> class.
+            </summary>
+        </member>
+        <member name="M:Newtonsoft.Json.Linq.JConstructor.#ctor(Newtonsoft.Json.Linq.JConstructor)">
+            <summary>
+            Initializes a new instance of the <see cref="T:Newtonsoft.Json.Linq.JConstructor"/> class from another <see cref="T:Newtonsoft.Json.Linq.JConstructor"/> object.
+            </summary>
+            <param name="other">A <see cref="T:Newtonsoft.Json.Linq.JConstructor"/> object to copy from.</param>
+        </member>
+        <member name="M:Newtonsoft.Json.Linq.JConstructor.#ctor(System.String,System.Object[])">
+            <summary>
+            Initializes a new instance of the <see cref="T:Newtonsoft.Json.Linq.JConstructor"/> class with the specified name and content.
+            </summary>
+            <param name="name">The constructor name.</param>
+            <param name="content">The contents of the constructor.</param>
+        </member>
+        <member name="M:Newtonsoft.Json.Linq.JConstructor.#ctor(System.String,System.Object)">
+            <summary>
+            Initializes a new instance of the <see cref="T:Newtonsoft.Json.Linq.JConstructor"/> class with the specified name and content.
+            </summary>
+            <param name="name">The constructor name.</param>
+            <param name="content">The contents of the constructor.</param>
+        </member>
+        <member name="M:Newtonsoft.Json.Linq.JConstructor.#ctor(System.String)">
+            <summary>
+            Initializes a new instance of the <see cref="T:Newtonsoft.Json.Linq.JConstructor"/> class with the specified name.
+            </summary>
+            <param name="name">The constructor name.</param>
+        </member>
+        <member name="M:Newtonsoft.Json.Linq.JConstructor.WriteTo(Newtonsoft.Json.JsonWriter,Newtonsoft.Json.JsonConverter[])">
+            <summary>
+            Writes this token to a <see cref="T:Newtonsoft.Json.JsonWriter"/>.
+            </summary>
+            <param name="writer">A <see cref="T:Newtonsoft.Json.JsonWriter"/> into which this method will write.</param>
+            <param name="converters">A collection of <see cref="T:Newtonsoft.Json.JsonConverter"/> which will be used when writing the token.</param>
+        </member>
+        <member name="M:Newtonsoft.Json.Linq.JConstructor.Load(Newtonsoft.Json.JsonReader)">
+            <summary>
+            Loads an <see cref="T:Newtonsoft.Json.Linq.JConstructor"/> from a <see cref="T:Newtonsoft.Json.JsonReader"/>. 
+            </summary>
+            <param name="reader">A <see cref="T:Newtonsoft.Json.JsonReader"/> that will be read for the content of the <see cref="T:Newtonsoft.Json.Linq.JConstructor"/>.</param>
+            <returns>A <see cref="T:Newtonsoft.Json.Linq.JConstructor"/> that contains the JSON that was read from the specified <see cref="T:Newtonsoft.Json.JsonReader"/>.</returns>
+        </member>
+        <member name="P:Newtonsoft.Json.Linq.JConstructor.ChildrenTokens">
+            <summary>
+            Gets the container's children tokens.
+            </summary>
+            <value>The container's children tokens.</value>
+        </member>
+        <member name="P:Newtonsoft.Json.Linq.JConstructor.Name">
+            <summary>
+            Gets or sets the name of this constructor.
+            </summary>
+            <value>The constructor name.</value>
+        </member>
+        <member name="P:Newtonsoft.Json.Linq.JConstructor.Type">
+            <summary>
+            Gets the node type for this <see cref="T:Newtonsoft.Json.Linq.JToken"/>.
+            </summary>
+            <value>The type.</value>
+        </member>
+        <member name="P:Newtonsoft.Json.Linq.JConstructor.Item(System.Object)">
+            <summary>
+            Gets the <see cref="T:Newtonsoft.Json.Linq.JToken"/> with the specified key.
+            </summary>
+            <value>The <see cref="T:Newtonsoft.Json.Linq.JToken"/> with the specified key.</value>
+        </member>
+        <member name="T:Newtonsoft.Json.Linq.JEnumerable`1">
+            <summary>
+            Represents a collection of <see cref="T:Newtonsoft.Json.Linq.JToken"/> objects.
+            </summary>
+            <typeparam name="T">The type of token</typeparam>
+        </member>
+        <member name="F:Newtonsoft.Json.Linq.JEnumerable`1.Empty">
+            <summary>
+            An empty collection of <see cref="T:Newtonsoft.Json.Linq.JToken"/> objects.
+            </summary>
+        </member>
+        <member name="M:Newtonsoft.Json.Linq.JEnumerable`1.#ctor(System.Collections.Generic.IEnumerable{`0})">
+            <summary>
+            Initializes a new instance of the <see cref="T:Newtonsoft.Json.Linq.JEnumerable`1"/> struct.
+            </summary>
+            <param name="enumerable">The enumerable.</param>
+        </member>
+        <member name="M:Newtonsoft.Json.Linq.JEnumerable`1.GetEnumerator">
+            <summary>
+            Returns an enumerator that iterates through the collection.
+            </summary>
+            <returns>
+            A <see cref="T:System.Collections.Generic.IEnumerator`1"/> that can be used to iterate through the collection.
+            </returns>
+        </member>
+        <member name="M:Newtonsoft.Json.Linq.JEnumerable`1.System#Collections#IEnumerable#GetEnumerator">
+            <summary>
+            Returns an enumerator that iterates through a collection.
+            </summary>
+            <returns>
+            An <see cref="T:System.Collections.IEnumerator"/> object that can be used to iterate through the collection.
+            </returns>
+        </member>
+        <member name="M:Newtonsoft.Json.Linq.JEnumerable`1.Equals(System.Object)">
+            <summary>
+            Determines whether the specified <see cref="T:System.Object"/> is equal to this instance.
+            </summary>
+            <param name="obj">The <see cref="T:System.Object"/> to compare with this instance.</param>
+            <returns>
+            	<c>true</c> if the specified <see cref="T:System.Object"/> is equal to this instance; otherwise, <c>false</c>.
+            </returns>
+        </member>
+        <member name="M:Newtonsoft.Json.Linq.JEnumerable`1.GetHashCode">
+            <summary>
+            Returns a hash code for this instance.
+            </summary>
+            <returns>
+            A hash code for this instance, suitable for use in hashing algorithms and data structures like a hash table. 
+            </returns>
+        </member>
+        <member name="P:Newtonsoft.Json.Linq.JEnumerable`1.Item(System.Object)">
+            <summary>
+            Gets the <see cref="T:Newtonsoft.Json.Linq.IJEnumerable`1"/> with the specified key.
+            </summary>
+            <value></value>
+        </member>
+        <member name="T:Newtonsoft.Json.Linq.JObject">
+            <summary>
+            Represents a JSON object.
+            </summary>
+            <example>
+              <code lang="cs" source="..\Src\Newtonsoft.Json.Tests\Documentation\LinqToJsonTests.cs" region="LinqToJsonCreateParse" title="Parsing a JSON Object from Text" />
+            </example>
+        </member>
+        <member name="M:Newtonsoft.Json.Linq.JObject.#ctor">
+            <summary>
+            Initializes a new instance of the <see cref="T:Newtonsoft.Json.Linq.JObject"/> class.
+            </summary>
+        </member>
+        <member name="M:Newtonsoft.Json.Linq.JObject.#ctor(Newtonsoft.Json.Linq.JObject)">
+            <summary>
+            Initializes a new instance of the <see cref="T:Newtonsoft.Json.Linq.JObject"/> class from another <see cref="T:Newtonsoft.Json.Linq.JObject"/> object.
+            </summary>
+            <param name="other">A <see cref="T:Newtonsoft.Json.Linq.JObject"/> object to copy from.</param>
+        </member>
+        <member name="M:Newtonsoft.Json.Linq.JObject.#ctor(System.Object[])">
+            <summary>
+            Initializes a new instance of the <see cref="T:Newtonsoft.Json.Linq.JObject"/> class with the specified content.
+            </summary>
+            <param name="content">The contents of the object.</param>
+        </member>
+        <member name="M:Newtonsoft.Json.Linq.JObject.#ctor(System.Object)">
+            <summary>
+            Initializes a new instance of the <see cref="T:Newtonsoft.Json.Linq.JObject"/> class with the specified content.
+            </summary>
+            <param name="content">The contents of the object.</param>
+        </member>
+        <member name="M:Newtonsoft.Json.Linq.JObject.Properties">
+            <summary>
+            Gets an <see cref="T:System.Collections.Generic.IEnumerable`1"/> of this object's properties.
+            </summary>
+            <returns>An <see cref="T:System.Collections.Generic.IEnumerable`1"/> of this object's properties.</returns>
+        </member>
+        <member name="M:Newtonsoft.Json.Linq.JObject.Property(System.String)">
+            <summary>
+            Gets a <see cref="T:Newtonsoft.Json.Linq.JProperty"/> the specified name.
+            </summary>
+            <param name="name">The property name.</param>
+            <returns>A <see cref="T:Newtonsoft.Json.Linq.JProperty"/> with the specified name or null.</returns>
+        </member>
+        <member name="M:Newtonsoft.Json.Linq.JObject.PropertyValues">
+            <summary>
+            Gets an <see cref="T:Newtonsoft.Json.Linq.JEnumerable`1"/> of this object's property values.
+            </summary>
+            <returns>An <see cref="T:Newtonsoft.Json.Linq.JEnumerable`1"/> of this object's property values.</returns>
+        </member>
+        <member name="M:Newtonsoft.Json.Linq.JObject.Load(Newtonsoft.Json.JsonReader)">
+            <summary>
+            Loads an <see cref="T:Newtonsoft.Json.Linq.JObject"/> from a <see cref="T:Newtonsoft.Json.JsonReader"/>. 
+            </summary>
+            <param name="reader">A <see cref="T:Newtonsoft.Json.JsonReader"/> that will be read for the content of the <see cref="T:Newtonsoft.Json.Linq.JObject"/>.</param>
+            <returns>A <see cref="T:Newtonsoft.Json.Linq.JObject"/> that contains the JSON that was read from the specified <see cref="T:Newtonsoft.Json.JsonReader"/>.</returns>
+        </member>
+        <member name="M:Newtonsoft.Json.Linq.JObject.Parse(System.String)">
+            <summary>
+            Load a <see cref="T:Newtonsoft.Json.Linq.JObject"/> from a string that contains JSON.
+            </summary>
+            <param name="json">A <see cref="T:System.String"/> that contains JSON.</param>
+            <returns>A <see cref="T:Newtonsoft.Json.Linq.JObject"/> populated from the string that contains JSON.</returns>
+            <example>
+              <code lang="cs" source="..\Src\Newtonsoft.Json.Tests\Documentation\LinqToJsonTests.cs" region="LinqToJsonCreateParse" title="Parsing a JSON Object from Text"/>
+            </example>
+        </member>
+        <member name="M:Newtonsoft.Json.Linq.JObject.FromObject(System.Object)">
+            <summary>
+            Creates a <see cref="T:Newtonsoft.Json.Linq.JObject"/> from an object.
+            </summary>
+            <param name="o">The object that will be used to create <see cref="T:Newtonsoft.Json.Linq.JObject"/>.</param>
+            <returns>A <see cref="T:Newtonsoft.Json.Linq.JObject"/> with the values of the specified object</returns>
+        </member>
+        <member name="M:Newtonsoft.Json.Linq.JObject.FromObject(System.Object,Newtonsoft.Json.JsonSerializer)">
+            <summary>
+            Creates a <see cref="T:Newtonsoft.Json.Linq.JArray"/> from an object.
+            </summary>
+            <param name="o">The object that will be used to create <see cref="T:Newtonsoft.Json.Linq.JArray"/>.</param>
+            <param name="jsonSerializer">The <see cref="T:Newtonsoft.Json.JsonSerializer"/> that will be used to read the object.</param>
+            <returns>A <see cref="T:Newtonsoft.Json.Linq.JArray"/> with the values of the specified object</returns>
+        </member>
+        <member name="M:Newtonsoft.Json.Linq.JObject.WriteTo(Newtonsoft.Json.JsonWriter,Newtonsoft.Json.JsonConverter[])">
+            <summary>
+            Writes this token to a <see cref="T:Newtonsoft.Json.JsonWriter"/>.
+            </summary>
+            <param name="writer">A <see cref="T:Newtonsoft.Json.JsonWriter"/> into which this method will write.</param>
+            <param name="converters">A collection of <see cref="T:Newtonsoft.Json.JsonConverter"/> which will be used when writing the token.</param>
+        </member>
+        <member name="M:Newtonsoft.Json.Linq.JObject.GetValue(System.String)">
+            <summary>
+            Gets the <see cref="T:Newtonsoft.Json.Linq.JToken"/> with the specified property name.
+            </summary>
+            <param name="propertyName">Name of the property.</param>
+            <returns>The <see cref="T:Newtonsoft.Json.Linq.JToken"/> with the specified property name.</returns>
+        </member>
+        <member name="M:Newtonsoft.Json.Linq.JObject.GetValue(System.String,System.StringComparison)">
+            <summary>
+            Gets the <see cref="T:Newtonsoft.Json.Linq.JToken"/> with the specified property name.
+            The exact property name will be searched for first and if no matching property is found then
+            the <see cref="T:System.StringComparison"/> will be used to match a property.
+            </summary>
+            <param name="propertyName">Name of the property.</param>
+            <param name="comparison">One of the enumeration values that specifies how the strings will be compared.</param>
+            <returns>The <see cref="T:Newtonsoft.Json.Linq.JToken"/> with the specified property name.</returns>
+        </member>
+        <member name="M:Newtonsoft.Json.Linq.JObject.TryGetValue(System.String,System.StringComparison,Newtonsoft.Json.Linq.JToken@)">
+            <summary>
+            Tries to get the <see cref="T:Newtonsoft.Json.Linq.JToken"/> with the specified property name.
+            The exact property name will be searched for first and if no matching property is found then
+            the <see cref="T:System.StringComparison"/> will be used to match a property.
+            </summary>
+            <param name="propertyName">Name of the property.</param>
+            <param name="value">The value.</param>
+            <param name="comparison">One of the enumeration values that specifies how the strings will be compared.</param>
+            <returns>true if a value was successfully retrieved; otherwise, false.</returns>
+        </member>
+        <member name="M:Newtonsoft.Json.Linq.JObject.Add(System.String,Newtonsoft.Json.Linq.JToken)">
+            <summary>
+            Adds the specified property name.
+            </summary>
+            <param name="propertyName">Name of the property.</param>
+            <param name="value">The value.</param>
+        </member>
+        <member name="M:Newtonsoft.Json.Linq.JObject.Remove(System.String)">
+            <summary>
+            Removes the property with the specified name.
+            </summary>
+            <param name="propertyName">Name of the property.</param>
+            <returns>true if item was successfully removed; otherwise, false.</returns>
+        </member>
+        <member name="M:Newtonsoft.Json.Linq.JObject.TryGetValue(System.String,Newtonsoft.Json.Linq.JToken@)">
+            <summary>
+            Tries the get value.
+            </summary>
+            <param name="propertyName">Name of the property.</param>
+            <param name="value">The value.</param>
+            <returns>true if a value was successfully retrieved; otherwise, false.</returns>
+        </member>
+        <member name="M:Newtonsoft.Json.Linq.JObject.GetEnumerator">
+            <summary>
+            Returns an enumerator that iterates through the collection.
+            </summary>
+            <returns>
+            A <see cref="T:System.Collections.Generic.IEnumerator`1"/> that can be used to iterate through the collection.
+            </returns>
+        </member>
+        <member name="M:Newtonsoft.Json.Linq.JObject.OnPropertyChanged(System.String)">
+            <summary>
+            Raises the <see cref="E:Newtonsoft.Json.Linq.JObject.PropertyChanged"/> event with the provided arguments.
+            </summary>
+            <param name="propertyName">Name of the property.</param>
+        </member>
+        <member name="M:Newtonsoft.Json.Linq.JObject.System#ComponentModel#ICustomTypeDescriptor#GetProperties">
+            <summary>
+            Returns the properties for this instance of a component.
+            </summary>
+            <returns>
+            A <see cref="T:System.ComponentModel.PropertyDescriptorCollection"/> that represents the properties for this component instance.
+            </returns>
+        </member>
+        <member name="M:Newtonsoft.Json.Linq.JObject.System#ComponentModel#ICustomTypeDescriptor#GetProperties(System.Attribute[])">
+            <summary>
+            Returns the properties for this instance of a component using the attribute array as a filter.
+            </summary>
+            <param name="attributes">An array of type <see cref="T:System.Attribute"/> that is used as a filter.</param>
+            <returns>
+            A <see cref="T:System.ComponentModel.PropertyDescriptorCollection"/> that represents the filtered properties for this component instance.
+            </returns>
+        </member>
+        <member name="M:Newtonsoft.Json.Linq.JObject.System#ComponentModel#ICustomTypeDescriptor#GetAttributes">
+            <summary>
+            Returns a collection of custom attributes for this instance of a component.
+            </summary>
+            <returns>
+            An <see cref="T:System.ComponentModel.AttributeCollection"/> containing the attributes for this object.
+            </returns>
+        </member>
+        <member name="M:Newtonsoft.Json.Linq.JObject.System#ComponentModel#ICustomTypeDescriptor#GetClassName">
+            <summary>
+            Returns the class name of this instance of a component.
+            </summary>
+            <returns>
+            The class name of the object, or null if the class does not have a name.
+            </returns>
+        </member>
+        <member name="M:Newtonsoft.Json.Linq.JObject.System#ComponentModel#ICustomTypeDescriptor#GetComponentName">
+            <summary>
+            Returns the name of this instance of a component.
+            </summary>
+            <returns>
+            The name of the object, or null if the object does not have a name.
+            </returns>
+        </member>
+        <member name="M:Newtonsoft.Json.Linq.JObject.System#ComponentModel#ICustomTypeDescriptor#GetConverter">
+            <summary>
+            Returns a type converter for this instance of a component.
+            </summary>
+            <returns>
+            A <see cref="T:System.ComponentModel.TypeConverter"/> that is the converter for this object, or null if there is no <see cref="T:System.ComponentModel.TypeConverter"/> for this object.
+            </returns>
+        </member>
+        <member name="M:Newtonsoft.Json.Linq.JObject.System#ComponentModel#ICustomTypeDescriptor#GetDefaultEvent">
+            <summary>
+            Returns the default event for this instance of a component.
+            </summary>
+            <returns>
+            An <see cref="T:System.ComponentModel.EventDescriptor"/> that represents the default event for this object, or null if this object does not have events.
+            </returns>
+        </member>
+        <member name="M:Newtonsoft.Json.Linq.JObject.System#ComponentModel#ICustomTypeDescriptor#GetDefaultProperty">
+            <summary>
+            Returns the default property for this instance of a component.
+            </summary>
+            <returns>
+            A <see cref="T:System.ComponentModel.PropertyDescriptor"/> that represents the default property for this object, or null if this object does not have properties.
+            </returns>
+        </member>
+        <member name="M:Newtonsoft.Json.Linq.JObject.System#ComponentModel#ICustomTypeDescriptor#GetEditor(System.Type)">
+            <summary>
+            Returns an editor of the specified type for this instance of a component.
+            </summary>
+            <param name="editorBaseType">A <see cref="T:System.Type"/> that represents the editor for this object.</param>
+            <returns>
+            An <see cref="T:System.Object"/> of the specified type that is the editor for this object, or null if the editor cannot be found.
+            </returns>
+        </member>
+        <member name="M:Newtonsoft.Json.Linq.JObject.System#ComponentModel#ICustomTypeDescriptor#GetEvents(System.Attribute[])">
+            <summary>
+            Returns the events for this instance of a component using the specified attribute array as a filter.
+            </summary>
+            <param name="attributes">An array of type <see cref="T:System.Attribute"/> that is used as a filter.</param>
+            <returns>
+            An <see cref="T:System.ComponentModel.EventDescriptorCollection"/> that represents the filtered events for this component instance.
+            </returns>
+        </member>
+        <member name="M:Newtonsoft.Json.Linq.JObject.System#ComponentModel#ICustomTypeDescriptor#GetEvents">
+            <summary>
+            Returns the events for this instance of a component.
+            </summary>
+            <returns>
+            An <see cref="T:System.ComponentModel.EventDescriptorCollection"/> that represents the events for this component instance.
+            </returns>
+        </member>
+        <member name="M:Newtonsoft.Json.Linq.JObject.System#ComponentModel#ICustomTypeDescriptor#GetPropertyOwner(System.ComponentModel.PropertyDescriptor)">
+            <summary>
+            Returns an object that contains the property described by the specified property descriptor.
+            </summary>
+            <param name="pd">A <see cref="T:System.ComponentModel.PropertyDescriptor"/> that represents the property whose owner is to be found.</param>
+            <returns>
+            An <see cref="T:System.Object"/> that represents the owner of the specified property.
+            </returns>
+        </member>
+        <member name="P:Newtonsoft.Json.Linq.JObject.ChildrenTokens">
+            <summary>
+            Gets the container's children tokens.
+            </summary>
+            <value>The container's children tokens.</value>
+        </member>
+        <member name="E:Newtonsoft.Json.Linq.JObject.PropertyChanged">
+            <summary>
+            Occurs when a property value changes.
+            </summary>
+        </member>
+        <member name="P:Newtonsoft.Json.Linq.JObject.Type">
+            <summary>
+            Gets the node type for this <see cref="T:Newtonsoft.Json.Linq.JToken"/>.
+            </summary>
+            <value>The type.</value>
+        </member>
+        <member name="P:Newtonsoft.Json.Linq.JObject.Item(System.Object)">
+            <summary>
+            Gets the <see cref="T:Newtonsoft.Json.Linq.JToken"/> with the specified key.
+            </summary>
+            <value>The <see cref="T:Newtonsoft.Json.Linq.JToken"/> with the specified key.</value>
+        </member>
+        <member name="P:Newtonsoft.Json.Linq.JObject.Item(System.String)">
+            <summary>
+            Gets or sets the <see cref="T:Newtonsoft.Json.Linq.JToken"/> with the specified property name.
+            </summary>
+            <value></value>
+        </member>
+        <member name="T:Newtonsoft.Json.Linq.JArray">
+            <summary>
+            Represents a JSON array.
+            </summary>
+            <example>
+              <code lang="cs" source="..\Src\Newtonsoft.Json.Tests\Documentation\LinqToJsonTests.cs" region="LinqToJsonCreateParseArray" title="Parsing a JSON Array from Text" />
+            </example>
+        </member>
+        <member name="M:Newtonsoft.Json.Linq.JArray.#ctor">
+            <summary>
+            Initializes a new instance of the <see cref="T:Newtonsoft.Json.Linq.JArray"/> class.
+            </summary>
+        </member>
+        <member name="M:Newtonsoft.Json.Linq.JArray.#ctor(Newtonsoft.Json.Linq.JArray)">
+            <summary>
+            Initializes a new instance of the <see cref="T:Newtonsoft.Json.Linq.JArray"/> class from another <see cref="T:Newtonsoft.Json.Linq.JArray"/> object.
+            </summary>
+            <param name="other">A <see cref="T:Newtonsoft.Json.Linq.JArray"/> object to copy from.</param>
+        </member>
+        <member name="M:Newtonsoft.Json.Linq.JArray.#ctor(System.Object[])">
+            <summary>
+            Initializes a new instance of the <see cref="T:Newtonsoft.Json.Linq.JArray"/> class with the specified content.
+            </summary>
+            <param name="content">The contents of the array.</param>
+        </member>
+        <member name="M:Newtonsoft.Json.Linq.JArray.#ctor(System.Object)">
+            <summary>
+            Initializes a new instance of the <see cref="T:Newtonsoft.Json.Linq.JArray"/> class with the specified content.
+            </summary>
+            <param name="content">The contents of the array.</param>
+        </member>
+        <member name="M:Newtonsoft.Json.Linq.JArray.Load(Newtonsoft.Json.JsonReader)">
+            <summary>
+            Loads an <see cref="T:Newtonsoft.Json.Linq.JArray"/> from a <see cref="T:Newtonsoft.Json.JsonReader"/>. 
+            </summary>
+            <param name="reader">A <see cref="T:Newtonsoft.Json.JsonReader"/> that will be read for the content of the <see cref="T:Newtonsoft.Json.Linq.JArray"/>.</param>
+            <returns>A <see cref="T:Newtonsoft.Json.Linq.JArray"/> that contains the JSON that was read from the specified <see cref="T:Newtonsoft.Json.JsonReader"/>.</returns>
+        </member>
+        <member name="M:Newtonsoft.Json.Linq.JArray.Parse(System.String)">
+            <summary>
+            Load a <see cref="T:Newtonsoft.Json.Linq.JArray"/> from a string that contains JSON.
+            </summary>
+            <param name="json">A <see cref="T:System.String"/> that contains JSON.</param>
+            <returns>A <see cref="T:Newtonsoft.Json.Linq.JArray"/> populated from the string that contains JSON.</returns>
+            <example>
+              <code lang="cs" source="..\Src\Newtonsoft.Json.Tests\Documentation\LinqToJsonTests.cs" region="LinqToJsonCreateParseArray" title="Parsing a JSON Array from Text"/>
+            </example>
+        </member>
+        <member name="M:Newtonsoft.Json.Linq.JArray.FromObject(System.Object)">
+            <summary>
+            Creates a <see cref="T:Newtonsoft.Json.Linq.JArray"/> from an object.
+            </summary>
+            <param name="o">The object that will be used to create <see cref="T:Newtonsoft.Json.Linq.JArray"/>.</param>
+            <returns>A <see cref="T:Newtonsoft.Json.Linq.JArray"/> with the values of the specified object</returns>
+        </member>
+        <member name="M:Newtonsoft.Json.Linq.JArray.FromObject(System.Object,Newtonsoft.Json.JsonSerializer)">
+            <summary>
+            Creates a <see cref="T:Newtonsoft.Json.Linq.JArray"/> from an object.
+            </summary>
+            <param name="o">The object that will be used to create <see cref="T:Newtonsoft.Json.Linq.JArray"/>.</param>
+            <param name="jsonSerializer">The <see cref="T:Newtonsoft.Json.JsonSerializer"/> that will be used to read the object.</param>
+            <returns>A <see cref="T:Newtonsoft.Json.Linq.JArray"/> with the values of the specified object</returns>
+        </member>
+        <member name="M:Newtonsoft.Json.Linq.JArray.WriteTo(Newtonsoft.Json.JsonWriter,Newtonsoft.Json.JsonConverter[])">
+            <summary>
+            Writes this token to a <see cref="T:Newtonsoft.Json.JsonWriter"/>.
+            </summary>
+            <param name="writer">A <see cref="T:Newtonsoft.Json.JsonWriter"/> into which this method will write.</param>
+            <param name="converters">A collection of <see cref="T:Newtonsoft.Json.JsonConverter"/> which will be used when writing the token.</param>
+        </member>
+        <member name="M:Newtonsoft.Json.Linq.JArray.IndexOf(Newtonsoft.Json.Linq.JToken)">
+            <summary>
+            Determines the index of a specific item in the <see cref="T:System.Collections.Generic.IList`1"/>.
+            </summary>
+            <param name="item">The object to locate in the <see cref="T:System.Collections.Generic.IList`1"/>.</param>
+            <returns>
+            The index of <paramref name="item"/> if found in the list; otherwise, -1.
+            </returns>
+        </member>
+        <member name="M:Newtonsoft.Json.Linq.JArray.Insert(System.Int32,Newtonsoft.Json.Linq.JToken)">
+            <summary>
+            Inserts an item to the <see cref="T:System.Collections.Generic.IList`1"/> at the specified index.
+            </summary>
+            <param name="index">The zero-based index at which <paramref name="item"/> should be inserted.</param>
+            <param name="item">The object to insert into the <see cref="T:System.Collections.Generic.IList`1"/>.</param>
+            <exception cref="T:System.ArgumentOutOfRangeException">
+            	<paramref name="index"/> is not a valid index in the <see cref="T:System.Collections.Generic.IList`1"/>.</exception>
+            <exception cref="T:System.NotSupportedException">The <see cref="T:System.Collections.Generic.IList`1"/> is read-only.</exception>
+        </member>
+        <member name="M:Newtonsoft.Json.Linq.JArray.RemoveAt(System.Int32)">
+            <summary>
+            Removes the <see cref="T:System.Collections.Generic.IList`1"/> item at the specified index.
+            </summary>
+            <param name="index">The zero-based index of the item to remove.</param>
+            <exception cref="T:System.ArgumentOutOfRangeException">
+            	<paramref name="index"/> is not a valid index in the <see cref="T:System.Collections.Generic.IList`1"/>.</exception>
+            <exception cref="T:System.NotSupportedException">The <see cref="T:System.Collections.Generic.IList`1"/> is read-only.</exception>
+        </member>
+        <member name="M:Newtonsoft.Json.Linq.JArray.Add(Newtonsoft.Json.Linq.JToken)">
+            <summary>
+            Adds an item to the <see cref="T:System.Collections.Generic.ICollection`1"/>.
+            </summary>
+            <param name="item">The object to add to the <see cref="T:System.Collections.Generic.ICollection`1"/>.</param>
+            <exception cref="T:System.NotSupportedException">The <see cref="T:System.Collections.Generic.ICollection`1"/> is read-only.</exception>
+        </member>
+        <member name="M:Newtonsoft.Json.Linq.JArray.Clear">
+            <summary>
+            Removes all items from the <see cref="T:System.Collections.Generic.ICollection`1"/>.
+            </summary>
+            <exception cref="T:System.NotSupportedException">The <see cref="T:System.Collections.Generic.ICollection`1"/> is read-only. </exception>
+        </member>
+        <member name="M:Newtonsoft.Json.Linq.JArray.Contains(Newtonsoft.Json.Linq.JToken)">
+            <summary>
+            Determines whether the <see cref="T:System.Collections.Generic.ICollection`1"/> contains a specific value.
+            </summary>
+            <param name="item">The object to locate in the <see cref="T:System.Collections.Generic.ICollection`1"/>.</param>
+            <returns>
+            true if <paramref name="item"/> is found in the <see cref="T:System.Collections.Generic.ICollection`1"/>; otherwise, false.
+            </returns>
+        </member>
+        <member name="M:Newtonsoft.Json.Linq.JArray.Remove(Newtonsoft.Json.Linq.JToken)">
+            <summary>
+            Removes the first occurrence of a specific object from the <see cref="T:System.Collections.Generic.ICollection`1"/>.
+            </summary>
+            <param name="item">The object to remove from the <see cref="T:System.Collections.Generic.ICollection`1"/>.</param>
+            <returns>
+            true if <paramref name="item"/> was successfully removed from the <see cref="T:System.Collections.Generic.ICollection`1"/>; otherwise, false. This method also returns false if <paramref name="item"/> is not found in the original <see cref="T:System.Collections.Generic.ICollection`1"/>.
+            </returns>
+            <exception cref="T:System.NotSupportedException">The <see cref="T:System.Collections.Generic.ICollection`1"/> is read-only.</exception>
+        </member>
+        <member name="P:Newtonsoft.Json.Linq.JArray.ChildrenTokens">
+            <summary>
+            Gets the container's children tokens.
+            </summary>
+            <value>The container's children tokens.</value>
+        </member>
+        <member name="P:Newtonsoft.Json.Linq.JArray.Type">
+            <summary>
+            Gets the node type for this <see cref="T:Newtonsoft.Json.Linq.JToken"/>.
+            </summary>
+            <value>The type.</value>
+        </member>
+        <member name="P:Newtonsoft.Json.Linq.JArray.Item(System.Object)">
+            <summary>
+            Gets the <see cref="T:Newtonsoft.Json.Linq.JToken"/> with the specified key.
+            </summary>
+            <value>The <see cref="T:Newtonsoft.Json.Linq.JToken"/> with the specified key.</value>
+        </member>
+        <member name="P:Newtonsoft.Json.Linq.JArray.Item(System.Int32)">
+            <summary>
+            Gets or sets the <see cref="T:Newtonsoft.Json.Linq.JToken"/> at the specified index.
+            </summary>
+            <value></value>
+        </member>
+        <member name="T:Newtonsoft.Json.Linq.JTokenReader">
+            <summary>
+            Represents a reader that provides fast, non-cached, forward-only access to serialized Json data.
+            </summary>
+        </member>
+        <member name="M:Newtonsoft.Json.Linq.JTokenReader.#ctor(Newtonsoft.Json.Linq.JToken)">
+            <summary>
+            Initializes a new instance of the <see cref="T:Newtonsoft.Json.Linq.JTokenReader"/> class.
+            </summary>
+            <param name="token">The token to read from.</param>
+        </member>
+        <member name="M:Newtonsoft.Json.Linq.JTokenReader.ReadAsBytes">
+            <summary>
+            Reads the next JSON token from the stream as a <see cref="T:Byte[]"/>.
+            </summary>
+            <returns>
+            A <see cref="T:Byte[]"/> or a null reference if the next JSON token is null. This method will return <c>null</c> at the end of an array.
+            </returns>
+        </member>
+        <member name="M:Newtonsoft.Json.Linq.JTokenReader.ReadAsDecimal">
+            <summary>
+            Reads the next JSON token from the stream as a <see cref="T:System.Nullable`1"/>.
+            </summary>
+            <returns>A <see cref="T:System.Nullable`1"/>. This method will return <c>null</c> at the end of an array.</returns>
+        </member>
+        <member name="M:Newtonsoft.Json.Linq.JTokenReader.ReadAsInt32">
+            <summary>
+            Reads the next JSON token from the stream as a <see cref="T:System.Nullable`1"/>.
+            </summary>
+            <returns>A <see cref="T:System.Nullable`1"/>. This method will return <c>null</c> at the end of an array.</returns>
+        </member>
+        <member name="M:Newtonsoft.Json.Linq.JTokenReader.ReadAsString">
+            <summary>
+            Reads the next JSON token from the stream as a <see cref="T:System.String"/>.
+            </summary>
+            <returns>A <see cref="T:System.String"/>. This method will return <c>null</c> at the end of an array.</returns>
+        </member>
+        <member name="M:Newtonsoft.Json.Linq.JTokenReader.ReadAsDateTime">
+            <summary>
+            Reads the next JSON token from the stream as a <see cref="T:System.Nullable`1"/>.
+            </summary>
+            <returns>A <see cref="T:System.String"/>. This method will return <c>null</c> at the end of an array.</returns>
+        </member>
+        <member name="M:Newtonsoft.Json.Linq.JTokenReader.Read">
+            <summary>
+            Reads the next JSON token from the stream.
+            </summary>
+            <returns>
+            true if the next token was read successfully; false if there are no more tokens to read.
+            </returns>
+        </member>
+        <member name="T:Newtonsoft.Json.Linq.JTokenWriter">
+            <summary>
+            Represents a writer that provides a fast, non-cached, forward-only way of generating Json data.
+            </summary>
+        </member>
+        <member name="M:Newtonsoft.Json.Linq.JTokenWriter.#ctor(Newtonsoft.Json.Linq.JContainer)">
+            <summary>
+            Initializes a new instance of the <see cref="T:Newtonsoft.Json.Linq.JTokenWriter"/> class writing to the given <see cref="T:Newtonsoft.Json.Linq.JContainer"/>.
+            </summary>
+            <param name="container">The container being written to.</param>
+        </member>
+        <member name="M:Newtonsoft.Json.Linq.JTokenWriter.#ctor">
+            <summary>
+            Initializes a new instance of the <see cref="T:Newtonsoft.Json.Linq.JTokenWriter"/> class.
+            </summary>
+        </member>
+        <member name="M:Newtonsoft.Json.Linq.JTokenWriter.Flush">
+            <summary>
+            Flushes whatever is in the buffer to the underlying streams and also flushes the underlying stream.
+            </summary>
+        </member>
+        <member name="M:Newtonsoft.Json.Linq.JTokenWriter.Close">
+            <summary>
+            Closes this stream and the underlying stream.
+            </summary>
+        </member>
+        <member name="M:Newtonsoft.Json.Linq.JTokenWriter.WriteStartObject">
+            <summary>
+            Writes the beginning of a Json object.
+            </summary>
+        </member>
+        <member name="M:Newtonsoft.Json.Linq.JTokenWriter.WriteStartArray">
+            <summary>
+            Writes the beginning of a Json array.
+            </summary>
+        </member>
+        <member name="M:Newtonsoft.Json.Linq.JTokenWriter.WriteStartConstructor(System.String)">
+            <summary>
+            Writes the start of a constructor with the given name.
+            </summary>
+            <param name="name">The name of the constructor.</param>
+        </member>
+        <member name="M:Newtonsoft.Json.Linq.JTokenWriter.WriteEnd(Newtonsoft.Json.JsonToken)">
+            <summary>
+            Writes the end.
+            </summary>
+            <param name="token">The token.</param>
+        </member>
+        <member name="M:Newtonsoft.Json.Linq.JTokenWriter.WritePropertyName(System.String)">
+            <summary>
+            Writes the property name of a name/value pair on a Json object.
+            </summary>
+            <param name="name">The name of the property.</param>
+        </member>
+        <member name="M:Newtonsoft.Json.Linq.JTokenWriter.WriteNull">
+            <summary>
+            Writes a null value.
+            </summary>
+        </member>
+        <member name="M:Newtonsoft.Json.Linq.JTokenWriter.WriteUndefined">
+            <summary>
+            Writes an undefined value.
+            </summary>
+        </member>
+        <member name="M:Newtonsoft.Json.Linq.JTokenWriter.WriteRaw(System.String)">
+            <summary>
+            Writes raw JSON.
+            </summary>
+            <param name="json">The raw JSON to write.</param>
+        </member>
+        <member name="M:Newtonsoft.Json.Linq.JTokenWriter.WriteComment(System.String)">
+            <summary>
+            Writes out a comment <code>/*...*/</code> containing the specified text.
+            </summary>
+            <param name="text">Text to place inside the comment.</param>
+        </member>
+        <member name="M:Newtonsoft.Json.Linq.JTokenWriter.WriteValue(System.String)">
+            <summary>
+            Writes a <see cref="T:System.String"/> value.
+            </summary>
+            <param name="value">The <see cref="T:System.String"/> value to write.</param>
+        </member>
+        <member name="M:Newtonsoft.Json.Linq.JTokenWriter.WriteValue(System.Int32)">
+            <summary>
+            Writes a <see cref="T:System.Int32"/> value.
+            </summary>
+            <param name="value">The <see cref="T:System.Int32"/> value to write.</param>
+        </member>
+        <member name="M:Newtonsoft.Json.Linq.JTokenWriter.WriteValue(System.UInt32)">
+            <summary>
+            Writes a <see cref="T:System.UInt32"/> value.
+            </summary>
+            <param name="value">The <see cref="T:System.UInt32"/> value to write.</param>
+        </member>
+        <member name="M:Newtonsoft.Json.Linq.JTokenWriter.WriteValue(System.Int64)">
+            <summary>
+            Writes a <see cref="T:System.Int64"/> value.
+            </summary>
+            <param name="value">The <see cref="T:System.Int64"/> value to write.</param>
+        </member>
+        <member name="M:Newtonsoft.Json.Linq.JTokenWriter.WriteValue(System.UInt64)">
+            <summary>
+            Writes a <see cref="T:System.UInt64"/> value.
+            </summary>
+            <param name="value">The <see cref="T:System.UInt64"/> value to write.</param>
+        </member>
+        <member name="M:Newtonsoft.Json.Linq.JTokenWriter.WriteValue(System.Single)">
+            <summary>
+            Writes a <see cref="T:System.Single"/> value.
+            </summary>
+            <param name="value">The <see cref="T:System.Single"/> value to write.</param>
+        </member>
+        <member name="M:Newtonsoft.Json.Linq.JTokenWriter.WriteValue(System.Double)">
+            <summary>
+            Writes a <see cref="T:System.Double"/> value.
+            </summary>
+            <param name="value">The <see cref="T:System.Double"/> value to write.</param>
+        </member>
+        <member name="M:Newtonsoft.Json.Linq.JTokenWriter.WriteValue(System.Boolean)">
+            <summary>
+            Writes a <see cref="T:System.Boolean"/> value.
+            </summary>
+            <param name="value">The <see cref="T:System.Boolean"/> value to write.</param>
+        </member>
+        <member name="M:Newtonsoft.Json.Linq.JTokenWriter.WriteValue(System.Int16)">
+            <summary>
+            Writes a <see cref="T:System.Int16"/> value.
+            </summary>
+            <param name="value">The <see cref="T:System.Int16"/> value to write.</param>
+        </member>
+        <member name="M:Newtonsoft.Json.Linq.JTokenWriter.WriteValue(System.UInt16)">
+            <summary>
+            Writes a <see cref="T:System.UInt16"/> value.
+            </summary>
+            <param name="value">The <see cref="T:System.UInt16"/> value to write.</param>
+        </member>
+        <member name="M:Newtonsoft.Json.Linq.JTokenWriter.WriteValue(System.Char)">
+            <summary>
+            Writes a <see cref="T:System.Char"/> value.
+            </summary>
+            <param name="value">The <see cref="T:System.Char"/> value to write.</param>
+        </member>
+        <member name="M:Newtonsoft.Json.Linq.JTokenWriter.WriteValue(System.Byte)">
+            <summary>
+            Writes a <see cref="T:System.Byte"/> value.
+            </summary>
+            <param name="value">The <see cref="T:System.Byte"/> value to write.</param>
+        </member>
+        <member name="M:Newtonsoft.Json.Linq.JTokenWriter.WriteValue(System.SByte)">
+            <summary>
+            Writes a <see cref="T:System.SByte"/> value.
+            </summary>
+            <param name="value">The <see cref="T:System.SByte"/> value to write.</param>
+        </member>
+        <member name="M:Newtonsoft.Json.Linq.JTokenWriter.WriteValue(System.Decimal)">
+            <summary>
+            Writes a <see cref="T:System.Decimal"/> value.
+            </summary>
+            <param name="value">The <see cref="T:System.Decimal"/> value to write.</param>
+        </member>
+        <member name="M:Newtonsoft.Json.Linq.JTokenWriter.WriteValue(System.DateTime)">
+            <summary>
+            Writes a <see cref="T:System.DateTime"/> value.
+            </summary>
+            <param name="value">The <see cref="T:System.DateTime"/> value to write.</param>
+        </member>
+        <member name="M:Newtonsoft.Json.Linq.JTokenWriter.WriteValue(System.Byte[])">
+            <summary>
+            Writes a <see cref="T:Byte[]"/> value.
+            </summary>
+            <param name="value">The <see cref="T:Byte[]"/> value to write.</param>
+        </member>
+        <member name="M:Newtonsoft.Json.Linq.JTokenWriter.WriteValue(System.TimeSpan)">
+            <summary>
+            Writes a <see cref="T:System.TimeSpan"/> value.
+            </summary>
+            <param name="value">The <see cref="T:System.TimeSpan"/> value to write.</param>
+        </member>
+        <member name="M:Newtonsoft.Json.Linq.JTokenWriter.WriteValue(System.Guid)">
+            <summary>
+            Writes a <see cref="T:System.Guid"/> value.
+            </summary>
+            <param name="value">The <see cref="T:System.Guid"/> value to write.</param>
+        </member>
+        <member name="M:Newtonsoft.Json.Linq.JTokenWriter.WriteValue(System.Uri)">
+            <summary>
+            Writes a <see cref="T:System.Uri"/> value.
+            </summary>
+            <param name="value">The <see cref="T:System.Uri"/> value to write.</param>
+        </member>
+        <member name="P:Newtonsoft.Json.Linq.JTokenWriter.Token">
+            <summary>
+            Gets the token being writen.
+            </summary>
+            <value>The token being writen.</value>
+        </member>
+        <member name="T:Newtonsoft.Json.Linq.JProperty">
+            <summary>
+            Represents a JSON property.
+            </summary>
+        </member>
+        <member name="M:Newtonsoft.Json.Linq.JProperty.#ctor(Newtonsoft.Json.Linq.JProperty)">
+            <summary>
+            Initializes a new instance of the <see cref="T:Newtonsoft.Json.Linq.JProperty"/> class from another <see cref="T:Newtonsoft.Json.Linq.JProperty"/> object.
+            </summary>
+            <param name="other">A <see cref="T:Newtonsoft.Json.Linq.JProperty"/> object to copy from.</param>
+        </member>
+        <member name="M:Newtonsoft.Json.Linq.JProperty.#ctor(System.String,System.Object[])">
+            <summary>
+            Initializes a new instance of the <see cref="T:Newtonsoft.Json.Linq.JProperty"/> class.
+            </summary>
+            <param name="name">The property name.</param>
+            <param name="content">The property content.</param>
+        </member>
+        <member name="M:Newtonsoft.Json.Linq.JProperty.#ctor(System.String,System.Object)">
+            <summary>
+            Initializes a new instance of the <see cref="T:Newtonsoft.Json.Linq.JProperty"/> class.
+            </summary>
+            <param name="name">The property name.</param>
+            <param name="content">The property content.</param>
+        </member>
+        <member name="M:Newtonsoft.Json.Linq.JProperty.WriteTo(Newtonsoft.Json.JsonWriter,Newtonsoft.Json.JsonConverter[])">
+            <summary>
+            Writes this token to a <see cref="T:Newtonsoft.Json.JsonWriter"/>.
+            </summary>
+            <param name="writer">A <see cref="T:Newtonsoft.Json.JsonWriter"/> into which this method will write.</param>
+            <param name="converters">A collection of <see cref="T:Newtonsoft.Json.JsonConverter"/> which will be used when writing the token.</param>
+        </member>
+        <member name="M:Newtonsoft.Json.Linq.JProperty.Load(Newtonsoft.Json.JsonReader)">
+            <summary>
+            Loads an <see cref="T:Newtonsoft.Json.Linq.JProperty"/> from a <see cref="T:Newtonsoft.Json.JsonReader"/>. 
+            </summary>
+            <param name="reader">A <see cref="T:Newtonsoft.Json.JsonReader"/> that will be read for the content of the <see cref="T:Newtonsoft.Json.Linq.JProperty"/>.</param>
+            <returns>A <see cref="T:Newtonsoft.Json.Linq.JProperty"/> that contains the JSON that was read from the specified <see cref="T:Newtonsoft.Json.JsonReader"/>.</returns>
+        </member>
+        <member name="P:Newtonsoft.Json.Linq.JProperty.ChildrenTokens">
+            <summary>
+            Gets the container's children tokens.
+            </summary>
+            <value>The container's children tokens.</value>
+        </member>
+        <member name="P:Newtonsoft.Json.Linq.JProperty.Name">
+            <summary>
+            Gets the property name.
+            </summary>
+            <value>The property name.</value>
+        </member>
+        <member name="P:Newtonsoft.Json.Linq.JProperty.Value">
+            <summary>
+            Gets or sets the property value.
+            </summary>
+            <value>The property value.</value>
+        </member>
+        <member name="P:Newtonsoft.Json.Linq.JProperty.Type">
+            <summary>
+            Gets the node type for this <see cref="T:Newtonsoft.Json.Linq.JToken"/>.
+            </summary>
+            <value>The type.</value>
+        </member>
+        <member name="T:Newtonsoft.Json.Linq.JTokenType">
+            <summary>
+            Specifies the type of token.
+            </summary>
+        </member>
+        <member name="F:Newtonsoft.Json.Linq.JTokenType.None">
+            <summary>
+            No token type has been set.
+            </summary>
+        </member>
+        <member name="F:Newtonsoft.Json.Linq.JTokenType.Object">
+            <summary>
+            A JSON object.
+            </summary>
+        </member>
+        <member name="F:Newtonsoft.Json.Linq.JTokenType.Array">
+            <summary>
+            A JSON array.
+            </summary>
+        </member>
+        <member name="F:Newtonsoft.Json.Linq.JTokenType.Constructor">
+            <summary>
+            A JSON constructor.
+            </summary>
+        </member>
+        <member name="F:Newtonsoft.Json.Linq.JTokenType.Property">
+            <summary>
+            A JSON object property.
+            </summary>
+        </member>
+        <member name="F:Newtonsoft.Json.Linq.JTokenType.Comment">
+            <summary>
+            A comment.
+            </summary>
+        </member>
+        <member name="F:Newtonsoft.Json.Linq.JTokenType.Integer">
+            <summary>
+            An integer value.
+            </summary>
+        </member>
+        <member name="F:Newtonsoft.Json.Linq.JTokenType.Float">
+            <summary>
+            A float value.
+            </summary>
+        </member>
+        <member name="F:Newtonsoft.Json.Linq.JTokenType.String">
+            <summary>
+            A string value.
+            </summary>
+        </member>
+        <member name="F:Newtonsoft.Json.Linq.JTokenType.Boolean">
+            <summary>
+            A boolean value.
+            </summary>
+        </member>
+        <member name="F:Newtonsoft.Json.Linq.JTokenType.Null">
+            <summary>
+            A null value.
+            </summary>
+        </member>
+        <member name="F:Newtonsoft.Json.Linq.JTokenType.Undefined">
+            <summary>
+            An undefined value.
+            </summary>
+        </member>
+        <member name="F:Newtonsoft.Json.Linq.JTokenType.Date">
+            <summary>
+            A date value.
+            </summary>
+        </member>
+        <member name="F:Newtonsoft.Json.Linq.JTokenType.Raw">
+            <summary>
+            A raw JSON value.
+            </summary>
+        </member>
+        <member name="F:Newtonsoft.Json.Linq.JTokenType.Bytes">
+            <summary>
+            A collection of bytes value.
+            </summary>
+        </member>
+        <member name="F:Newtonsoft.Json.Linq.JTokenType.Guid">
+            <summary>
+            A Guid value.
+            </summary>
+        </member>
+        <member name="F:Newtonsoft.Json.Linq.JTokenType.Uri">
+            <summary>
+            A Uri value.
+            </summary>
+        </member>
+        <member name="F:Newtonsoft.Json.Linq.JTokenType.TimeSpan">
+            <summary>
+            A TimeSpan value.
+            </summary>
+        </member>
+        <member name="T:Newtonsoft.Json.Schema.Extensions">
+            <summary>
+            Contains the JSON schema extension methods.
+            </summary>
+        </member>
+        <member name="M:Newtonsoft.Json.Schema.Extensions.IsValid(Newtonsoft.Json.Linq.JToken,Newtonsoft.Json.Schema.JsonSchema)">
+            <summary>
+            Determines whether the <see cref="T:Newtonsoft.Json.Linq.JToken"/> is valid.
+            </summary>
+            <param name="source">The source <see cref="T:Newtonsoft.Json.Linq.JToken"/> to test.</param>
+            <param name="schema">The schema to test with.</param>
+            <returns>
+            	<c>true</c> if the specified <see cref="T:Newtonsoft.Json.Linq.JToken"/> is valid; otherwise, <c>false</c>.
+            </returns>
+        </member>
+        <member name="M:Newtonsoft.Json.Schema.Extensions.IsValid(Newtonsoft.Json.Linq.JToken,Newtonsoft.Json.Schema.JsonSchema,System.Collections.Generic.IList{System.String}@)">
+            <summary>
+            Determines whether the <see cref="T:Newtonsoft.Json.Linq.JToken"/> is valid.
+            </summary>
+            <param name="source">The source <see cref="T:Newtonsoft.Json.Linq.JToken"/> to test.</param>
+            <param name="schema">The schema to test with.</param>
+            <param name="errorMessages">When this method returns, contains any error messages generated while validating. </param>
+            <returns>
+            	<c>true</c> if the specified <see cref="T:Newtonsoft.Json.Linq.JToken"/> is valid; otherwise, <c>false</c>.
+            </returns>
+        </member>
+        <member name="M:Newtonsoft.Json.Schema.Extensions.Validate(Newtonsoft.Json.Linq.JToken,Newtonsoft.Json.Schema.JsonSchema)">
+            <summary>
+            Validates the specified <see cref="T:Newtonsoft.Json.Linq.JToken"/>.
+            </summary>
+            <param name="source">The source <see cref="T:Newtonsoft.Json.Linq.JToken"/> to test.</param>
+            <param name="schema">The schema to test with.</param>
+        </member>
+        <member name="M:Newtonsoft.Json.Schema.Extensions.Validate(Newtonsoft.Json.Linq.JToken,Newtonsoft.Json.Schema.JsonSchema,Newtonsoft.Json.Schema.ValidationEventHandler)">
+            <summary>
+            Validates the specified <see cref="T:Newtonsoft.Json.Linq.JToken"/>.
+            </summary>
+            <param name="source">The source <see cref="T:Newtonsoft.Json.Linq.JToken"/> to test.</param>
+            <param name="schema">The schema to test with.</param>
+            <param name="validationEventHandler">The validation event handler.</param>
+        </member>
+        <member name="T:Newtonsoft.Json.Schema.JsonSchemaException">
+            <summary>
+            Returns detailed information about the schema exception.
+            </summary>
+        </member>
+        <member name="M:Newtonsoft.Json.Schema.JsonSchemaException.#ctor">
+            <summary>
+            Initializes a new instance of the <see cref="T:Newtonsoft.Json.Schema.JsonSchemaException"/> class.
+            </summary>
+        </member>
+        <member name="M:Newtonsoft.Json.Schema.JsonSchemaException.#ctor(System.String)">
+            <summary>
+            Initializes a new instance of the <see cref="T:Newtonsoft.Json.Schema.JsonSchemaException"/> class
+            with a specified error message.
+            </summary>
+            <param name="message">The error message that explains the reason for the exception.</param>
+        </member>
+        <member name="M:Newtonsoft.Json.Schema.JsonSchemaException.#ctor(System.String,System.Exception)">
+            <summary>
+            Initializes a new instance of the <see cref="T:Newtonsoft.Json.Schema.JsonSchemaException"/> class
+            with a specified error message and a reference to the inner exception that is the cause of this exception.
+            </summary>
+            <param name="message">The error message that explains the reason for the exception.</param>
+            <param name="innerException">The exception that is the cause of the current exception, or a null reference (Nothing in Visual Basic) if no inner exception is specified.</param>
+        </member>
+        <member name="M:Newtonsoft.Json.Schema.JsonSchemaException.#ctor(System.Runtime.Serialization.SerializationInfo,System.Runtime.Serialization.StreamingContext)">
+            <summary>
+            Initializes a new instance of the <see cref="T:Newtonsoft.Json.Schema.JsonSchemaException"/> class.
+            </summary>
+            <param name="info">The <see cref="T:System.Runtime.Serialization.SerializationInfo"/> that holds the serialized object data about the exception being thrown.</param>
+            <param name="context">The <see cref="T:System.Runtime.Serialization.StreamingContext"/> that contains contextual information about the source or destination.</param>
+            <exception cref="T:System.ArgumentNullException">The <paramref name="info"/> parameter is null. </exception>
+            <exception cref="T:System.Runtime.Serialization.SerializationException">The class name is null or <see cref="P:System.Exception.HResult"/> is zero (0). </exception>
+        </member>
+        <member name="P:Newtonsoft.Json.Schema.JsonSchemaException.LineNumber">
+            <summary>
+            Gets the line number indicating where the error occurred.
+            </summary>
+            <value>The line number indicating where the error occurred.</value>
+        </member>
+        <member name="P:Newtonsoft.Json.Schema.JsonSchemaException.LinePosition">
+            <summary>
+            Gets the line position indicating where the error occurred.
+            </summary>
+            <value>The line position indicating where the error occurred.</value>
+        </member>
+        <member name="P:Newtonsoft.Json.Schema.JsonSchemaException.Path">
+            <summary>
+            Gets the path to the JSON where the error occurred.
+            </summary>
+            <value>The path to the JSON where the error occurred.</value>
+        </member>
+        <member name="T:Newtonsoft.Json.Schema.JsonSchemaResolver">
+            <summary>
+            Resolves <see cref="T:Newtonsoft.Json.Schema.JsonSchema"/> from an id.
+            </summary>
+        </member>
+        <member name="M:Newtonsoft.Json.Schema.JsonSchemaResolver.#ctor">
+            <summary>
+            Initializes a new instance of the <see cref="T:Newtonsoft.Json.Schema.JsonSchemaResolver"/> class.
+            </summary>
+        </member>
+        <member name="M:Newtonsoft.Json.Schema.JsonSchemaResolver.GetSchema(System.String)">
+            <summary>
+            Gets a <see cref="T:Newtonsoft.Json.Schema.JsonSchema"/> for the specified reference.
+            </summary>
+            <param name="reference">The id.</param>
+            <returns>A <see cref="T:Newtonsoft.Json.Schema.JsonSchema"/> for the specified reference.</returns>
+        </member>
+        <member name="P:Newtonsoft.Json.Schema.JsonSchemaResolver.LoadedSchemas">
+            <summary>
+            Gets or sets the loaded schemas.
+            </summary>
+            <value>The loaded schemas.</value>
+        </member>
+        <member name="T:Newtonsoft.Json.Schema.UndefinedSchemaIdHandling">
+            <summary>
+            Specifies undefined schema Id handling options for the <see cref="T:Newtonsoft.Json.Schema.JsonSchemaGenerator"/>.
+            </summary>
+        </member>
+        <member name="F:Newtonsoft.Json.Schema.UndefinedSchemaIdHandling.None">
+            <summary>
+            Do not infer a schema Id.
+            </summary>
+        </member>
+        <member name="F:Newtonsoft.Json.Schema.UndefinedSchemaIdHandling.UseTypeName">
+            <summary>
+            Use the .NET type name as the schema Id.
+            </summary>
+        </member>
+        <member name="F:Newtonsoft.Json.Schema.UndefinedSchemaIdHandling.UseAssemblyQualifiedName">
+            <summary>
+            Use the assembly qualified .NET type name as the schema Id.
+            </summary>
+        </member>
+        <member name="T:Newtonsoft.Json.Schema.ValidationEventArgs">
+            <summary>
+            Returns detailed information related to the <see cref="T:Newtonsoft.Json.Schema.ValidationEventHandler"/>.
+            </summary>
+        </member>
+        <member name="P:Newtonsoft.Json.Schema.ValidationEventArgs.Exception">
+            <summary>
+            Gets the <see cref="T:Newtonsoft.Json.Schema.JsonSchemaException"/> associated with the validation error.
+            </summary>
+            <value>The JsonSchemaException associated with the validation error.</value>
+        </member>
+        <member name="P:Newtonsoft.Json.Schema.ValidationEventArgs.Path">
+            <summary>
+            Gets the path of the JSON location where the validation error occurred.
+            </summary>
+            <value>The path of the JSON location where the validation error occurred.</value>
+        </member>
+        <member name="P:Newtonsoft.Json.Schema.ValidationEventArgs.Message">
+            <summary>
+            Gets the text description corresponding to the validation error.
+            </summary>
+            <value>The text description.</value>
+        </member>
+        <member name="T:Newtonsoft.Json.Schema.ValidationEventHandler">
+            <summary>
+            Represents the callback method that will handle JSON schema validation events and the <see cref="T:Newtonsoft.Json.Schema.ValidationEventArgs"/>.
+            </summary>
+        </member>
+        <member name="T:Newtonsoft.Json.Serialization.CamelCasePropertyNamesContractResolver">
+            <summary>
+            Resolves member mappings for a type, camel casing property names.
+            </summary>
+        </member>
+        <member name="T:Newtonsoft.Json.Serialization.DefaultContractResolver">
+            <summary>
+            Used by <see cref="T:Newtonsoft.Json.JsonSerializer"/> to resolves a <see cref="T:Newtonsoft.Json.Serialization.JsonContract"/> for a given <see cref="T:System.Type"/>.
+            </summary>
+        </member>
+        <member name="T:Newtonsoft.Json.Serialization.IContractResolver">
+            <summary>
+            Used by <see cref="T:Newtonsoft.Json.JsonSerializer"/> to resolves a <see cref="T:Newtonsoft.Json.Serialization.JsonContract"/> for a given <see cref="T:System.Type"/>.
+            </summary>
+            <example>
+              <code lang="cs" source="..\Src\Newtonsoft.Json.Tests\Documentation\SerializationTests.cs" region="ReducingSerializedJsonSizeContractResolverObject" title="IContractResolver Class"/>
+              <code lang="cs" source="..\Src\Newtonsoft.Json.Tests\Documentation\SerializationTests.cs" region="ReducingSerializedJsonSizeContractResolverExample" title="IContractResolver Example"/>
+            </example>
+        </member>
+        <member name="M:Newtonsoft.Json.Serialization.IContractResolver.ResolveContract(System.Type)">
+            <summary>
+            Resolves the contract for a given type.
+            </summary>
+            <param name="type">The type to resolve a contract for.</param>
+            <returns>The contract for a given type.</returns>
+        </member>
+        <member name="M:Newtonsoft.Json.Serialization.DefaultContractResolver.#ctor">
+            <summary>
+            Initializes a new instance of the <see cref="T:Newtonsoft.Json.Serialization.DefaultContractResolver"/> class.
+            </summary>
+        </member>
+        <member name="M:Newtonsoft.Json.Serialization.DefaultContractResolver.#ctor(System.Boolean)">
+            <summary>
+            Initializes a new instance of the <see cref="T:Newtonsoft.Json.Serialization.DefaultContractResolver"/> class.
+            </summary>
+            <param name="shareCache">
+            If set to <c>true</c> the <see cref="T:Newtonsoft.Json.Serialization.DefaultContractResolver"/> will use a cached shared with other resolvers of the same type.
+            Sharing the cache will significantly performance because expensive reflection will only happen once but could cause unexpected
+            behavior if different instances of the resolver are suppose to produce different results. When set to false it is highly
+            recommended to reuse <see cref="T:Newtonsoft.Json.Serialization.DefaultContractResolver"/> instances with the <see cref="T:Newtonsoft.Json.JsonSerializer"/>.
+            </param>
+        </member>
+        <member name="M:Newtonsoft.Json.Serialization.DefaultContractResolver.ResolveContract(System.Type)">
+            <summary>
+            Resolves the contract for a given type.
+            </summary>
+            <param name="type">The type to resolve a contract for.</param>
+            <returns>The contract for a given type.</returns>
+        </member>
+        <member name="M:Newtonsoft.Json.Serialization.DefaultContractResolver.GetSerializableMembers(System.Type)">
+            <summary>
+            Gets the serializable members for the type.
+            </summary>
+            <param name="objectType">The type to get serializable members for.</param>
+            <returns>The serializable members for the type.</returns>
+        </member>
+        <member name="M:Newtonsoft.Json.Serialization.DefaultContractResolver.CreateObjectContract(System.Type)">
+            <summary>
+            Creates a <see cref="T:Newtonsoft.Json.Serialization.JsonObjectContract"/> for the given type.
+            </summary>
+            <param name="objectType">Type of the object.</param>
+            <returns>A <see cref="T:Newtonsoft.Json.Serialization.JsonObjectContract"/> for the given type.</returns>
+        </member>
+        <member name="M:Newtonsoft.Json.Serialization.DefaultContractResolver.CreateConstructorParameters(System.Reflection.ConstructorInfo,Newtonsoft.Json.Serialization.JsonPropertyCollection)">
+            <summary>
+            Creates the constructor parameters.
+            </summary>
+            <param name="constructor">The constructor to create properties for.</param>
+            <param name="memberProperties">The type's member properties.</param>
+            <returns>Properties for the given <see cref="T:System.Reflection.ConstructorInfo"/>.</returns>
+        </member>
+        <member name="M:Newtonsoft.Json.Serialization.DefaultContractResolver.CreatePropertyFromConstructorParameter(Newtonsoft.Json.Serialization.JsonProperty,System.Reflection.ParameterInfo)">
+            <summary>
+            Creates a <see cref="T:Newtonsoft.Json.Serialization.JsonProperty"/> for the given <see cref="T:System.Reflection.ParameterInfo"/>.
+            </summary>
+            <param name="matchingMemberProperty">The matching member property.</param>
+            <param name="parameterInfo">The constructor parameter.</param>
+            <returns>A created <see cref="T:Newtonsoft.Json.Serialization.JsonProperty"/> for the given <see cref="T:System.Reflection.ParameterInfo"/>.</returns>
+        </member>
+        <member name="M:Newtonsoft.Json.Serialization.DefaultContractResolver.ResolveContractConverter(System.Type)">
+            <summary>
+            Resolves the default <see cref="T:Newtonsoft.Json.JsonConverter"/> for the contract.
+            </summary>
+            <param name="objectType">Type of the object.</param>
+            <returns>The contract's default <see cref="T:Newtonsoft.Json.JsonConverter"/>.</returns>
+        </member>
+        <member name="M:Newtonsoft.Json.Serialization.DefaultContractResolver.CreateDictionaryContract(System.Type)">
+            <summary>
+            Creates a <see cref="T:Newtonsoft.Json.Serialization.JsonDictionaryContract"/> for the given type.
+            </summary>
+            <param name="objectType">Type of the object.</param>
+            <returns>A <see cref="T:Newtonsoft.Json.Serialization.JsonDictionaryContract"/> for the given type.</returns>
+        </member>
+        <member name="M:Newtonsoft.Json.Serialization.DefaultContractResolver.CreateArrayContract(System.Type)">
+            <summary>
+            Creates a <see cref="T:Newtonsoft.Json.Serialization.JsonArrayContract"/> for the given type.
+            </summary>
+            <param name="objectType">Type of the object.</param>
+            <returns>A <see cref="T:Newtonsoft.Json.Serialization.JsonArrayContract"/> for the given type.</returns>
+        </member>
+        <member name="M:Newtonsoft.Json.Serialization.DefaultContractResolver.CreatePrimitiveContract(System.Type)">
+            <summary>
+            Creates a <see cref="T:Newtonsoft.Json.Serialization.JsonPrimitiveContract"/> for the given type.
+            </summary>
+            <param name="objectType">Type of the object.</param>
+            <returns>A <see cref="T:Newtonsoft.Json.Serialization.JsonPrimitiveContract"/> for the given type.</returns>
+        </member>
+        <member name="M:Newtonsoft.Json.Serialization.DefaultContractResolver.CreateLinqContract(System.Type)">
+            <summary>
+            Creates a <see cref="T:Newtonsoft.Json.Serialization.JsonLinqContract"/> for the given type.
+            </summary>
+            <param name="objectType">Type of the object.</param>
+            <returns>A <see cref="T:Newtonsoft.Json.Serialization.JsonLinqContract"/> for the given type.</returns>
+        </member>
+        <member name="M:Newtonsoft.Json.Serialization.DefaultContractResolver.CreateISerializableContract(System.Type)">
+            <summary>
+            Creates a <see cref="T:Newtonsoft.Json.Serialization.JsonISerializableContract"/> for the given type.
+            </summary>
+            <param name="objectType">Type of the object.</param>
+            <returns>A <see cref="T:Newtonsoft.Json.Serialization.JsonISerializableContract"/> for the given type.</returns>
+        </member>
+        <member name="M:Newtonsoft.Json.Serialization.DefaultContractResolver.CreateStringContract(System.Type)">
+            <summary>
+            Creates a <see cref="T:Newtonsoft.Json.Serialization.JsonStringContract"/> for the given type.
+            </summary>
+            <param name="objectType">Type of the object.</param>
+            <returns>A <see cref="T:Newtonsoft.Json.Serialization.JsonStringContract"/> for the given type.</returns>
+        </member>
+        <member name="M:Newtonsoft.Json.Serialization.DefaultContractResolver.CreateContract(System.Type)">
+            <summary>
+            Determines which contract type is created for the given type.
+            </summary>
+            <param name="objectType">Type of the object.</param>
+            <returns>A <see cref="T:Newtonsoft.Json.Serialization.JsonContract"/> for the given type.</returns>
+        </member>
+        <member name="M:Newtonsoft.Json.Serialization.DefaultContractResolver.CreateProperties(System.Type,Newtonsoft.Json.MemberSerialization)">
+            <summary>
+            Creates properties for the given <see cref="T:Newtonsoft.Json.Serialization.JsonContract"/>.
+            </summary>
+            <param name="type">The type to create properties for.</param>
+            /// <param name="memberSerialization">The member serialization mode for the type.</param>
+            <returns>Properties for the given <see cref="T:Newtonsoft.Json.Serialization.JsonContract"/>.</returns>
+        </member>
+        <member name="M:Newtonsoft.Json.Serialization.DefaultContractResolver.CreateMemberValueProvider(System.Reflection.MemberInfo)">
+            <summary>
+            Creates the <see cref="T:Newtonsoft.Json.Serialization.IValueProvider"/> used by the serializer to get and set values from a member.
+            </summary>
+            <param name="member">The member.</param>
+            <returns>The <see cref="T:Newtonsoft.Json.Serialization.IValueProvider"/> used by the serializer to get and set values from a member.</returns>
+        </member>
+        <member name="M:Newtonsoft.Json.Serialization.DefaultContractResolver.CreateProperty(System.Reflection.MemberInfo,Newtonsoft.Json.MemberSerialization)">
+            <summary>
+            Creates a <see cref="T:Newtonsoft.Json.Serialization.JsonProperty"/> for the given <see cref="T:System.Reflection.MemberInfo"/>.
+            </summary>
+            <param name="memberSerialization">The member's parent <see cref="T:Newtonsoft.Json.MemberSerialization"/>.</param>
+            <param name="member">The member to create a <see cref="T:Newtonsoft.Json.Serialization.JsonProperty"/> for.</param>
+            <returns>A created <see cref="T:Newtonsoft.Json.Serialization.JsonProperty"/> for the given <see cref="T:System.Reflection.MemberInfo"/>.</returns>
+        </member>
+        <member name="M:Newtonsoft.Json.Serialization.DefaultContractResolver.ResolvePropertyName(System.String)">
+            <summary>
+            Resolves the name of the property.
+            </summary>
+            <param name="propertyName">Name of the property.</param>
+            <returns>Name of the property.</returns>
+        </member>
+        <member name="M:Newtonsoft.Json.Serialization.DefaultContractResolver.GetResolvedPropertyName(System.String)">
+            <summary>
+            Gets the resolved name of the property.
+            </summary>
+            <param name="propertyName">Name of the property.</param>
+            <returns>Name of the property.</returns>
+        </member>
+        <member name="P:Newtonsoft.Json.Serialization.DefaultContractResolver.DynamicCodeGeneration">
+            <summary>
+            Gets a value indicating whether members are being get and set using dynamic code generation.
+            This value is determined by the runtime permissions available.
+            </summary>
+            <value>
+            	<c>true</c> if using dynamic code generation; otherwise, <c>false</c>.
+            </value>
+        </member>
+        <member name="P:Newtonsoft.Json.Serialization.DefaultContractResolver.DefaultMembersSearchFlags">
+            <summary>
+            Gets or sets the default members search flags.
+            </summary>
+            <value>The default members search flags.</value>
+        </member>
+        <member name="P:Newtonsoft.Json.Serialization.DefaultContractResolver.SerializeCompilerGeneratedMembers">
+            <summary>
+            Gets or sets a value indicating whether compiler generated members should be serialized.
+            </summary>
+            <value>
+            	<c>true</c> if serialized compiler generated members; otherwise, <c>false</c>.
+            </value>
+        </member>
+        <member name="P:Newtonsoft.Json.Serialization.DefaultContractResolver.IgnoreSerializableInterface">
+            <summary>
+            Gets or sets a value indicating whether to ignore the <see cref="T:System.Runtime.Serialization.ISerializable"/> interface when serializing and deserializing types.
+            </summary>
+            <value>
+            	<c>true</c> if the <see cref="T:System.Runtime.Serialization.ISerializable"/> interface will be ignored when serializing and deserializing types; otherwise, <c>false</c>.
+            </value>
+        </member>
+        <member name="P:Newtonsoft.Json.Serialization.DefaultContractResolver.IgnoreSerializableAttribute">
+            <summary>
+            Gets or sets a value indicating whether to ignore the <see cref="T:System.SerializableAttribute"/> attribute when serializing and deserializing types.
+            </summary>
+            <value>
+            	<c>true</c> if the <see cref="T:System.SerializableAttribute"/> attribute will be ignored when serializing and deserializing types; otherwise, <c>false</c>.
+            </value>
+        </member>
+        <member name="M:Newtonsoft.Json.Serialization.CamelCasePropertyNamesContractResolver.#ctor">
+            <summary>
+            Initializes a new instance of the <see cref="T:Newtonsoft.Json.Serialization.CamelCasePropertyNamesContractResolver"/> class.
+            </summary>
+        </member>
+        <member name="M:Newtonsoft.Json.Serialization.CamelCasePropertyNamesContractResolver.ResolvePropertyName(System.String)">
+            <summary>
+            Resolves the name of the property.
+            </summary>
+            <param name="propertyName">Name of the property.</param>
+            <returns>The property name camel cased.</returns>
+        </member>
+        <member name="T:Newtonsoft.Json.Serialization.DefaultSerializationBinder">
+            <summary>
+            The default serialization binder used when resolving and loading classes from type names.
+            </summary>
+        </member>
+        <member name="M:Newtonsoft.Json.Serialization.DefaultSerializationBinder.BindToType(System.String,System.String)">
+            <summary>
+            When overridden in a derived class, controls the binding of a serialized object to a type.
+            </summary>
+            <param name="assemblyName">Specifies the <see cref="T:System.Reflection.Assembly"/> name of the serialized object.</param>
+            <param name="typeName">Specifies the <see cref="T:System.Type"/> name of the serialized object.</param>
+            <returns>
+            The type of the object the formatter creates a new instance of.
+            </returns>
+        </member>
+        <member name="T:Newtonsoft.Json.Serialization.DynamicValueProvider">
+            <summary>
+            Get and set values for a <see cref="T:System.Reflection.MemberInfo"/> using dynamic methods.
+            </summary>
+        </member>
+        <member name="T:Newtonsoft.Json.Serialization.IValueProvider">
+            <summary>
+            Provides methods to get and set values.
+            </summary>
+        </member>
+        <member name="M:Newtonsoft.Json.Serialization.IValueProvider.SetValue(System.Object,System.Object)">
+            <summary>
+            Sets the value.
+            </summary>
+            <param name="target">The target to set the value on.</param>
+            <param name="value">The value to set on the target.</param>
+        </member>
+        <member name="M:Newtonsoft.Json.Serialization.IValueProvider.GetValue(System.Object)">
+            <summary>
+            Gets the value.
+            </summary>
+            <param name="target">The target to get the value from.</param>
+            <returns>The value.</returns>
+        </member>
+        <member name="M:Newtonsoft.Json.Serialization.DynamicValueProvider.#ctor(System.Reflection.MemberInfo)">
+            <summary>
+            Initializes a new instance of the <see cref="T:Newtonsoft.Json.Serialization.DynamicValueProvider"/> class.
+            </summary>
+            <param name="memberInfo">The member info.</param>
+        </member>
+        <member name="M:Newtonsoft.Json.Serialization.DynamicValueProvider.SetValue(System.Object,System.Object)">
+            <summary>
+            Sets the value.
+            </summary>
+            <param name="target">The target to set the value on.</param>
+            <param name="value">The value to set on the target.</param>
+        </member>
+        <member name="M:Newtonsoft.Json.Serialization.DynamicValueProvider.GetValue(System.Object)">
+            <summary>
+            Gets the value.
+            </summary>
+            <param name="target">The target to get the value from.</param>
+            <returns>The value.</returns>
+        </member>
+        <member name="T:Newtonsoft.Json.Serialization.ErrorContext">
+            <summary>
+            Provides information surrounding an error.
+            </summary>
+        </member>
+        <member name="P:Newtonsoft.Json.Serialization.ErrorContext.Error">
+            <summary>
+            Gets or sets the error.
+            </summary>
+            <value>The error.</value>
+        </member>
+        <member name="P:Newtonsoft.Json.Serialization.ErrorContext.OriginalObject">
+            <summary>
+            Gets the original object that caused the error.
+            </summary>
+            <value>The original object that caused the error.</value>
+        </member>
+        <member name="P:Newtonsoft.Json.Serialization.ErrorContext.Member">
+            <summary>
+            Gets the member that caused the error.
+            </summary>
+            <value>The member that caused the error.</value>
+        </member>
+        <member name="P:Newtonsoft.Json.Serialization.ErrorContext.Path">
+            <summary>
+            Gets the path of the JSON location where the error occurred.
+            </summary>
+            <value>The path of the JSON location where the error occurred.</value>
+        </member>
+        <member name="P:Newtonsoft.Json.Serialization.ErrorContext.Handled">
+            <summary>
+            Gets or sets a value indicating whether this <see cref="T:Newtonsoft.Json.Serialization.ErrorContext"/> is handled.
+            </summary>
+            <value><c>true</c> if handled; otherwise, <c>false</c>.</value>
+        </member>
+        <member name="T:Newtonsoft.Json.Serialization.ErrorEventArgs">
+            <summary>
+            Provides data for the Error event.
+            </summary>
+        </member>
+        <member name="M:Newtonsoft.Json.Serialization.ErrorEventArgs.#ctor(System.Object,Newtonsoft.Json.Serialization.ErrorContext)">
+            <summary>
+            Initializes a new instance of the <see cref="T:Newtonsoft.Json.Serialization.ErrorEventArgs"/> class.
+            </summary>
+            <param name="currentObject">The current object.</param>
+            <param name="errorContext">The error context.</param>
+        </member>
+        <member name="P:Newtonsoft.Json.Serialization.ErrorEventArgs.CurrentObject">
+            <summary>
+            Gets the current object the error event is being raised against.
+            </summary>
+            <value>The current object the error event is being raised against.</value>
+        </member>
+        <member name="P:Newtonsoft.Json.Serialization.ErrorEventArgs.ErrorContext">
+            <summary>
+            Gets the error context.
+            </summary>
+            <value>The error context.</value>
+        </member>
+        <member name="T:Newtonsoft.Json.Serialization.JsonArrayContract">
+            <summary>
+            Contract details for a <see cref="T:System.Type"/> used by the <see cref="T:Newtonsoft.Json.JsonSerializer"/>.
+            </summary>
+        </member>
+        <member name="M:Newtonsoft.Json.Serialization.JsonArrayContract.#ctor(System.Type)">
+            <summary>
+            Initializes a new instance of the <see cref="T:Newtonsoft.Json.Serialization.JsonArrayContract"/> class.
+            </summary>
+            <param name="underlyingType">The underlying type for the contract.</param>
+        </member>
+        <member name="P:Newtonsoft.Json.Serialization.JsonArrayContract.CollectionItemType">
+            <summary>
+            Gets the <see cref="T:System.Type"/> of the collection items.
+            </summary>
+            <value>The <see cref="T:System.Type"/> of the collection items.</value>
+        </member>
+        <member name="P:Newtonsoft.Json.Serialization.JsonArrayContract.IsMultidimensionalArray">
+            <summary>
+            Gets a value indicating whether the collection type is a multidimensional array.
+            </summary>
+            <value><c>true</c> if the collection type is a multidimensional array; otherwise, <c>false</c>.</value>
+        </member>
+        <member name="T:Newtonsoft.Json.Serialization.SerializationCallback">
+            <summary>
+            Handles <see cref="T:Newtonsoft.Json.JsonSerializer"/> serialization callback events.
+            </summary>
+            <param name="o">The object that raised the callback event.</param>
+            <param name="context">The streaming context.</param>
+        </member>
+        <member name="T:Newtonsoft.Json.Serialization.SerializationErrorCallback">
+            <summary>
+            Handles <see cref="T:Newtonsoft.Json.JsonSerializer"/> serialization error callback events.
+            </summary>
+            <param name="o">The object that raised the callback event.</param>
+            <param name="context">The streaming context.</param>
+            <param name="errorContext">The error context.</param>
+        </member>
+        <member name="T:Newtonsoft.Json.Serialization.JsonDictionaryContract">
+            <summary>
+            Contract details for a <see cref="T:System.Type"/> used by the <see cref="T:Newtonsoft.Json.JsonSerializer"/>.
+            </summary>
+        </member>
+        <member name="M:Newtonsoft.Json.Serialization.JsonDictionaryContract.#ctor(System.Type)">
+            <summary>
+            Initializes a new instance of the <see cref="T:Newtonsoft.Json.Serialization.JsonDictionaryContract"/> class.
+            </summary>
+            <param name="underlyingType">The underlying type for the contract.</param>
+        </member>
+        <member name="P:Newtonsoft.Json.Serialization.JsonDictionaryContract.PropertyNameResolver">
+            <summary>
+            Gets or sets the property name resolver.
+            </summary>
+            <value>The property name resolver.</value>
+        </member>
+        <member name="P:Newtonsoft.Json.Serialization.JsonDictionaryContract.DictionaryKeyType">
+            <summary>
+            Gets the <see cref="T:System.Type"/> of the dictionary keys.
+            </summary>
+            <value>The <see cref="T:System.Type"/> of the dictionary keys.</value>
+        </member>
+        <member name="P:Newtonsoft.Json.Serialization.JsonDictionaryContract.DictionaryValueType">
+            <summary>
+            Gets the <see cref="T:System.Type"/> of the dictionary values.
+            </summary>
+            <value>The <see cref="T:System.Type"/> of the dictionary values.</value>
+        </member>
+        <member name="T:Newtonsoft.Json.Serialization.JsonISerializableContract">
+            <summary>
+            Contract details for a <see cref="T:System.Type"/> used by the <see cref="T:Newtonsoft.Json.JsonSerializer"/>.
+            </summary>
+        </member>
+        <member name="M:Newtonsoft.Json.Serialization.JsonISerializableContract.#ctor(System.Type)">
+            <summary>
+            Initializes a new instance of the <see cref="T:Newtonsoft.Json.Serialization.JsonISerializableContract"/> class.
+            </summary>
+            <param name="underlyingType">The underlying type for the contract.</param>
+        </member>
+        <member name="P:Newtonsoft.Json.Serialization.JsonISerializableContract.ISerializableCreator">
+            <summary>
+            Gets or sets the ISerializable object constructor.
+            </summary>
+            <value>The ISerializable object constructor.</value>
+        </member>
+        <member name="T:Newtonsoft.Json.Serialization.JsonLinqContract">
+            <summary>
+            Contract details for a <see cref="T:System.Type"/> used by the <see cref="T:Newtonsoft.Json.JsonSerializer"/>.
+            </summary>
+        </member>
+        <member name="M:Newtonsoft.Json.Serialization.JsonLinqContract.#ctor(System.Type)">
+            <summary>
+            Initializes a new instance of the <see cref="T:Newtonsoft.Json.Serialization.JsonLinqContract"/> class.
+            </summary>
+            <param name="underlyingType">The underlying type for the contract.</param>
+        </member>
+        <member name="T:Newtonsoft.Json.Serialization.JsonPrimitiveContract">
+            <summary>
+            Contract details for a <see cref="T:System.Type"/> used by the <see cref="T:Newtonsoft.Json.JsonSerializer"/>.
+            </summary>
+        </member>
+        <member name="M:Newtonsoft.Json.Serialization.JsonPrimitiveContract.#ctor(System.Type)">
+            <summary>
+            Initializes a new instance of the <see cref="T:Newtonsoft.Json.Serialization.JsonPrimitiveContract"/> class.
+            </summary>
+            <param name="underlyingType">The underlying type for the contract.</param>
+        </member>
+        <member name="T:Newtonsoft.Json.Serialization.JsonProperty">
+            <summary>
+            Maps a JSON property to a .NET member or constructor parameter.
+            </summary>
+        </member>
+        <member name="M:Newtonsoft.Json.Serialization.JsonProperty.ToString">
+            <summary>
+            Returns a <see cref="T:System.String"/> that represents this instance.
+            </summary>
+            <returns>
+            A <see cref="T:System.String"/> that represents this instance.
+            </returns>
+        </member>
+        <member name="P:Newtonsoft.Json.Serialization.JsonProperty.PropertyName">
+            <summary>
+            Gets or sets the name of the property.
+            </summary>
+            <value>The name of the property.</value>
+        </member>
+        <member name="P:Newtonsoft.Json.Serialization.JsonProperty.DeclaringType">
+            <summary>
+            Gets or sets the type that declared this property.
+            </summary>
+            <value>The type that declared this property.</value>
+        </member>
+        <member name="P:Newtonsoft.Json.Serialization.JsonProperty.Order">
+            <summary>
+            Gets or sets the order of serialization and deserialization of a member.
+            </summary>
+            <value>The numeric order of serialization or deserialization.</value>
+        </member>
+        <member name="P:Newtonsoft.Json.Serialization.JsonProperty.UnderlyingName">
+            <summary>
+            Gets or sets the name of the underlying member or parameter.
+            </summary>
+            <value>The name of the underlying member or parameter.</value>
+        </member>
+        <member name="P:Newtonsoft.Json.Serialization.JsonProperty.ValueProvider">
+            <summary>
+            Gets the <see cref="T:Newtonsoft.Json.Serialization.IValueProvider"/> that will get and set the <see cref="T:Newtonsoft.Json.Serialization.JsonProperty"/> during serialization.
+            </summary>
+            <value>The <see cref="T:Newtonsoft.Json.Serialization.IValueProvider"/> that will get and set the <see cref="T:Newtonsoft.Json.Serialization.JsonProperty"/> during serialization.</value>
+        </member>
+        <member name="P:Newtonsoft.Json.Serialization.JsonProperty.PropertyType">
+            <summary>
+            Gets or sets the type of the property.
+            </summary>
+            <value>The type of the property.</value>
+        </member>
+        <member name="P:Newtonsoft.Json.Serialization.JsonProperty.Converter">
+            <summary>
+            Gets or sets the <see cref="T:Newtonsoft.Json.JsonConverter"/> for the property.
+            If set this converter takes presidence over the contract converter for the property type.
+            </summary>
+            <value>The converter.</value>
+        </member>
+        <member name="P:Newtonsoft.Json.Serialization.JsonProperty.MemberConverter">
+            <summary>
+            Gets the member converter.
+            </summary>
+            <value>The member converter.</value>
+        </member>
+        <member name="P:Newtonsoft.Json.Serialization.JsonProperty.Ignored">
+            <summary>
+            Gets a value indicating whether this <see cref="T:Newtonsoft.Json.Serialization.JsonProperty"/> is ignored.
+            </summary>
+            <value><c>true</c> if ignored; otherwise, <c>false</c>.</value>
+        </member>
+        <member name="P:Newtonsoft.Json.Serialization.JsonProperty.Readable">
+            <summary>
+            Gets a value indicating whether this <see cref="T:Newtonsoft.Json.Serialization.JsonProperty"/> is readable.
+            </summary>
+            <value><c>true</c> if readable; otherwise, <c>false</c>.</value>
+        </member>
+        <member name="P:Newtonsoft.Json.Serialization.JsonProperty.Writable">
+            <summary>
+            Gets a value indicating whether this <see cref="T:Newtonsoft.Json.Serialization.JsonProperty"/> is writable.
+            </summary>
+            <value><c>true</c> if writable; otherwise, <c>false</c>.</value>
+        </member>
+        <member name="P:Newtonsoft.Json.Serialization.JsonProperty.HasMemberAttribute">
+            <summary>
+            Gets a value indicating whether this <see cref="T:Newtonsoft.Json.Serialization.JsonProperty"/> has a member attribute.
+            </summary>
+            <value><c>true</c> if has a member attribute; otherwise, <c>false</c>.</value>
+        </member>
+        <member name="P:Newtonsoft.Json.Serialization.JsonProperty.DefaultValue">
+            <summary>
+            Gets the default value.
+            </summary>
+            <value>The default value.</value>
+        </member>
+        <member name="P:Newtonsoft.Json.Serialization.JsonProperty.Required">
+            <summary>
+            Gets a value indicating whether this <see cref="T:Newtonsoft.Json.Serialization.JsonProperty"/> is required.
+            </summary>
+            <value>A value indicating whether this <see cref="T:Newtonsoft.Json.Serialization.JsonProperty"/> is required.</value>
+        </member>
+        <member name="P:Newtonsoft.Json.Serialization.JsonProperty.IsReference">
+            <summary>
+            Gets a value indicating whether this property preserves object references.
+            </summary>
+            <value>
+            	<c>true</c> if this instance is reference; otherwise, <c>false</c>.
+            </value>
+        </member>
+        <member name="P:Newtonsoft.Json.Serialization.JsonProperty.NullValueHandling">
+            <summary>
+            Gets the property null value handling.
+            </summary>
+            <value>The null value handling.</value>
+        </member>
+        <member name="P:Newtonsoft.Json.Serialization.JsonProperty.DefaultValueHandling">
+            <summary>
+            Gets the property default value handling.
+            </summary>
+            <value>The default value handling.</value>
+        </member>
+        <member name="P:Newtonsoft.Json.Serialization.JsonProperty.ReferenceLoopHandling">
+            <summary>
+            Gets the property reference loop handling.
+            </summary>
+            <value>The reference loop handling.</value>
+        </member>
+        <member name="P:Newtonsoft.Json.Serialization.JsonProperty.ObjectCreationHandling">
+            <summary>
+            Gets the property object creation handling.
+            </summary>
+            <value>The object creation handling.</value>
+        </member>
+        <member name="P:Newtonsoft.Json.Serialization.JsonProperty.TypeNameHandling">
+            <summary>
+            Gets or sets the type name handling.
+            </summary>
+            <value>The type name handling.</value>
+        </member>
+        <member name="P:Newtonsoft.Json.Serialization.JsonProperty.ShouldSerialize">
+            <summary>
+            Gets or sets a predicate used to determine whether the property should be serialize.
+            </summary>
+            <value>A predicate used to determine whether the property should be serialize.</value>
+        </member>
+        <member name="P:Newtonsoft.Json.Serialization.JsonProperty.GetIsSpecified">
+            <summary>
+            Gets or sets a predicate used to determine whether the property should be serialized.
+            </summary>
+            <value>A predicate used to determine whether the property should be serialized.</value>
+        </member>
+        <member name="P:Newtonsoft.Json.Serialization.JsonProperty.SetIsSpecified">
+            <summary>
+            Gets or sets an action used to set whether the property has been deserialized.
+            </summary>
+            <value>An action used to set whether the property has been deserialized.</value>
+        </member>
+        <member name="P:Newtonsoft.Json.Serialization.JsonProperty.ItemConverter">
+            <summary>
+            Gets or sets the converter used when serializing the property's collection items.
+            </summary>
+            <value>The collection's items converter.</value>
+        </member>
+        <member name="P:Newtonsoft.Json.Serialization.JsonProperty.ItemIsReference">
+            <summary>
+            Gets or sets whether this property's collection items are serialized as a reference.
+            </summary>
+            <value>Whether this property's collection items are serialized as a reference.</value>
+        </member>
+        <member name="P:Newtonsoft.Json.Serialization.JsonProperty.ItemTypeNameHandling">
+            <summary>
+            Gets or sets the the type name handling used when serializing the property's collection items.
+            </summary>
+            <value>The collection's items type name handling.</value>
+        </member>
+        <member name="P:Newtonsoft.Json.Serialization.JsonProperty.ItemReferenceLoopHandling">
+            <summary>
+            Gets or sets the the reference loop handling used when serializing the property's collection items.
+            </summary>
+            <value>The collection's items reference loop handling.</value>
+        </member>
+        <member name="T:Newtonsoft.Json.Serialization.JsonPropertyCollection">
+            <summary>
+            A collection of <see cref="T:Newtonsoft.Json.Serialization.JsonProperty"/> objects.
+            </summary>
+        </member>
+        <member name="M:Newtonsoft.Json.Serialization.JsonPropertyCollection.#ctor(System.Type)">
+            <summary>
+            Initializes a new instance of the <see cref="T:Newtonsoft.Json.Serialization.JsonPropertyCollection"/> class.
+            </summary>
+            <param name="type">The type.</param>
+        </member>
+        <member name="M:Newtonsoft.Json.Serialization.JsonPropertyCollection.GetKeyForItem(Newtonsoft.Json.Serialization.JsonProperty)">
+            <summary>
+            When implemented in a derived class, extracts the key from the specified element.
+            </summary>
+            <param name="item">The element from which to extract the key.</param>
+            <returns>The key for the specified element.</returns>
+        </member>
+        <member name="M:Newtonsoft.Json.Serialization.JsonPropertyCollection.AddProperty(Newtonsoft.Json.Serialization.JsonProperty)">
+            <summary>
+            Adds a <see cref="T:Newtonsoft.Json.Serialization.JsonProperty"/> object.
+            </summary>
+            <param name="property">The property to add to the collection.</param>
+        </member>
+        <member name="M:Newtonsoft.Json.Serialization.JsonPropertyCollection.GetClosestMatchProperty(System.String)">
+            <summary>
+            Gets the closest matching <see cref="T:Newtonsoft.Json.Serialization.JsonProperty"/> object.
+            First attempts to get an exact case match of propertyName and then
+            a case insensitive match.
+            </summary>
+            <param name="propertyName">Name of the property.</param>
+            <returns>A matching property if found.</returns>
+        </member>
+        <member name="M:Newtonsoft.Json.Serialization.JsonPropertyCollection.GetProperty(System.String,System.StringComparison)">
+            <summary>
+            Gets a property by property name.
+            </summary>
+            <param name="propertyName">The name of the property to get.</param>
+            <param name="comparisonType">Type property name string comparison.</param>
+            <returns>A matching property if found.</returns>
+        </member>
+        <member name="T:Newtonsoft.Json.MissingMemberHandling">
+            <summary>
+            Specifies missing member handling options for the <see cref="T:Newtonsoft.Json.JsonSerializer"/>.
+            </summary>
+        </member>
+        <member name="F:Newtonsoft.Json.MissingMemberHandling.Ignore">
+            <summary>
+            Ignore a missing member and do not attempt to deserialize it.
+            </summary>
+        </member>
+        <member name="F:Newtonsoft.Json.MissingMemberHandling.Error">
+            <summary>
+            Throw a <see cref="T:Newtonsoft.Json.JsonSerializationException"/> when a missing member is encountered during deserialization.
+            </summary>
+        </member>
+        <member name="T:Newtonsoft.Json.NullValueHandling">
+            <summary>
+            Specifies null value handling options for the <see cref="T:Newtonsoft.Json.JsonSerializer"/>.
+            </summary>
+            <example>
+              <code lang="cs" source="..\Src\Newtonsoft.Json.Tests\Documentation\SerializationTests.cs" region="ReducingSerializedJsonSizeNullValueHandlingObject" title="NullValueHandling Class"/>
+              <code lang="cs" source="..\Src\Newtonsoft.Json.Tests\Documentation\SerializationTests.cs" region="ReducingSerializedJsonSizeNullValueHandlingExample" title="NullValueHandling Ignore Example"/>
+            </example>
+        </member>
+        <member name="F:Newtonsoft.Json.NullValueHandling.Include">
+            <summary>
+            Include null values when serializing and deserializing objects.
+            </summary>
+        </member>
+        <member name="F:Newtonsoft.Json.NullValueHandling.Ignore">
+            <summary>
+            Ignore null values when serializing and deserializing objects.
+            </summary>
+        </member>
+        <member name="T:Newtonsoft.Json.ReferenceLoopHandling">
+            <summary>
+            Specifies reference loop handling options for the <see cref="T:Newtonsoft.Json.JsonSerializer"/>.
+            </summary>
+        </member>
+        <member name="F:Newtonsoft.Json.ReferenceLoopHandling.Error">
+            <summary>
+            Throw a <see cref="T:Newtonsoft.Json.JsonSerializationException"/> when a loop is encountered.
+            </summary>
+        </member>
+        <member name="F:Newtonsoft.Json.ReferenceLoopHandling.Ignore">
+            <summary>
+            Ignore loop references and do not serialize.
+            </summary>
+        </member>
+        <member name="F:Newtonsoft.Json.ReferenceLoopHandling.Serialize">
+            <summary>
+            Serialize loop references.
+            </summary>
+        </member>
+        <member name="T:Newtonsoft.Json.Schema.JsonSchema">
+            <summary>
+            An in-memory representation of a JSON Schema.
+            </summary>
+        </member>
+        <member name="M:Newtonsoft.Json.Schema.JsonSchema.#ctor">
+            <summary>
+            Initializes a new instance of the <see cref="T:Newtonsoft.Json.Schema.JsonSchema"/> class.
+            </summary>
+        </member>
+        <member name="M:Newtonsoft.Json.Schema.JsonSchema.Read(Newtonsoft.Json.JsonReader)">
+            <summary>
+            Reads a <see cref="T:Newtonsoft.Json.Schema.JsonSchema"/> from the specified <see cref="T:Newtonsoft.Json.JsonReader"/>.
+            </summary>
+            <param name="reader">The <see cref="T:Newtonsoft.Json.JsonReader"/> containing the JSON Schema to read.</param>
+            <returns>The <see cref="T:Newtonsoft.Json.Schema.JsonSchema"/> object representing the JSON Schema.</returns>
+        </member>
+        <member name="M:Newtonsoft.Json.Schema.JsonSchema.Read(Newtonsoft.Json.JsonReader,Newtonsoft.Json.Schema.JsonSchemaResolver)">
+            <summary>
+            Reads a <see cref="T:Newtonsoft.Json.Schema.JsonSchema"/> from the specified <see cref="T:Newtonsoft.Json.JsonReader"/>.
+            </summary>
+            <param name="reader">The <see cref="T:Newtonsoft.Json.JsonReader"/> containing the JSON Schema to read.</param>
+            <param name="resolver">The <see cref="T:Newtonsoft.Json.Schema.JsonSchemaResolver"/> to use when resolving schema references.</param>
+            <returns>The <see cref="T:Newtonsoft.Json.Schema.JsonSchema"/> object representing the JSON Schema.</returns>
+        </member>
+        <member name="M:Newtonsoft.Json.Schema.JsonSchema.Parse(System.String)">
+            <summary>
+            Load a <see cref="T:Newtonsoft.Json.Schema.JsonSchema"/> from a string that contains schema JSON.
+            </summary>
+            <param name="json">A <see cref="T:System.String"/> that contains JSON.</param>
+            <returns>A <see cref="T:Newtonsoft.Json.Schema.JsonSchema"/> populated from the string that contains JSON.</returns>
+        </member>
+        <member name="M:Newtonsoft.Json.Schema.JsonSchema.Parse(System.String,Newtonsoft.Json.Schema.JsonSchemaResolver)">
+            <summary>
+            Parses the specified json.
+            </summary>
+            <param name="json">The json.</param>
+            <param name="resolver">The resolver.</param>
+            <returns>A <see cref="T:Newtonsoft.Json.Schema.JsonSchema"/> populated from the string that contains JSON.</returns>
+        </member>
+        <member name="M:Newtonsoft.Json.Schema.JsonSchema.WriteTo(Newtonsoft.Json.JsonWriter)">
+            <summary>
+            Writes this schema to a <see cref="T:Newtonsoft.Json.JsonWriter"/>.
+            </summary>
+            <param name="writer">A <see cref="T:Newtonsoft.Json.JsonWriter"/> into which this method will write.</param>
+        </member>
+        <member name="M:Newtonsoft.Json.Schema.JsonSchema.WriteTo(Newtonsoft.Json.JsonWriter,Newtonsoft.Json.Schema.JsonSchemaResolver)">
+            <summary>
+            Writes this schema to a <see cref="T:Newtonsoft.Json.JsonWriter"/> using the specified <see cref="T:Newtonsoft.Json.Schema.JsonSchemaResolver"/>.
+            </summary>
+            <param name="writer">A <see cref="T:Newtonsoft.Json.JsonWriter"/> into which this method will write.</param>
+            <param name="resolver">The resolver used.</param>
+        </member>
+        <member name="M:Newtonsoft.Json.Schema.JsonSchema.ToString">
+            <summary>
+            Returns a <see cref="T:System.String"/> that represents the current <see cref="T:System.Object"/>.
+            </summary>
+            <returns>
+            A <see cref="T:System.String"/> that represents the current <see cref="T:System.Object"/>.
+            </returns>
+        </member>
+        <member name="P:Newtonsoft.Json.Schema.JsonSchema.Id">
+            <summary>
+            Gets or sets the id.
+            </summary>
+        </member>
+        <member name="P:Newtonsoft.Json.Schema.JsonSchema.Title">
+            <summary>
+            Gets or sets the title.
+            </summary>
+        </member>
+        <member name="P:Newtonsoft.Json.Schema.JsonSchema.Required">
+            <summary>
+            Gets or sets whether the object is required.
+            </summary>
+        </member>
+        <member name="P:Newtonsoft.Json.Schema.JsonSchema.ReadOnly">
+            <summary>
+            Gets or sets whether the object is read only.
+            </summary>
+        </member>
+        <member name="P:Newtonsoft.Json.Schema.JsonSchema.Hidden">
+            <summary>
+            Gets or sets whether the object is visible to users.
+            </summary>
+        </member>
+        <member name="P:Newtonsoft.Json.Schema.JsonSchema.Transient">
+            <summary>
+            Gets or sets whether the object is transient.
+            </summary>
+        </member>
+        <member name="P:Newtonsoft.Json.Schema.JsonSchema.Description">
+            <summary>
+            Gets or sets the description of the object.
+            </summary>
+        </member>
+        <member name="P:Newtonsoft.Json.Schema.JsonSchema.Type">
+            <summary>
+            Gets or sets the types of values allowed by the object.
+            </summary>
+            <value>The type.</value>
+        </member>
+        <member name="P:Newtonsoft.Json.Schema.JsonSchema.Pattern">
+            <summary>
+            Gets or sets the pattern.
+            </summary>
+            <value>The pattern.</value>
+        </member>
+        <member name="P:Newtonsoft.Json.Schema.JsonSchema.MinimumLength">
+            <summary>
+            Gets or sets the minimum length.
+            </summary>
+            <value>The minimum length.</value>
+        </member>
+        <member name="P:Newtonsoft.Json.Schema.JsonSchema.MaximumLength">
+            <summary>
+            Gets or sets the maximum length.
+            </summary>
+            <value>The maximum length.</value>
+        </member>
+        <member name="P:Newtonsoft.Json.Schema.JsonSchema.DivisibleBy">
+            <summary>
+            Gets or sets a number that the value should be divisble by.
+            </summary>
+            <value>A number that the value should be divisble by.</value>
+        </member>
+        <member name="P:Newtonsoft.Json.Schema.JsonSchema.Minimum">
+            <summary>
+            Gets or sets the minimum.
+            </summary>
+            <value>The minimum.</value>
+        </member>
+        <member name="P:Newtonsoft.Json.Schema.JsonSchema.Maximum">
+            <summary>
+            Gets or sets the maximum.
+            </summary>
+            <value>The maximum.</value>
+        </member>
+        <member name="P:Newtonsoft.Json.Schema.JsonSchema.ExclusiveMinimum">
+            <summary>
+            Gets or sets a flag indicating whether the value can not equal the number defined by the "minimum" attribute.
+            </summary>
+            <value>A flag indicating whether the value can not equal the number defined by the "minimum" attribute.</value>
+        </member>
+        <member name="P:Newtonsoft.Json.Schema.JsonSchema.ExclusiveMaximum">
+            <summary>
+            Gets or sets a flag indicating whether the value can not equal the number defined by the "maximum" attribute.
+            </summary>
+            <value>A flag indicating whether the value can not equal the number defined by the "maximum" attribute.</value>
+        </member>
+        <member name="P:Newtonsoft.Json.Schema.JsonSchema.MinimumItems">
+            <summary>
+            Gets or sets the minimum number of items.
+            </summary>
+            <value>The minimum number of items.</value>
+        </member>
+        <member name="P:Newtonsoft.Json.Schema.JsonSchema.MaximumItems">
+            <summary>
+            Gets or sets the maximum number of items.
+            </summary>
+            <value>The maximum number of items.</value>
+        </member>
+        <member name="P:Newtonsoft.Json.Schema.JsonSchema.Items">
+            <summary>
+            Gets or sets the <see cref="T:Newtonsoft.Json.Schema.JsonSchema"/> of items.
+            </summary>
+            <value>The <see cref="T:Newtonsoft.Json.Schema.JsonSchema"/> of items.</value>
+        </member>
+        <member name="P:Newtonsoft.Json.Schema.JsonSchema.PositionalItemsValidation">
+            <summary>
+            Gets or sets a value indicating whether items in an array are validated using the <see cref="T:Newtonsoft.Json.Schema.JsonSchema"/> instance at their array position from <see cref="P:Newtonsoft.Json.Schema.JsonSchema.Items"/>.
+            </summary>
+            <value>
+            	<c>true</c> if items are validated using their array position; otherwise, <c>false</c>.
+            </value>
+        </member>
+        <member name="P:Newtonsoft.Json.Schema.JsonSchema.AdditionalItems">
+            <summary>
+            Gets or sets the <see cref="T:Newtonsoft.Json.Schema.JsonSchema"/> of additional items.
+            </summary>
+            <value>The <see cref="T:Newtonsoft.Json.Schema.JsonSchema"/> of additional items.</value>
+        </member>
+        <member name="P:Newtonsoft.Json.Schema.JsonSchema.AllowAdditionalItems">
+            <summary>
+            Gets or sets a value indicating whether additional items are allowed.
+            </summary>
+            <value>
+            	<c>true</c> if additional items are allowed; otherwise, <c>false</c>.
+            </value>
+        </member>
+        <member name="P:Newtonsoft.Json.Schema.JsonSchema.UniqueItems">
+            <summary>
+            Gets or sets whether the array items must be unique.
+            </summary>
+        </member>
+        <member name="P:Newtonsoft.Json.Schema.JsonSchema.Properties">
+            <summary>
+            Gets or sets the <see cref="T:Newtonsoft.Json.Schema.JsonSchema"/> of properties.
+            </summary>
+            <value>The <see cref="T:Newtonsoft.Json.Schema.JsonSchema"/> of properties.</value>
+        </member>
+        <member name="P:Newtonsoft.Json.Schema.JsonSchema.AdditionalProperties">
+            <summary>
+            Gets or sets the <see cref="T:Newtonsoft.Json.Schema.JsonSchema"/> of additional properties.
+            </summary>
+            <value>The <see cref="T:Newtonsoft.Json.Schema.JsonSchema"/> of additional properties.</value>
+        </member>
+        <member name="P:Newtonsoft.Json.Schema.JsonSchema.PatternProperties">
+            <summary>
+            Gets or sets the pattern properties.
+            </summary>
+            <value>The pattern properties.</value>
+        </member>
+        <member name="P:Newtonsoft.Json.Schema.JsonSchema.AllowAdditionalProperties">
+            <summary>
+            Gets or sets a value indicating whether additional properties are allowed.
+            </summary>
+            <value>
+            	<c>true</c> if additional properties are allowed; otherwise, <c>false</c>.
+            </value>
+        </member>
+        <member name="P:Newtonsoft.Json.Schema.JsonSchema.Requires">
+            <summary>
+            Gets or sets the required property if this property is present.
+            </summary>
+            <value>The required property if this property is present.</value>
+        </member>
+        <member name="P:Newtonsoft.Json.Schema.JsonSchema.Enum">
+            <summary>
+            Gets or sets the a collection of valid enum values allowed.
+            </summary>
+            <value>A collection of valid enum values allowed.</value>
+        </member>
+        <member name="P:Newtonsoft.Json.Schema.JsonSchema.Disallow">
+            <summary>
+            Gets or sets disallowed types.
+            </summary>
+            <value>The disallow types.</value>
+        </member>
+        <member name="P:Newtonsoft.Json.Schema.JsonSchema.Default">
+            <summary>
+            Gets or sets the default value.
+            </summary>
+            <value>The default value.</value>
+        </member>
+        <member name="P:Newtonsoft.Json.Schema.JsonSchema.Extends">
+            <summary>
+            Gets or sets the collection of <see cref="T:Newtonsoft.Json.Schema.JsonSchema"/> that this schema extends.
+            </summary>
+            <value>The collection of <see cref="T:Newtonsoft.Json.Schema.JsonSchema"/> that this schema extends.</value>
+        </member>
+        <member name="P:Newtonsoft.Json.Schema.JsonSchema.Format">
+            <summary>
+            Gets or sets the format.
+            </summary>
+            <value>The format.</value>
+        </member>
+        <member name="T:Newtonsoft.Json.Schema.JsonSchemaGenerator">
+            <summary>
+            Generates a <see cref="T:Newtonsoft.Json.Schema.JsonSchema"/> from a specified <see cref="T:System.Type"/>.
+            </summary>
+        </member>
+        <member name="M:Newtonsoft.Json.Schema.JsonSchemaGenerator.Generate(System.Type)">
+            <summary>
+            Generate a <see cref="T:Newtonsoft.Json.Schema.JsonSchema"/> from the specified type.
+            </summary>
+            <param name="type">The type to generate a <see cref="T:Newtonsoft.Json.Schema.JsonSchema"/> from.</param>
+            <returns>A <see cref="T:Newtonsoft.Json.Schema.JsonSchema"/> generated from the specified type.</returns>
+        </member>
+        <member name="M:Newtonsoft.Json.Schema.JsonSchemaGenerator.Generate(System.Type,Newtonsoft.Json.Schema.JsonSchemaResolver)">
+            <summary>
+            Generate a <see cref="T:Newtonsoft.Json.Schema.JsonSchema"/> from the specified type.
+            </summary>
+            <param name="type">The type to generate a <see cref="T:Newtonsoft.Json.Schema.JsonSchema"/> from.</param>
+            <param name="resolver">The <see cref="T:Newtonsoft.Json.Schema.JsonSchemaResolver"/> used to resolve schema references.</param>
+            <returns>A <see cref="T:Newtonsoft.Json.Schema.JsonSchema"/> generated from the specified type.</returns>
+        </member>
+        <member name="M:Newtonsoft.Json.Schema.JsonSchemaGenerator.Generate(System.Type,System.Boolean)">
+            <summary>
+            Generate a <see cref="T:Newtonsoft.Json.Schema.JsonSchema"/> from the specified type.
+            </summary>
+            <param name="type">The type to generate a <see cref="T:Newtonsoft.Json.Schema.JsonSchema"/> from.</param>
+            <param name="rootSchemaNullable">Specify whether the generated root <see cref="T:Newtonsoft.Json.Schema.JsonSchema"/> will be nullable.</param>
+            <returns>A <see cref="T:Newtonsoft.Json.Schema.JsonSchema"/> generated from the specified type.</returns>
+        </member>
+        <member name="M:Newtonsoft.Json.Schema.JsonSchemaGenerator.Generate(System.Type,Newtonsoft.Json.Schema.JsonSchemaResolver,System.Boolean)">
+            <summary>
+            Generate a <see cref="T:Newtonsoft.Json.Schema.JsonSchema"/> from the specified type.
+            </summary>
+            <param name="type">The type to generate a <see cref="T:Newtonsoft.Json.Schema.JsonSchema"/> from.</param>
+            <param name="resolver">The <see cref="T:Newtonsoft.Json.Schema.JsonSchemaResolver"/> used to resolve schema references.</param>
+            <param name="rootSchemaNullable">Specify whether the generated root <see cref="T:Newtonsoft.Json.Schema.JsonSchema"/> will be nullable.</param>
+            <returns>A <see cref="T:Newtonsoft.Json.Schema.JsonSchema"/> generated from the specified type.</returns>
+        </member>
+        <member name="P:Newtonsoft.Json.Schema.JsonSchemaGenerator.UndefinedSchemaIdHandling">
+            <summary>
+            Gets or sets how undefined schemas are handled by the serializer.
+            </summary>
+        </member>
+        <member name="P:Newtonsoft.Json.Schema.JsonSchemaGenerator.ContractResolver">
+            <summary>
+            Gets or sets the contract resolver.
+            </summary>
+            <value>The contract resolver.</value>
+        </member>
+        <member name="T:Newtonsoft.Json.Schema.JsonSchemaType">
+            <summary>
+            The value types allowed by the <see cref="T:Newtonsoft.Json.Schema.JsonSchema"/>.
+            </summary>
+        </member>
+        <member name="F:Newtonsoft.Json.Schema.JsonSchemaType.None">
+            <summary>
+            No type specified.
+            </summary>
+        </member>
+        <member name="F:Newtonsoft.Json.Schema.JsonSchemaType.String">
+            <summary>
+            String type.
+            </summary>
+        </member>
+        <member name="F:Newtonsoft.Json.Schema.JsonSchemaType.Float">
+            <summary>
+            Float type.
+            </summary>
+        </member>
+        <member name="F:Newtonsoft.Json.Schema.JsonSchemaType.Integer">
+            <summary>
+            Integer type.
+            </summary>
+        </member>
+        <member name="F:Newtonsoft.Json.Schema.JsonSchemaType.Boolean">
+            <summary>
+            Boolean type.
+            </summary>
+        </member>
+        <member name="F:Newtonsoft.Json.Schema.JsonSchemaType.Object">
+            <summary>
+            Object type.
+            </summary>
+        </member>
+        <member name="F:Newtonsoft.Json.Schema.JsonSchemaType.Array">
+            <summary>
+            Array type.
+            </summary>
+        </member>
+        <member name="F:Newtonsoft.Json.Schema.JsonSchemaType.Null">
+            <summary>
+            Null type.
+            </summary>
+        </member>
+        <member name="F:Newtonsoft.Json.Schema.JsonSchemaType.Any">
+            <summary>
+            Any type.
+            </summary>
+        </member>
+        <member name="T:Newtonsoft.Json.Serialization.JsonObjectContract">
+            <summary>
+            Contract details for a <see cref="T:System.Type"/> used by the <see cref="T:Newtonsoft.Json.JsonSerializer"/>.
+            </summary>
+        </member>
+        <member name="M:Newtonsoft.Json.Serialization.JsonObjectContract.#ctor(System.Type)">
+            <summary>
+            Initializes a new instance of the <see cref="T:Newtonsoft.Json.Serialization.JsonObjectContract"/> class.
+            </summary>
+            <param name="underlyingType">The underlying type for the contract.</param>
+        </member>
+        <member name="P:Newtonsoft.Json.Serialization.JsonObjectContract.MemberSerialization">
+            <summary>
+            Gets or sets the object member serialization.
+            </summary>
+            <value>The member object serialization.</value>
+        </member>
+        <member name="P:Newtonsoft.Json.Serialization.JsonObjectContract.ItemRequired">
+            <summary>
+            Gets or sets a value that indicates whether the object's properties are required.
+            </summary>
+            <value>
+            	A value indicating whether the object's properties are required.
+            </value>
+        </member>
+        <member name="P:Newtonsoft.Json.Serialization.JsonObjectContract.Properties">
+            <summary>
+            Gets the object's properties.
+            </summary>
+            <value>The object's properties.</value>
+        </member>
+        <member name="P:Newtonsoft.Json.Serialization.JsonObjectContract.ConstructorParameters">
+            <summary>
+            Gets the constructor parameters required for any non-default constructor
+            </summary>
+        </member>
+        <member name="P:Newtonsoft.Json.Serialization.JsonObjectContract.OverrideConstructor">
+            <summary>
+            Gets or sets the override constructor used to create the object.
+            This is set when a constructor is marked up using the
+            JsonConstructor attribute.
+            </summary>
+            <value>The override constructor.</value>
+        </member>
+        <member name="P:Newtonsoft.Json.Serialization.JsonObjectContract.ParametrizedConstructor">
+            <summary>
+            Gets or sets the parametrized constructor used to create the object.
+            </summary>
+            <value>The parametrized constructor.</value>
+        </member>
+        <member name="T:Newtonsoft.Json.Serialization.JsonStringContract">
+            <summary>
+            Contract details for a <see cref="T:System.Type"/> used by the <see cref="T:Newtonsoft.Json.JsonSerializer"/>.
+            </summary>
+        </member>
+        <member name="M:Newtonsoft.Json.Serialization.JsonStringContract.#ctor(System.Type)">
+            <summary>
+            Initializes a new instance of the <see cref="T:Newtonsoft.Json.Serialization.JsonStringContract"/> class.
+            </summary>
+            <param name="underlyingType">The underlying type for the contract.</param>
+        </member>
+        <member name="T:Newtonsoft.Json.Serialization.ObjectConstructor`1">
+            <summary>
+            Represents a method that constructs an object.
+            </summary>
+            <typeparam name="T">The object type to create.</typeparam>
+        </member>
+        <member name="T:Newtonsoft.Json.Serialization.OnErrorAttribute">
+            <summary>
+            When applied to a method, specifies that the method is called when an error occurs serializing an object.
+            </summary>
+        </member>
+        <member name="T:Newtonsoft.Json.Serialization.ReflectionValueProvider">
+            <summary>
+            Get and set values for a <see cref="T:System.Reflection.MemberInfo"/> using reflection.
+            </summary>
+        </member>
+        <member name="M:Newtonsoft.Json.Serialization.ReflectionValueProvider.#ctor(System.Reflection.MemberInfo)">
+            <summary>
+            Initializes a new instance of the <see cref="T:Newtonsoft.Json.Serialization.ReflectionValueProvider"/> class.
+            </summary>
+            <param name="memberInfo">The member info.</param>
+        </member>
+        <member name="M:Newtonsoft.Json.Serialization.ReflectionValueProvider.SetValue(System.Object,System.Object)">
+            <summary>
+            Sets the value.
+            </summary>
+            <param name="target">The target to set the value on.</param>
+            <param name="value">The value to set on the target.</param>
+        </member>
+        <member name="M:Newtonsoft.Json.Serialization.ReflectionValueProvider.GetValue(System.Object)">
+            <summary>
+            Gets the value.
+            </summary>
+            <param name="target">The target to get the value from.</param>
+            <returns>The value.</returns>
+        </member>
+        <member name="T:Newtonsoft.Json.TypeNameHandling">
+            <summary>
+            Specifies type name handling options for the <see cref="T:Newtonsoft.Json.JsonSerializer"/>.
+            </summary>
+        </member>
+        <member name="F:Newtonsoft.Json.TypeNameHandling.None">
+            <summary>
+            Do not include the .NET type name when serializing types.
+            </summary>
+        </member>
+        <member name="F:Newtonsoft.Json.TypeNameHandling.Objects">
+            <summary>
+            Include the .NET type name when serializing into a JSON object structure.
+            </summary>
+        </member>
+        <member name="F:Newtonsoft.Json.TypeNameHandling.Arrays">
+            <summary>
+            Include the .NET type name when serializing into a JSON array structure.
+            </summary>
+        </member>
+        <member name="F:Newtonsoft.Json.TypeNameHandling.All">
+            <summary>
+            Always include the .NET type name when serializing.
+            </summary>
+        </member>
+        <member name="F:Newtonsoft.Json.TypeNameHandling.Auto">
+            <summary>
+            Include the .NET type name when the type of the object being serialized is not the same as its declared type.
+            </summary>
+        </member>
+        <member name="M:Newtonsoft.Json.Utilities.ConvertUtils.Convert(System.Object,System.Globalization.CultureInfo,System.Type)">
+            <summary>
+            Converts the value to the specified type.
+            </summary>
+            <param name="initialValue">The value to convert.</param>
+            <param name="culture">The culture to use when converting.</param>
+            <param name="targetType">The type to convert the value to.</param>
+            <returns>The converted type.</returns>
+        </member>
+        <member name="M:Newtonsoft.Json.Utilities.ConvertUtils.TryConvert(System.Object,System.Globalization.CultureInfo,System.Type,System.Object@)">
+            <summary>
+            Converts the value to the specified type.
+            </summary>
+            <param name="initialValue">The value to convert.</param>
+            <param name="culture">The culture to use when converting.</param>
+            <param name="targetType">The type to convert the value to.</param>
+            <param name="convertedValue">The converted value if the conversion was successful or the default value of <c>T</c> if it failed.</param>
+            <returns>
+            	<c>true</c> if <c>initialValue</c> was converted successfully; otherwise, <c>false</c>.
+            </returns>
+        </member>
+        <member name="M:Newtonsoft.Json.Utilities.ConvertUtils.ConvertOrCast(System.Object,System.Globalization.CultureInfo,System.Type)">
+            <summary>
+            Converts the value to the specified type. If the value is unable to be converted, the
+            value is checked whether it assignable to the specified type.
+            </summary>
+            <param name="initialValue">The value to convert.</param>
+            <param name="culture">The culture to use when converting.</param>
+            <param name="targetType">The type to convert or cast the value to.</param>
+            <returns>
+            The converted type. If conversion was unsuccessful, the initial value
+            is returned if assignable to the target type.
+            </returns>
+        </member>
+        <member name="M:Newtonsoft.Json.Utilities.EnumUtils.GetNamesAndValues``1">
+            <summary>
+            Gets a dictionary of the names and values of an Enum type.
+            </summary>
+            <returns></returns>
+        </member>
+        <member name="M:Newtonsoft.Json.Utilities.EnumUtils.GetNamesAndValues``1(System.Type)">
+            <summary>
+            Gets a dictionary of the names and values of an Enum type.
+            </summary>
+            <param name="enumType">The enum type to get names and values for.</param>
+            <returns></returns>
+        </member>
+        <member name="T:Newtonsoft.Json.JsonToken">
+            <summary>
+            Specifies the type of Json token.
+            </summary>
+        </member>
+        <member name="F:Newtonsoft.Json.JsonToken.None">
+            <summary>
+            This is returned by the <see cref="T:Newtonsoft.Json.JsonReader"/> if a <see cref="M:Newtonsoft.Json.JsonReader.Read"/> method has not been called. 
+            </summary>
+        </member>
+        <member name="F:Newtonsoft.Json.JsonToken.StartObject">
+            <summary>
+            An object start token.
+            </summary>
+        </member>
+        <member name="F:Newtonsoft.Json.JsonToken.StartArray">
+            <summary>
+            An array start token.
+            </summary>
+        </member>
+        <member name="F:Newtonsoft.Json.JsonToken.StartConstructor">
+            <summary>
+            A constructor start token.
+            </summary>
+        </member>
+        <member name="F:Newtonsoft.Json.JsonToken.PropertyName">
+            <summary>
+            An object property name.
+            </summary>
+        </member>
+        <member name="F:Newtonsoft.Json.JsonToken.Comment">
+            <summary>
+            A comment.
+            </summary>
+        </member>
+        <member name="F:Newtonsoft.Json.JsonToken.Raw">
+            <summary>
+            Raw JSON.
+            </summary>
+        </member>
+        <member name="F:Newtonsoft.Json.JsonToken.Integer">
+            <summary>
+            An integer.
+            </summary>
+        </member>
+        <member name="F:Newtonsoft.Json.JsonToken.Float">
+            <summary>
+            A float.
+            </summary>
+        </member>
+        <member name="F:Newtonsoft.Json.JsonToken.String">
+            <summary>
+            A string.
+            </summary>
+        </member>
+        <member name="F:Newtonsoft.Json.JsonToken.Boolean">
+            <summary>
+            A boolean.
+            </summary>
+        </member>
+        <member name="F:Newtonsoft.Json.JsonToken.Null">
+            <summary>
+            A null token.
+            </summary>
+        </member>
+        <member name="F:Newtonsoft.Json.JsonToken.Undefined">
+            <summary>
+            An undefined token.
+            </summary>
+        </member>
+        <member name="F:Newtonsoft.Json.JsonToken.EndObject">
+            <summary>
+            An object end token.
+            </summary>
+        </member>
+        <member name="F:Newtonsoft.Json.JsonToken.EndArray">
+            <summary>
+            An array end token.
+            </summary>
+        </member>
+        <member name="F:Newtonsoft.Json.JsonToken.EndConstructor">
+            <summary>
+            A constructor end token.
+            </summary>
+        </member>
+        <member name="F:Newtonsoft.Json.JsonToken.Date">
+            <summary>
+            A Date.
+            </summary>
+        </member>
+        <member name="F:Newtonsoft.Json.JsonToken.Bytes">
+            <summary>
+            Byte data.
+            </summary>
+        </member>
+        <member name="T:Newtonsoft.Json.Utilities.StringBuffer">
+            <summary>
+            Builds a string. Unlike StringBuilder this class lets you reuse it's internal buffer.
+            </summary>
+        </member>
+        <member name="M:Newtonsoft.Json.Utilities.CollectionUtils.IsNullOrEmpty``1(System.Collections.Generic.ICollection{``0})">
+            <summary>
+            Determines whether the collection is null or empty.
+            </summary>
+            <param name="collection">The collection.</param>
+            <returns>
+            	<c>true</c> if the collection is null or empty; otherwise, <c>false</c>.
+            </returns>
+        </member>
+        <member name="M:Newtonsoft.Json.Utilities.CollectionUtils.AddRange``1(System.Collections.Generic.IList{``0},System.Collections.Generic.IEnumerable{``0})">
+            <summary>
+            Adds the elements of the specified collection to the specified generic IList.
+            </summary>
+            <param name="initial">The list to add to.</param>
+            <param name="collection">The collection of elements to add.</param>
+        </member>
+        <member name="M:Newtonsoft.Json.Utilities.CollectionUtils.IndexOf``1(System.Collections.Generic.IEnumerable{``0},``0,System.Collections.Generic.IEqualityComparer{``0})">
+            <summary>
+            Returns the index of the first occurrence in a sequence by using a specified IEqualityComparer.
+            </summary>
+            <typeparam name="TSource">The type of the elements of source.</typeparam>
+            <param name="list">A sequence in which to locate a value.</param>
+            <param name="value">The object to locate in the sequence</param>
+            <param name="comparer">An equality comparer to compare values.</param>
+            <returns>The zero-based index of the first occurrence of value within the entire sequence, if found; otherwise, –1.</returns>
+        </member>
+        <member name="M:Newtonsoft.Json.Utilities.ReflectionUtils.GetCollectionItemType(System.Type)">
+            <summary>
+            Gets the type of the typed collection's items.
+            </summary>
+            <param name="type">The type.</param>
+            <returns>The type of the typed collection's items.</returns>
+        </member>
+        <member name="M:Newtonsoft.Json.Utilities.ReflectionUtils.GetMemberUnderlyingType(System.Reflection.MemberInfo)">
+            <summary>
+            Gets the member's underlying type.
+            </summary>
+            <param name="member">The member.</param>
+            <returns>The underlying type of the member.</returns>
+        </member>
+        <member name="M:Newtonsoft.Json.Utilities.ReflectionUtils.IsIndexedProperty(System.Reflection.MemberInfo)">
+            <summary>
+            Determines whether the member is an indexed property.
+            </summary>
+            <param name="member">The member.</param>
+            <returns>
+            	<c>true</c> if the member is an indexed property; otherwise, <c>false</c>.
+            </returns>
+        </member>
+        <member name="M:Newtonsoft.Json.Utilities.ReflectionUtils.IsIndexedProperty(System.Reflection.PropertyInfo)">
+            <summary>
+            Determines whether the property is an indexed property.
+            </summary>
+            <param name="property">The property.</param>
+            <returns>
+            	<c>true</c> if the property is an indexed property; otherwise, <c>false</c>.
+            </returns>
+        </member>
+        <member name="M:Newtonsoft.Json.Utilities.ReflectionUtils.GetMemberValue(System.Reflection.MemberInfo,System.Object)">
+            <summary>
+            Gets the member's value on the object.
+            </summary>
+            <param name="member">The member.</param>
+            <param name="target">The target object.</param>
+            <returns>The member's value on the object.</returns>
+        </member>
+        <member name="M:Newtonsoft.Json.Utilities.ReflectionUtils.SetMemberValue(System.Reflection.MemberInfo,System.Object,System.Object)">
+            <summary>
+            Sets the member's value on the target object.
+            </summary>
+            <param name="member">The member.</param>
+            <param name="target">The target.</param>
+            <param name="value">The value.</param>
+        </member>
+        <member name="M:Newtonsoft.Json.Utilities.ReflectionUtils.CanReadMemberValue(System.Reflection.MemberInfo,System.Boolean)">
+            <summary>
+            Determines whether the specified MemberInfo can be read.
+            </summary>
+            <param name="member">The MemberInfo to determine whether can be read.</param>
+            /// <param name="nonPublic">if set to <c>true</c> then allow the member to be gotten non-publicly.</param>
+            <returns>
+            	<c>true</c> if the specified MemberInfo can be read; otherwise, <c>false</c>.
+            </returns>
+        </member>
+        <member name="M:Newtonsoft.Json.Utilities.ReflectionUtils.CanSetMemberValue(System.Reflection.MemberInfo,System.Boolean,System.Boolean)">
+            <summary>
+            Determines whether the specified MemberInfo can be set.
+            </summary>
+            <param name="member">The MemberInfo to determine whether can be set.</param>
+            <param name="nonPublic">if set to <c>true</c> then allow the member to be set non-publicly.</param>
+            <param name="canSetReadOnly">if set to <c>true</c> then allow the member to be set if read-only.</param>
+            <returns>
+            	<c>true</c> if the specified MemberInfo can be set; otherwise, <c>false</c>.
+            </returns>
+        </member>
+        <member name="M:Newtonsoft.Json.Utilities.StringUtils.IsWhiteSpace(System.String)">
+            <summary>
+            Determines whether the string is all white space. Empty string will return false.
+            </summary>
+            <param name="s">The string to test whether it is all white space.</param>
+            <returns>
+            	<c>true</c> if the string is all white space; otherwise, <c>false</c>.
+            </returns>
+        </member>
+        <member name="M:Newtonsoft.Json.Utilities.StringUtils.NullEmptyString(System.String)">
+            <summary>
+            Nulls an empty string.
+            </summary>
+            <param name="s">The string.</param>
+            <returns>Null if the string was null, otherwise the string unchanged.</returns>
+        </member>
+        <member name="T:Newtonsoft.Json.WriteState">
+            <summary>
+            Specifies the state of the <see cref="T:Newtonsoft.Json.JsonWriter"/>.
+            </summary>
+        </member>
+        <member name="F:Newtonsoft.Json.WriteState.Error">
+            <summary>
+            An exception has been thrown, which has left the <see cref="T:Newtonsoft.Json.JsonWriter"/> in an invalid state.
+            You may call the <see cref="M:Newtonsoft.Json.JsonWriter.Close"/> method to put the <see cref="T:Newtonsoft.Json.JsonWriter"/> in the <c>Closed</c> state.
+            Any other <see cref="T:Newtonsoft.Json.JsonWriter"/> method calls results in an <see cref="T:System.InvalidOperationException"/> being thrown. 
+            </summary>
+        </member>
+        <member name="F:Newtonsoft.Json.WriteState.Closed">
+            <summary>
+            The <see cref="M:Newtonsoft.Json.JsonWriter.Close"/> method has been called. 
+            </summary>
+        </member>
+        <member name="F:Newtonsoft.Json.WriteState.Object">
+            <summary>
+            An object is being written. 
+            </summary>
+        </member>
+        <member name="F:Newtonsoft.Json.WriteState.Array">
+            <summary>
+            A array is being written.
+            </summary>
+        </member>
+        <member name="F:Newtonsoft.Json.WriteState.Constructor">
+            <summary>
+            A constructor is being written.
+            </summary>
+        </member>
+        <member name="F:Newtonsoft.Json.WriteState.Property">
+            <summary>
+            A property is being written.
+            </summary>
+        </member>
+        <member name="F:Newtonsoft.Json.WriteState.Start">
+            <summary>
+            A write method has not been called.
+            </summary>
+        </member>
+    </members>
+</doc>
