Copying //UE4/Dev-Core to //UE4/Main

diff --git a/Engine/Source/Programs/CrashReporter/CrashReportWebSite/packages/Unity.4.0.1/lib/net45/Microsoft.Practices.Unity.Configuration.XML b/Engine/Source/Programs/CrashReporter/CrashReportWebSite/packages/Unity.4.0.1/lib/net45/Microsoft.Practices.Unity.Configuration.XML
new file mode 100644
--- /dev/null
+++ b/Engine/Source/Programs/CrashReporter/CrashReportWebSite/packages/Unity.4.0.1/lib/net45/Microsoft.Practices.Unity.Configuration.XML
@@ -0,0 +1,2424 @@
+<?xml version="1.0"?>
+<doc>
+    <assembly>
+        <name>Microsoft.Practices.Unity.Configuration</name>
+    </assembly>
+    <members>
+        <member name="T:Microsoft.Practices.Unity.Configuration.AliasElement">
+            <summary>
+            A configuration element storing information about a single type alias.
+            </summary>
+        </member>
+        <member name="T:Microsoft.Practices.Unity.Configuration.ConfigurationHelpers.DeserializableConfigurationElement">
+            <summary>
+            Base class for configuration elements with a default implementation of
+            public deserialization.
+            </summary>
+        </member>
+        <member name="M:Microsoft.Practices.Unity.Configuration.ConfigurationHelpers.DeserializableConfigurationElement.Deserialize(System.Xml.XmlReader)">
+            <summary>
+            Load this element from the given <see cref="T:System.Xml.XmlReader"/>.
+            </summary>
+            <param name="reader">Contains the XML to initialize from.</param>
+        </member>
+        <member name="M:Microsoft.Practices.Unity.Configuration.ConfigurationHelpers.DeserializableConfigurationElement.SerializeContent(System.Xml.XmlWriter)">
+            <summary>
+            Write the contents of this element to the given <see cref="T:System.Xml.XmlWriter"/>.
+            </summary>
+            <remarks>The caller of this method has already written the start element tag before
+            calling this method, so deriving classes only need to write the element content, not
+            the start or end tags.</remarks>
+            <param name="writer">Writer to send XML content to.</param>
+        </member>
+        <member name="M:Microsoft.Practices.Unity.Configuration.AliasElement.#ctor">
+            <summary>
+            Construct a new, uninitialized <see cref="T:Microsoft.Practices.Unity.Configuration.AliasElement"/>.
+            </summary>
+        </member>
+        <member name="M:Microsoft.Practices.Unity.Configuration.AliasElement.#ctor(System.String,System.Type)">
+            <summary>
+            Construct a new <see cref="T:Microsoft.Practices.Unity.Configuration.AliasElement"/> that is initialized
+            to alias <paramref name="alias"/> to the target <paramref name="targetType"/>.
+            </summary>
+            <param name="alias">Alias to use.</param>
+            <param name="targetType">Type that is aliased.</param>
+        </member>
+        <member name="M:Microsoft.Practices.Unity.Configuration.AliasElement.SerializeContent(System.Xml.XmlWriter)">
+            <summary>
+            Write the contents of this element to the given <see cref="T:System.Xml.XmlWriter"/>.
+            </summary>
+            <remarks>The caller of this method has already written the start element tag before
+            calling this method, so deriving classes only need to write the element content, not
+            the start or end tags.</remarks>
+            <param name="writer">Writer to send XML content to.</param>
+        </member>
+        <member name="P:Microsoft.Practices.Unity.Configuration.AliasElement.Alias">
+            <summary>
+            The alias used for this type.
+            </summary>
+        </member>
+        <member name="P:Microsoft.Practices.Unity.Configuration.AliasElement.TypeName">
+            <summary>
+            The fully qualified name this alias refers to.
+            </summary>
+        </member>
+        <member name="T:Microsoft.Practices.Unity.Configuration.AliasElementCollection">
+            <summary>
+            A collection of <see cref="T:Microsoft.Practices.Unity.Configuration.AliasElement"/>s.
+            </summary>
+        </member>
+        <member name="T:Microsoft.Practices.Unity.Configuration.ConfigurationHelpers.DeserializableConfigurationElementCollection`1">
+            <summary>
+            Specialization of <see cref="T:Microsoft.Practices.Unity.Configuration.ConfigurationHelpers.DeserializableConfigurationElementCollectionBase`1"/>
+            that provides a canned implementation of <see cref="M:System.Configuration.ConfigurationElementCollection.CreateNewElement"/>.
+            </summary>
+            <typeparam name="TElement">Type of configuration element in the collection.</typeparam>
+        </member>
+        <member name="T:Microsoft.Practices.Unity.Configuration.ConfigurationHelpers.DeserializableConfigurationElementCollectionBase`1">
+            <summary>
+            A base helper class for implementing configuration collections.
+            </summary>
+            <typeparam name="TElement">Type of configuration element contained in 
+            the collection.</typeparam>
+        </member>
+        <member name="M:Microsoft.Practices.Unity.Configuration.ConfigurationHelpers.DeserializableConfigurationElementCollectionBase`1.GetElement(System.Int32)">
+            <summary>
+            Plug point to get objects out of the collection.
+            </summary>
+            <param name="index">Index in the collection to retrieve the item from.</param>
+            <returns>Item at that index or null if not present.</returns>
+        </member>
+        <member name="M:Microsoft.Practices.Unity.Configuration.ConfigurationHelpers.DeserializableConfigurationElementCollectionBase`1.GetElement(System.Object)">
+            <summary>
+            Plug point to get objects out of the collection.
+            </summary>
+            <param name="key">Key to look up the object by.</param>
+            <returns>Item with that key or null if not present.</returns>
+        </member>
+        <member name="M:Microsoft.Practices.Unity.Configuration.ConfigurationHelpers.DeserializableConfigurationElementCollectionBase`1.Deserialize(System.Xml.XmlReader)">
+            <summary>
+            Load this element from the given <see cref="T:System.Xml.XmlReader"/>.
+            </summary>
+            <param name="reader">Contains the XML to initialize from.</param>
+        </member>
+        <member name="M:Microsoft.Practices.Unity.Configuration.ConfigurationHelpers.DeserializableConfigurationElementCollectionBase`1.GetEnumerator">
+            <summary>
+            Returns an enumerator that iterates through the collection.
+            </summary>
+            <returns>
+            A <see cref="T:System.Collections.Generic.IEnumerator`1"/> that can be used to iterate through the collection.
+            </returns>
+            <filterpriority>1</filterpriority>
+        </member>
+        <member name="M:Microsoft.Practices.Unity.Configuration.ConfigurationHelpers.DeserializableConfigurationElementCollectionBase`1.Add(`0)">
+            <summary>
+            Add a new element to the collection.
+            </summary>
+            <param name="element">Element to add.</param>
+        </member>
+        <member name="M:Microsoft.Practices.Unity.Configuration.ConfigurationHelpers.DeserializableConfigurationElementCollectionBase`1.RemoveAt(System.Int32)">
+            <summary>
+            Remove an element from the collection at the given index.
+            </summary>
+            <param name="index">The index of the item to remove.</param>
+        </member>
+        <member name="M:Microsoft.Practices.Unity.Configuration.ConfigurationHelpers.DeserializableConfigurationElementCollectionBase`1.Clear">
+            <summary>
+            Remove all the items in the collection.
+            </summary>
+        </member>
+        <member name="M:Microsoft.Practices.Unity.Configuration.ConfigurationHelpers.DeserializableConfigurationElementCollectionBase`1.SerializeElementContents(System.Xml.XmlWriter,System.String)">
+            <summary>
+            Write out the contents of this collection to the given
+            <paramref name="writer"/> without a containing element
+            corresponding directly to this container element. Each
+            child element will have a tag name given by
+            <paramref name="elementName"/>.
+            </summary>
+            <param name="writer"><see cref="T:System.Xml.XmlWriter"/> to output XML to.</param>
+            <param name="elementName">Name of tag to generate.</param>
+        </member>
+        <member name="P:Microsoft.Practices.Unity.Configuration.ConfigurationHelpers.DeserializableConfigurationElementCollectionBase`1.Item(System.Int32)">
+            <summary>
+            Indexer to retrieve items in the collection by index.
+            </summary>
+            <param name="index">Index of the item to get or set.</param>
+            <returns>The item at the given index.</returns>
+        </member>
+        <member name="M:Microsoft.Practices.Unity.Configuration.ConfigurationHelpers.DeserializableConfigurationElementCollection`1.CreateNewElement">
+            <summary>
+            When overridden in a derived class, creates a new <see cref="T:System.Configuration.ConfigurationElement"/>.
+            </summary>
+            <returns>
+            A new <see cref="T:System.Configuration.ConfigurationElement"/>.
+            </returns>
+        </member>
+        <member name="M:Microsoft.Practices.Unity.Configuration.AliasElementCollection.OnDeserializeUnrecognizedElement(System.String,System.Xml.XmlReader)">
+            <summary>
+            Causes the configuration system to throw an exception.
+            </summary>
+            <returns>
+            true if the unrecognized element was deserialized successfully; otherwise, false. The default is false.
+            </returns>
+            <param name="elementName">The name of the unrecognized element.</param>
+            <param name="reader">An input stream that reads XML from the configuration file.</param>
+            <exception cref="T:System.Configuration.ConfigurationErrorsException">The element specified in <paramref name="elementName"/> is the &lt;clear&gt; element.</exception>
+            <exception cref="T:System.ArgumentException"><paramref name="elementName"/> starts with the reserved prefix "config" or "lock".</exception>
+        </member>
+        <member name="M:Microsoft.Practices.Unity.Configuration.AliasElementCollection.GetElementKey(System.Configuration.ConfigurationElement)">
+            <summary>
+            Gets the element key for a specified configuration element when overridden in a derived class.
+            </summary>
+            <returns>
+            An <see cref="T:System.Object"/> that acts as the key for the specified <see cref="T:System.Configuration.ConfigurationElement"/>.
+            </returns>
+            <param name="element">The <see cref="T:System.Configuration.ConfigurationElement"/> to return the key for. </param>
+        </member>
+        <member name="P:Microsoft.Practices.Unity.Configuration.AliasElementCollection.Item(System.String)">
+            <summary>
+            Indexer that allows you to get or set an alias by the alias name.
+            </summary>
+            <param name="alias">Alias of element to get or set.</param>
+            <returns>The type name at that alias.</returns>
+        </member>
+        <member name="T:Microsoft.Practices.Unity.Configuration.ArrayElement">
+            <summary>
+            A configuration element used to configure injection of
+            a specific set of values into an array.
+            </summary>
+        </member>
+        <member name="T:Microsoft.Practices.Unity.Configuration.ParameterValueElement">
+            <summary>
+            Base class for configuration elements that describe a value that will
+            be injected.
+            </summary>
+        </member>
+        <member name="M:Microsoft.Practices.Unity.Configuration.ParameterValueElement.#ctor">
+            <summary>
+            Initialize a new instance of <see cref="T:Microsoft.Practices.Unity.Configuration.ParameterValueElement"/>.
+            </summary>
+        </member>
+        <member name="M:Microsoft.Practices.Unity.Configuration.ParameterValueElement.GetInjectionParameterValue(Microsoft.Practices.Unity.IUnityContainer,System.Type)">
+            <summary>
+            Generate an <see cref="T:Microsoft.Practices.Unity.InjectionParameterValue"/> object
+            that will be used to configure the container for a type registration.
+            </summary>
+            <param name="container">Container that is being configured. Supplied in order
+            to let custom implementations retrieve services; do not configure the container
+            directly in this method.</param>
+            <param name="parameterType">Type of the </param>
+            <returns></returns>
+        </member>
+        <member name="M:Microsoft.Practices.Unity.Configuration.ParameterValueElement.GuardPropertyValueIsPresent(System.Collections.Generic.IDictionary{System.String,System.String},System.String)">
+            <summary>
+            Validate that an expected attribute is present in the given
+            dictionary and that it has a non-empty value.
+            </summary>
+            <param name="propertyValues">Dictionary of name/value pairs to check.</param>
+            <param name="requiredProperty">attribute name to check.</param>
+        </member>
+        <member name="P:Microsoft.Practices.Unity.Configuration.ParameterValueElement.Key">
+            <summary>
+            Return a unique string that can be used to identify this object. Used
+            by the configuration collection support.
+            </summary>
+        </member>
+        <member name="M:Microsoft.Practices.Unity.Configuration.ArrayElement.SerializeContent(System.Xml.XmlWriter)">
+            <summary>
+            Write the contents of this element to the given <see cref="T:System.Xml.XmlWriter"/>.
+            </summary>
+            <remarks>The caller of this method has already written the start element tag before
+            calling this method, so deriving classes only need to write the element content, not
+            the start or end tags.</remarks>
+            <param name="writer">Writer to send XML content to.</param>
+        </member>
+        <member name="M:Microsoft.Practices.Unity.Configuration.ArrayElement.GetInjectionParameterValue(Microsoft.Practices.Unity.IUnityContainer,System.Type)">
+            <summary>
+            Generate an <see cref="T:Microsoft.Practices.Unity.InjectionParameterValue"/> object
+            that will be used to configure the container for a type registration.
+            </summary>
+            <param name="container">Container that is being configured. Supplied in order
+            to let custom implementations retrieve services; do not configure the container
+            directly in this method.</param>
+            <param name="parameterType">Type of the </param>
+            <returns></returns>
+        </member>
+        <member name="P:Microsoft.Practices.Unity.Configuration.ArrayElement.TypeName">
+            <summary>
+            Type of array to inject. This is actually the type of the array elements,
+            not the array type. Optional, if not specified we take the type from
+            our containing element.
+            </summary>
+        </member>
+        <member name="P:Microsoft.Practices.Unity.Configuration.ArrayElement.Values">
+            <summary>
+            Values used to calculate the contents of the array.
+            </summary>
+        </member>
+        <member name="T:Microsoft.Practices.Unity.Configuration.AssemblyElement">
+            <summary>
+            A configuration element representing the namespace
+            tags in the config file.
+            </summary>
+        </member>
+        <member name="T:Microsoft.Practices.Unity.Configuration.NamedElement">
+            <summary>
+            An element with a single "name" property, used for
+            the namespaces and assemblies.
+            </summary>
+        </member>
+        <member name="M:Microsoft.Practices.Unity.Configuration.NamedElement.SerializeContent(System.Xml.XmlWriter)">
+            <summary>
+            Write the contents of this element to the given <see cref="T:System.Xml.XmlWriter"/>.
+            </summary>
+            <remarks>The caller of this method has already written the start element tag before
+            calling this method, so deriving classes only need to write the element content, not
+            the start or end tags.</remarks>
+            <param name="writer">Writer to send XML content to.</param>
+        </member>
+        <member name="P:Microsoft.Practices.Unity.Configuration.NamedElement.Name">
+            <summary>
+            Name attribute for this element.
+            </summary>
+        </member>
+        <member name="T:Microsoft.Practices.Unity.Configuration.AssemblyElementCollection">
+            <summary>
+            A collection of <see cref="T:Microsoft.Practices.Unity.Configuration.AssemblyElement"/>s in configuration.
+            </summary>
+        </member>
+        <member name="M:Microsoft.Practices.Unity.Configuration.AssemblyElementCollection.GetElementKey(System.Configuration.ConfigurationElement)">
+            <summary>
+            Gets the element key for a specified configuration element when overridden in a derived class.
+            </summary>
+            <returns>
+            An <see cref="T:System.Object"/> that acts as the key for the specified <see cref="T:System.Configuration.ConfigurationElement"/>.
+            </returns>
+            <param name="element">The <see cref="T:System.Configuration.ConfigurationElement"/> to return the key for. </param>
+        </member>
+        <member name="T:Microsoft.Practices.Unity.Configuration.ConfigurationHelpers.ConfigurationElementExtensions">
+            <summary>
+            Helpful extension methods when implementing configuration sections
+            that deserialize "unwrapped" elements - elements that should be
+            deserialized into a container but can be present outside
+            that container in the actual config file.
+            </summary>
+        </member>
+        <member name="M:Microsoft.Practices.Unity.Configuration.ConfigurationHelpers.ConfigurationElementExtensions.ReadUnwrappedElement``1(System.Configuration.ConfigurationElement,System.Xml.XmlReader,Microsoft.Practices.Unity.Configuration.ConfigurationHelpers.DeserializableConfigurationElementCollectionBase{``0})">
+            <summary>
+            Deserialize an element of the given type, store it in
+            the collection object, and 
+            </summary>
+            <typeparam name="TElementType">Type of element to create and deserialize.</typeparam>
+            <param name="baseElement">Parent element containing element to deserialize.</param>
+            <param name="reader">Xml reader containing state to deserialize from.</param>
+            <param name="elementCollection">Collection to store the created element into.</param>
+            <returns>The created element.</returns>
+        </member>
+        <member name="M:Microsoft.Practices.Unity.Configuration.ConfigurationHelpers.ConfigurationElementExtensions.ReadElementByType``1(System.Configuration.ConfigurationElement,System.Xml.XmlReader,System.Type,Microsoft.Practices.Unity.Configuration.ConfigurationHelpers.DeserializableConfigurationElementCollectionBase{``0})">
+            <summary>
+            Deserialize an element, basing the element type on the one
+            supplied at runtime, and then store the element into the
+            given <paramref name="elementCollection"/>.
+            </summary>
+            <remarks>This method is useful when reading elements into a polymorphic collection.</remarks>
+            <typeparam name="TElementType">Base type of element to store.</typeparam>
+            <param name="baseElement">Element that contains the collection being stored into.</param>
+            <param name="reader">Xml Reader containing state to deserialize from.</param>
+            <param name="elementType">Runtime type of element to create.</param>
+            <param name="elementCollection">Collection to store the created element into.</param>
+            <returns>The created element.</returns>
+        </member>
+        <member name="T:Microsoft.Practices.Unity.Configuration.ConfigurationHelpers.InputStream">
+            <summary>
+            Class that tracks the current input state of the parser.
+            </summary>
+        </member>
+        <member name="T:Microsoft.Practices.Unity.Configuration.ConfigurationHelpers.ParseBuilder">
+            <summary>
+            A simple implementing of the rules for a Parsing Expression Grammar
+            parsing algorithm. This supplies basic methods to do the primitives
+            of the PEG, and combinators to create larger rules.
+            </summary>
+        </member>
+        <member name="M:Microsoft.Practices.Unity.Configuration.ConfigurationHelpers.ParseBuilder.Any(Microsoft.Practices.Unity.Configuration.ConfigurationHelpers.InputStream)">
+            <summary>
+            The PEG "dot" operator that matches and consumes one character.
+            </summary>
+            <param name="input">Input to the parser.</param>
+            <returns>The parse result.</returns>
+        </member>
+        <member name="M:Microsoft.Practices.Unity.Configuration.ConfigurationHelpers.ParseBuilder.Match(System.Char)">
+            <summary>
+            Parse function generator that returns a method to match a single,
+            specific character.
+            </summary>
+            <param name="charToMatch">Character to match.</param>
+            <returns>The generated parsing function.</returns>
+        </member>
+        <member name="M:Microsoft.Practices.Unity.Configuration.ConfigurationHelpers.ParseBuilder.Match(System.Func{System.Char,System.Boolean})">
+            <summary>
+            Parse function generator that checks if the current character matches
+            the predicate supplied.
+            </summary>
+            <param name="predicate">Predicate used to determine if the character is in
+            the given range.</param>
+            <returns>The generated parsing function.</returns>
+        </member>
+        <member name="M:Microsoft.Practices.Unity.Configuration.ConfigurationHelpers.ParseBuilder.ZeroOrMore(System.Func{Microsoft.Practices.Unity.Configuration.ConfigurationHelpers.InputStream,Microsoft.Practices.Unity.Configuration.ConfigurationHelpers.ParseResult})">
+            <summary>
+            The "*" operator - match zero or more of the inner parse expressions.
+            </summary>
+            <param name="inner">Parse method to repeat matching.</param>
+            <returns>The generated parsing function.</returns>
+        </member>
+        <member name="M:Microsoft.Practices.Unity.Configuration.ConfigurationHelpers.ParseBuilder.Sequence(System.Func{Microsoft.Practices.Unity.Configuration.ConfigurationHelpers.InputStream,Microsoft.Practices.Unity.Configuration.ConfigurationHelpers.ParseResult}[])">
+            <summary>
+            Parsing combinator that matches all of the given expressions in
+            order, or matches none of them.
+            </summary>
+            <param name="expressions">Expressions that form the sequence to match.</param>
+            <returns>The combined sequence.</returns>
+        </member>
+        <member name="M:Microsoft.Practices.Unity.Configuration.ConfigurationHelpers.ParseBuilder.FirstOf(System.Func{Microsoft.Practices.Unity.Configuration.ConfigurationHelpers.InputStream,Microsoft.Practices.Unity.Configuration.ConfigurationHelpers.ParseResult}[])">
+            <summary>
+            Parsing combinator that implements the PEG prioritized choice operator. Basically,
+            try each of the expressions in order, and match if any of them match, stopping on the
+            first match.
+            </summary>
+            <param name="expressions">Expressions that form the set of alternatives.</param>
+            <returns>The combined parsing method.</returns>
+        </member>
+        <member name="M:Microsoft.Practices.Unity.Configuration.ConfigurationHelpers.ParseBuilder.Not(System.Func{Microsoft.Practices.Unity.Configuration.ConfigurationHelpers.InputStream,Microsoft.Practices.Unity.Configuration.ConfigurationHelpers.ParseResult})">
+            <summary>
+            Parsing combinator implementing the "not" predicate. This wraps
+            the given <paramref name="expression"/> parsing method with a check
+            to see if it matched. If it matched, then the Not fails, and vice
+            versa. The result consumes no input.
+            </summary>
+            <param name="expression">The parse method to wrap.</param>
+            <returns>The generated parsing function.</returns>
+        </member>
+        <member name="M:Microsoft.Practices.Unity.Configuration.ConfigurationHelpers.ParseBuilder.EOF(Microsoft.Practices.Unity.Configuration.ConfigurationHelpers.InputStream)">
+            <summary>
+            Parsing expression that matches End of input.
+            </summary>
+            <param name="input">Parser input.</param>
+            <returns>Parse result</returns>
+        </member>
+        <member name="M:Microsoft.Practices.Unity.Configuration.ConfigurationHelpers.ParseBuilder.WithAction(System.Func{Microsoft.Practices.Unity.Configuration.ConfigurationHelpers.InputStream,Microsoft.Practices.Unity.Configuration.ConfigurationHelpers.ParseResult},System.Action{Microsoft.Practices.Unity.Configuration.ConfigurationHelpers.ParseResult})">
+            <summary>
+            Combinator that executes an action if the given expression matched.
+            </summary>
+            <param name="expression">Parsing expression to match.</param>
+            <param name="onMatched">Action to execute if <paramref name="expression"/>
+            matched. Input is the matched text from <paramref name="expression"/>.</param>
+            <returns>The result of <paramref name="expression"/>.</returns>
+        </member>
+        <member name="M:Microsoft.Practices.Unity.Configuration.ConfigurationHelpers.ParseBuilder.WithAction(System.Func{Microsoft.Practices.Unity.Configuration.ConfigurationHelpers.InputStream,Microsoft.Practices.Unity.Configuration.ConfigurationHelpers.ParseResult},System.Func{Microsoft.Practices.Unity.Configuration.ConfigurationHelpers.ParseResult,Microsoft.Practices.Unity.Configuration.ConfigurationHelpers.ParseResult})">
+            <summary>
+            Combinator that executes an action if the given expression matched.
+            </summary>
+            <param name="expression">parsing expression to match.</param>
+            <param name="onMatched">Method to execute if a match happens. This method returns
+            the <see cref="T:Microsoft.Practices.Unity.Configuration.ConfigurationHelpers.ParseResult"/> that will be returned from the combined expression.</param>
+            <returns>The result of <paramref name="onMatched"/> if expression matched, else
+            whatever <paramref name="expression"/> returned.</returns>
+        </member>
+        <member name="T:Microsoft.Practices.Unity.Configuration.ConfigurationHelpers.ParseResult">
+            <summary>
+            Object containing the result of attempting to match a PEG rule.
+            This object is the return type for all parsing methods.
+            </summary>
+        </member>
+        <member name="P:Microsoft.Practices.Unity.Configuration.ConfigurationHelpers.ParseResult.Matched">
+            <summary>
+            Did the rule match?
+            </summary>
+        </member>
+        <member name="P:Microsoft.Practices.Unity.Configuration.ConfigurationHelpers.ParseResult.MatchedString">
+            <summary>
+            The characters that were matched (if any)
+            </summary>
+        </member>
+        <member name="P:Microsoft.Practices.Unity.Configuration.ConfigurationHelpers.ParseResult.ResultData">
+            <summary>
+            Any extra information provided by the parsing expression
+            (only set if the parse matched). The nature
+            of the data varies depending on the parsing expression.
+            </summary>
+        </member>
+        <member name="T:Microsoft.Practices.Unity.Configuration.ConfigurationHelpers.SequenceResult">
+            <summary>
+            Helper methods to make it easier to pull the data
+            out of the result of a sequence expression.
+            </summary>
+        </member>
+        <member name="M:Microsoft.Practices.Unity.Configuration.ConfigurationHelpers.SequenceResult.System#Collections#IEnumerable#GetEnumerator">
+            <summary>
+            Returns an enumerator that iterates through a collection.
+            </summary>
+            <returns>
+            An <see cref="T:System.Collections.IEnumerator"/> object that can be used to iterate through the collection.
+            </returns>
+            <filterpriority>2</filterpriority>
+        </member>
+        <member name="M:Microsoft.Practices.Unity.Configuration.ConfigurationHelpers.SequenceResult.GetEnumerator">
+            <summary>
+            Returns an enumerator that iterates through the collection.
+            </summary>
+            <returns>
+            A <see cref="T:System.Collections.Generic.IEnumerator`1"/> that can be used to iterate through the collection.
+            </returns>
+            <filterpriority>1</filterpriority>
+        </member>
+        <member name="M:Microsoft.Practices.Unity.Configuration.ConfigurationHelpers.SequenceResult.Add(Microsoft.Practices.Unity.Configuration.ConfigurationHelpers.ParseResult)">
+            <summary>
+            Adds an item to the <see cref="T:System.Collections.Generic.ICollection`1"/>.
+            </summary>
+            <param name="item">The object to add to the <see cref="T:System.Collections.Generic.ICollection`1"/>.</param>
+            <exception cref="T:System.NotSupportedException">The <see cref="T:System.Collections.Generic.ICollection`1"/> is read-only.</exception>
+        </member>
+        <member name="M:Microsoft.Practices.Unity.Configuration.ConfigurationHelpers.SequenceResult.Clear">
+            <summary>
+            Removes all items from the <see cref="T:System.Collections.Generic.ICollection`1"/>.
+            </summary>
+            <exception cref="T:System.NotSupportedException">The <see cref="T:System.Collections.Generic.ICollection`1"/> is read-only. </exception>
+        </member>
+        <member name="M:Microsoft.Practices.Unity.Configuration.ConfigurationHelpers.SequenceResult.Contains(Microsoft.Practices.Unity.Configuration.ConfigurationHelpers.ParseResult)">
+            <summary>
+            Determines whether the <see cref="T:System.Collections.Generic.ICollection`1"/> contains a specific value.
+            </summary>
+            <returns>
+            true if <paramref name="item"/> is found in the <see cref="T:System.Collections.Generic.ICollection`1"/>; otherwise, false.
+            </returns>
+            <param name="item">The object to locate in the <see cref="T:System.Collections.Generic.ICollection`1"/>.</param>
+        </member>
+        <member name="M:Microsoft.Practices.Unity.Configuration.ConfigurationHelpers.SequenceResult.CopyTo(Microsoft.Practices.Unity.Configuration.ConfigurationHelpers.ParseResult[],System.Int32)">
+            <summary>
+            Copies the elements of the sequence to an <see cref="T:System.Array"/>, starting at a particular <see cref="T:System.Array"/> index.
+            </summary>
+            <param name="array">The one-dimensional <see cref="T:System.Array"/> that is the destination of the elements copied from this sequence. The <see cref="T:System.Array"/> must have zero-based indexing.</param>
+            <param name="arrayIndex">The zero-based index in <paramref name="array"/> at which copying begins.</param>
+            <exception cref="T:System.ArgumentNullException"><paramref name="array"/> is null.</exception>
+            <exception cref="T:System.ArgumentOutOfRangeException"><paramref name="arrayIndex"/> is less than 0.</exception>
+            <exception cref="T:System.ArgumentException"><paramref name="array"/> is multidimensional.
+            -or-
+            <paramref name="arrayIndex"/> is equal to or greater than the length of <paramref name="array"/>.
+            -or-
+            The number of elements in the source is greater than the available space from <paramref name="arrayIndex"/> to the end of the destination <paramref name="array"/>.
+            </exception>
+        </member>
+        <member name="M:Microsoft.Practices.Unity.Configuration.ConfigurationHelpers.SequenceResult.Remove(Microsoft.Practices.Unity.Configuration.ConfigurationHelpers.ParseResult)">
+            <summary>
+            Removes the first occurrence of a specific object from the <see cref="T:System.Collections.Generic.ICollection`1"/>.
+            </summary>
+            <returns>
+            true if <paramref name="item"/> was successfully removed from the <see cref="T:System.Collections.Generic.ICollection`1"/>; otherwise, false. This method also returns false if <paramref name="item"/> is not found in the original <see cref="T:System.Collections.Generic.ICollection`1"/>.
+            </returns>
+            <param name="item">The object to remove from the <see cref="T:System.Collections.Generic.ICollection`1"/>.</param>
+            <exception cref="T:System.NotSupportedException">The <see cref="T:System.Collections.Generic.ICollection`1"/> is read-only.</exception>
+        </member>
+        <member name="M:Microsoft.Practices.Unity.Configuration.ConfigurationHelpers.SequenceResult.IndexOf(Microsoft.Practices.Unity.Configuration.ConfigurationHelpers.ParseResult)">
+            <summary>
+            Determines the index of a specific item in the <see cref="T:System.Collections.Generic.IList`1"/>.
+            </summary>
+            <returns>
+            The index of <paramref name="item"/> if found in the list; otherwise, -1.
+            </returns>
+            <param name="item">The object to locate in the <see cref="T:System.Collections.Generic.IList`1"/>.</param>
+        </member>
+        <member name="M:Microsoft.Practices.Unity.Configuration.ConfigurationHelpers.SequenceResult.Insert(System.Int32,Microsoft.Practices.Unity.Configuration.ConfigurationHelpers.ParseResult)">
+            <summary>
+            Inserts an item to the <see cref="T:System.Collections.Generic.IList`1"/> at the specified index.
+            </summary>
+            <param name="index">The zero-based index at which <paramref name="item"/> should be inserted.</param>
+            <param name="item">The object to insert into the <see cref="T:System.Collections.Generic.IList`1"/>.</param>
+            <exception cref="T:System.ArgumentOutOfRangeException"><paramref name="index"/> is not a valid index in the <see cref="T:System.Collections.Generic.IList`1"/>.</exception>
+            <exception cref="T:System.NotSupportedException">The <see cref="T:System.Collections.Generic.IList`1"/> is read-only.</exception>
+        </member>
+        <member name="M:Microsoft.Practices.Unity.Configuration.ConfigurationHelpers.SequenceResult.RemoveAt(System.Int32)">
+            <summary>
+            Removes the <see cref="T:System.Collections.Generic.IList`1"/> item at the specified index.
+            </summary>
+            <param name="index">The zero-based index of the item to remove.</param>
+            <exception cref="T:System.ArgumentOutOfRangeException"><paramref name="index"/> is not a valid index in the <see cref="T:System.Collections.Generic.IList`1"/>.</exception>
+            <exception cref="T:System.NotSupportedException">The <see cref="T:System.Collections.Generic.IList`1"/> is read-only.</exception>
+        </member>
+        <member name="P:Microsoft.Practices.Unity.Configuration.ConfigurationHelpers.SequenceResult.Count">
+            <summary>
+            Gets the number of elements contained in the <see cref="T:System.Collections.Generic.ICollection`1"/>.
+            </summary>
+            <returns>
+            The number of elements contained in the <see cref="T:System.Collections.Generic.ICollection`1"/>.
+            </returns>
+        </member>
+        <member name="P:Microsoft.Practices.Unity.Configuration.ConfigurationHelpers.SequenceResult.IsReadOnly">
+            <summary>
+            Gets a value indicating whether the <see cref="T:System.Collections.Generic.ICollection`1"/> is read-only.
+            </summary>
+            <returns>
+            true if the <see cref="T:System.Collections.Generic.ICollection`1"/> is read-only; otherwise, false.
+            </returns>
+        </member>
+        <member name="P:Microsoft.Practices.Unity.Configuration.ConfigurationHelpers.SequenceResult.Item(System.Int32)">
+            <summary>
+            Gets or sets the element at the specified index.
+            </summary>
+            <returns>
+            The element at the specified index.
+            </returns>
+            <param name="index">The zero-based index of the element to get or set.</param>
+            <exception cref="T:System.ArgumentOutOfRangeException"><paramref name="index"/> is not a valid index in the <see cref="T:System.Collections.Generic.IList`1"/>.</exception>
+            <exception cref="T:System.NotSupportedException">The property is set and the <see cref="T:System.Collections.Generic.IList`1"/> is read-only.</exception>
+        </member>
+        <member name="T:Microsoft.Practices.Unity.Configuration.ConfigurationHelpers.TypeNameInfo">
+            <summary>
+            Class containing information about a type name.
+            </summary>
+        </member>
+        <member name="P:Microsoft.Practices.Unity.Configuration.ConfigurationHelpers.TypeNameInfo.Name">
+            <summary>
+            The base name of the class
+            </summary>
+        </member>
+        <member name="P:Microsoft.Practices.Unity.Configuration.ConfigurationHelpers.TypeNameInfo.Namespace">
+            <summary>
+            Namespace if any
+            </summary>
+        </member>
+        <member name="P:Microsoft.Practices.Unity.Configuration.ConfigurationHelpers.TypeNameInfo.AssemblyName">
+            <summary>
+            Assembly name, if any
+            </summary>
+        </member>
+        <member name="T:Microsoft.Practices.Unity.Configuration.ConfigurationHelpers.XmlWriterExtensions">
+            <summary>
+            Helper methods on <see cref="T:System.Xml.XmlWriter"/>.
+            </summary>
+        </member>
+        <member name="M:Microsoft.Practices.Unity.Configuration.ConfigurationHelpers.XmlWriterExtensions.WriteElement(System.Xml.XmlWriter,System.String,System.Action{System.Xml.XmlWriter})">
+            <summary>
+            A helper method to make it more foolproof to write elements. This takes care of writing the
+            start and end element tags, and takes a nested closure with the code to write the content of
+            the tag. That way the caller doesn't need to worry about the details of getting the start
+            and end tags correct.
+            </summary>
+            <remarks>
+            We don't support XML Namespaces here because .NET configuration doesn't use them so
+            we don't need it for this current implementation.
+            </remarks>
+            <param name="writer">XmlWriter to write to.</param>
+            <param name="elementName">Name of element.</param>
+            <param name="writeContent">Nested lambda which, when executed, will create the content for the
+            element.</param>
+            <returns><paramref name="writer"/> (for method chaining if desired).</returns>
+        </member>
+        <member name="M:Microsoft.Practices.Unity.Configuration.ConfigurationHelpers.XmlWriterExtensions.WriteAttributeIfNotEmpty(System.Xml.XmlWriter,System.String,System.String)">
+            <summary>
+            A helper method to make it easier to output attributes. If the <paramref name="attributeValue"/> is
+            null or an empty string, output nothing, else output the given XML attribute.
+            </summary>
+            <param name="writer">Writer to output to.</param>
+            <param name="attributeName">Attribute name to write.</param>
+            <param name="attributeValue">Value for the attribute.</param>
+            <returns><paramref name="writer"/> (for method chaining if desired).</returns>
+        </member>
+        <member name="T:Microsoft.Practices.Unity.Configuration.ContainerConfiguringElement">
+            <summary>
+            A base class for those elements that can be used
+            to configure a unity container.
+            </summary>
+        </member>
+        <member name="M:Microsoft.Practices.Unity.Configuration.ContainerConfiguringElement.#ctor">
+            <summary>
+            Create a new instance of <see cref="T:Microsoft.Practices.Unity.Configuration.ContainerConfiguringElement"/>.
+            </summary>
+        </member>
+        <member name="M:Microsoft.Practices.Unity.Configuration.ContainerConfiguringElement.ConfigureContainer(Microsoft.Practices.Unity.IUnityContainer)">
+            <summary>
+            Apply this element's configuration to the given <paramref name="container"/>.
+            </summary>
+            <param name="container">Container to configure.</param>
+        </member>
+        <member name="P:Microsoft.Practices.Unity.Configuration.ContainerConfiguringElement.Key">
+            <summary>
+            Return a unique key that can be used to manage this element in a collection.
+            </summary>
+        </member>
+        <member name="T:Microsoft.Practices.Unity.Configuration.ConfigurationHelpers.DictionaryExtensions">
+            <summary>
+            A couple of useful extension methods on IDictionary
+            </summary>
+        </member>
+        <member name="M:Microsoft.Practices.Unity.Configuration.ConfigurationHelpers.DictionaryExtensions.GetOrNull``2(System.Collections.Generic.IDictionary{``0,``1},``0)">
+            <summary>
+            Get the value from a dictionary, or null if there is no value.
+            </summary>
+            <typeparam name="TKey">Key type of dictionary.</typeparam>
+            <typeparam name="TValue">Value type of dictionary.</typeparam>
+            <param name="dictionary">Dictionary to search.</param>
+            <param name="key">Key to look up.</param>
+            <returns>The value at the key or null if not in the dictionary.</returns>
+        </member>
+        <member name="T:Microsoft.Practices.Unity.Configuration.ConfigurationHelpers.UnknownElementHandlerMap`1">
+            <summary>
+            A helper class used to map element tag names to a handler method
+            used to interpret that element.
+            </summary>
+            <typeparam name="TContainingElement"></typeparam>
+        </member>
+        <member name="M:Microsoft.Practices.Unity.Configuration.ConfigurationHelpers.UnknownElementHandlerMap`1.Add(System.String,System.Action{`0,System.Xml.XmlReader})">
+            <summary>
+            Add method to enable dictionary initializer syntax 
+            </summary>
+            <param name="elementName"></param>
+            <param name="processingAction"></param>
+        </member>
+        <member name="M:Microsoft.Practices.Unity.Configuration.ConfigurationHelpers.UnknownElementHandlerMap`1.ProcessElement(`0,System.String,System.Xml.XmlReader)">
+            <summary>
+            Process an unknown element according to the map entries.
+            </summary>
+            <param name="parentElement">Parent element that hit this unknown element.</param>
+            <param name="elementName">Name of the unknown element.</param>
+            <param name="reader">XmlReader positioned at start of element.</param>
+            <returns>true if processed, false if not.</returns>
+        </member>
+        <member name="M:Microsoft.Practices.Unity.Configuration.ConfigurationHelpers.UnknownElementHandlerMap`1.System#Collections#IEnumerable#GetEnumerator">
+            <summary>
+            Returns an enumerator that iterates through a collection.
+            </summary>
+            <returns>
+            An <see cref="T:System.Collections.IEnumerator"/> object that can be used to iterate through the collection.
+            </returns>
+            <filterpriority>2</filterpriority>
+        </member>
+        <member name="M:Microsoft.Practices.Unity.Configuration.ConfigurationHelpers.UnknownElementHandlerMap`1.GetEnumerator">
+            <summary>
+            Returns an enumerator that iterates through the collection.
+            </summary>
+            <returns>
+            A <see cref="T:System.Collections.Generic.IEnumerator`1"/> that can be used to iterate through the collection.
+            </returns>
+            <filterpriority>1</filterpriority>
+        </member>
+        <member name="T:Microsoft.Practices.Unity.Configuration.ConfigurationHelpers.UnknownElementHandlerMap">
+            <summary>
+            A helper class used to map element tag names to a handler method
+            used to interpret that element.
+            </summary>
+        </member>
+        <member name="M:Microsoft.Practices.Unity.Configuration.ConfigurationHelpers.UnknownElementHandlerMap.System#Collections#IEnumerable#GetEnumerator">
+            <summary>
+            Returns an enumerator that iterates through a collection.
+            </summary>
+            <returns>
+            An <see cref="T:System.Collections.IEnumerator"/> object that can be used to iterate through the collection.
+            </returns>
+            <filterpriority>2</filterpriority>
+        </member>
+        <member name="M:Microsoft.Practices.Unity.Configuration.ConfigurationHelpers.UnknownElementHandlerMap.GetEnumerator">
+            <summary>
+            Returns an enumerator that iterates through the collection.
+            </summary>
+            <returns>
+            A <see cref="T:System.Collections.Generic.IEnumerator`1"/> that can be used to iterate through the collection.
+            </returns>
+            <filterpriority>1</filterpriority>
+        </member>
+        <member name="T:Microsoft.Practices.Unity.Configuration.ConfigurationHelpers.ValueElementHelper">
+            <summary>
+            A helper class that assists in deserializing parameter and property
+            elements. These elements both have a single "value" child element that
+            specify the value to inject for the property or parameter.
+            </summary>
+        </member>
+        <member name="M:Microsoft.Practices.Unity.Configuration.ConfigurationHelpers.ValueElementHelper.#ctor(Microsoft.Practices.Unity.Configuration.ConfigurationHelpers.IValueProvidingElement)">
+            <summary>
+            Create a new <see cref="T:Microsoft.Practices.Unity.Configuration.ConfigurationHelpers.ValueElementHelper"/> that wraps reading
+            values and storing them in the given <paramref name="parentElement"/>.
+            </summary>
+            <param name="parentElement">Element that contains the value elements.</param>
+        </member>
+        <member name="M:Microsoft.Practices.Unity.Configuration.ConfigurationHelpers.ValueElementHelper.GetValue(Microsoft.Practices.Unity.Configuration.ParameterValueElement)">
+            <summary>
+            Gets a <see cref="T:Microsoft.Practices.Unity.Configuration.ParameterValueElement"/>, or if none is present,
+            returns a default <see cref="T:Microsoft.Practices.Unity.Configuration.DependencyElement"/>.
+            </summary>
+            <param name="currentValue">The <see cref="T:Microsoft.Practices.Unity.Configuration.ParameterValueElement"/>.</param>
+            <returns>The given <paramref name="currentValue"/>, unless
+            <paramref name="currentValue"/> is null, in which case returns
+            a <see cref="T:Microsoft.Practices.Unity.Configuration.DependencyElement"/>.</returns>
+        </member>
+        <member name="M:Microsoft.Practices.Unity.Configuration.ConfigurationHelpers.ValueElementHelper.DeserializeUnrecognizedAttribute(System.String,System.String)">
+            <summary>
+            Helper method used during deserialization to handle
+            attributes for the dependency and value tags.
+            </summary>
+            <param name="name">attribute name.</param>
+            <param name="value">attribute value.</param>
+            <returns>true</returns>
+        </member>
+        <member name="M:Microsoft.Practices.Unity.Configuration.ConfigurationHelpers.ValueElementHelper.DeserializeUnknownElement(System.String,System.Xml.XmlReader)">
+            <summary>
+            Helper method used during deserialization to handle the default
+            value element tags.
+            </summary>
+            <param name="elementName">The element name.</param>
+            <param name="reader">XML data to read.</param>
+            <returns>True if deserialization succeeded, false if it failed.</returns>
+        </member>
+        <member name="M:Microsoft.Practices.Unity.Configuration.ConfigurationHelpers.ValueElementHelper.CompleteValueElement(System.Xml.XmlReader)">
+            <summary>
+            Call this method at the end of deserialization of your element to
+            set your value element.
+            </summary>
+        </member>
+        <member name="M:Microsoft.Practices.Unity.Configuration.ConfigurationHelpers.ValueElementHelper.SerializeParameterValueElement(System.Xml.XmlWriter,Microsoft.Practices.Unity.Configuration.ParameterValueElement,System.Boolean)">
+            <summary>
+            Serialize a <see cref="T:Microsoft.Practices.Unity.Configuration.ParameterValueElement"/> object out to XML.
+            This method is aware of and implements the shorthand attributes
+            for dependency and value elements.
+            </summary>
+            <param name="writer">Writer to output XML to.</param>
+            <param name="element">The <see cref="T:Microsoft.Practices.Unity.Configuration.ParameterValueElement"/> to serialize.</param>
+            <param name="elementsOnly">If true, always output an element. If false, then
+            dependency and value elements will be serialized as attributes in the parent tag.</param>
+        </member>
+        <member name="T:Microsoft.Practices.Unity.Configuration.ConstructorElement">
+            <summary>
+            Configuration element representing a constructor configuration.
+            </summary>
+        </member>
+        <member name="T:Microsoft.Practices.Unity.Configuration.InjectionMemberElement">
+            <summary>
+            Base class for configuration elements that generate <see cref="T:Microsoft.Practices.Unity.InjectionMember"/>
+            object to configure a container.
+            </summary>
+        </member>
+        <member name="M:Microsoft.Practices.Unity.Configuration.InjectionMemberElement.GetInjectionMembers(Microsoft.Practices.Unity.IUnityContainer,System.Type,System.Type,System.String)">
+            <summary>
+            Return the set of <see cref="T:Microsoft.Practices.Unity.InjectionMember"/>s that are needed
+            to configure the container according to this configuration element.
+            </summary>
+            <param name="container">Container that is being configured.</param>
+            <param name="fromType">Type that is being registered.</param>
+            <param name="toType">Type that <paramref name="fromType"/> is being mapped to.</param>
+            <param name="name">Name this registration is under.</param>
+            <returns>One or more <see cref="T:Microsoft.Practices.Unity.InjectionMember"/> objects that should be
+            applied to the container registration.</returns>
+        </member>
+        <member name="M:Microsoft.Practices.Unity.Configuration.InjectionMemberElement.GetMemberElementName(Microsoft.Practices.Unity.Configuration.InjectionMemberElement)">
+            <summary>
+            Get the standard tag name for an <see cref="T:Microsoft.Practices.Unity.Configuration.InjectionMemberElement"/>
+            taking into account currently loaded section extensions.
+            </summary>
+            <param name="memberElement">Element to get the name for.</param>
+            <returns>The element name.</returns>
+            <exception cref="T:System.ArgumentException">If the member element is not currently registered
+            with the section.</exception>
+        </member>
+        <member name="P:Microsoft.Practices.Unity.Configuration.InjectionMemberElement.Key">
+            <summary>
+            Each element must have a unique key, which is generated by the subclasses.
+            </summary>
+        </member>
+        <member name="P:Microsoft.Practices.Unity.Configuration.InjectionMemberElement.ElementName">
+            <summary>
+            Element name to use to serialize this into XML.
+            </summary>
+        </member>
+        <member name="M:Microsoft.Practices.Unity.Configuration.ConstructorElement.SerializeContent(System.Xml.XmlWriter)">
+            <summary>
+            Write the contents of this element to the given <see cref="T:System.Xml.XmlWriter"/>.
+            </summary>
+            <remarks>The caller of this method has already written the start element tag before
+            calling this method, so deriving classes only need to write the element content, not
+            the start or end tags.</remarks>
+            <param name="writer">Writer to send XML content to.</param>
+        </member>
+        <member name="M:Microsoft.Practices.Unity.Configuration.ConstructorElement.GetInjectionMembers(Microsoft.Practices.Unity.IUnityContainer,System.Type,System.Type,System.String)">
+            <summary>
+            Return the set of <see cref="T:Microsoft.Practices.Unity.InjectionMember"/>s that are needed
+            to configure the container according to this configuration element.
+            </summary>
+            <param name="container">Container that is being configured.</param>
+            <param name="fromType">Type that is being registered.</param>
+            <param name="toType">Type that <paramref name="fromType"/> is being mapped to.</param>
+            <param name="name">Name this registration is under.</param>
+            <returns>One or more <see cref="T:Microsoft.Practices.Unity.InjectionMember"/> objects that should be
+            applied to the container registration.</returns>
+        </member>
+        <member name="P:Microsoft.Practices.Unity.Configuration.ConstructorElement.Parameters">
+            <summary>
+            The parameters of the constructor to call.
+            </summary>
+        </member>
+        <member name="P:Microsoft.Practices.Unity.Configuration.ConstructorElement.Key">
+            <summary>
+            Each element must have a unique key, which is generated by the subclasses.
+            </summary>
+        </member>
+        <member name="P:Microsoft.Practices.Unity.Configuration.ConstructorElement.ElementName">
+            <summary>
+            Element name to use to serialize this into XML.
+            </summary>
+        </member>
+        <member name="T:Microsoft.Practices.Unity.Configuration.ContainerConfiguringElementCollection">
+            <summary>
+            A collection of <see cref="T:Microsoft.Practices.Unity.Configuration.ContainerConfiguringElement"/>s as
+            loaded from configuration.
+            </summary>
+        </member>
+        <member name="M:Microsoft.Practices.Unity.Configuration.ContainerConfiguringElementCollection.CreateNewElement">
+            <summary>
+            When overridden in a derived class, creates a new <see cref="T:System.Configuration.ConfigurationElement"/>.
+            </summary>
+            <returns>
+            A new <see cref="T:System.Configuration.ConfigurationElement"/>.
+            </returns>
+        </member>
+        <member name="M:Microsoft.Practices.Unity.Configuration.ContainerConfiguringElementCollection.GetElementKey(System.Configuration.ConfigurationElement)">
+            <summary>
+            Gets the element key for a specified configuration element when overridden in a derived class.
+            </summary>
+            <returns>
+            An <see cref="T:System.Object"/> that acts as the key for the specified <see cref="T:System.Configuration.ConfigurationElement"/>.
+            </returns>
+            <param name="element">The <see cref="T:System.Configuration.ConfigurationElement"/> to return the key for. </param>
+        </member>
+        <member name="T:Microsoft.Practices.Unity.Configuration.ContainerElement">
+            <summary>
+            A configuration element class defining the set of registrations to be
+            put into a container.
+            </summary>
+        </member>
+        <member name="M:Microsoft.Practices.Unity.Configuration.ContainerElement.Configure(Microsoft.Practices.Unity.IUnityContainer)">
+            <summary>
+            Original configuration API kept for backwards compatibility.
+            </summary>
+            <param name="container">Container to configure</param>
+        </member>
+        <member name="M:Microsoft.Practices.Unity.Configuration.ContainerElement.ConfigureContainer(Microsoft.Practices.Unity.IUnityContainer)">
+            <summary>
+            Apply the configuration information in this element to the
+            given <paramref name="container"/>.
+            </summary>
+            <param name="container">Container to configure.</param>
+        </member>
+        <member name="M:Microsoft.Practices.Unity.Configuration.ContainerElement.OnDeserializeUnrecognizedElement(System.String,System.Xml.XmlReader)">
+            <summary>
+            Gets a value indicating whether an unknown element is encountered during deserialization.
+            </summary>
+            <returns>
+            true when an unknown element is encountered while deserializing; otherwise, false.
+            </returns>
+            <param name="elementName">The name of the unknown subelement.</param>
+            <param name="reader">The <see cref="T:System.Xml.XmlReader"/> being used for deserialization.</param>
+            <exception cref="T:System.Configuration.ConfigurationErrorsException">The element identified by <paramref name="elementName"/> is locked.
+            - or -
+            One or more of the element's attributes is locked.
+            - or -
+            <paramref name="elementName"/> is unrecognized, or the element has an unrecognized attribute.
+            - or -
+            The element has a Boolean attribute with an invalid value.
+            - or -
+            An attempt was made to deserialize a property more than once.
+            - or -
+            An attempt was made to deserialize a property that is not a valid member of the element.
+            - or -
+            The element cannot contain a CDATA or text element.
+            </exception>
+        </member>
+        <member name="M:Microsoft.Practices.Unity.Configuration.ContainerElement.SerializeContent(System.Xml.XmlWriter)">
+            <summary>
+            Write the contents of this element to the given <see cref="T:System.Xml.XmlWriter"/>.
+            </summary>
+            <remarks>The caller of this method has already written the start element tag before
+            calling this method, so deriving classes only need to write the element content, not
+            the start or end tags.</remarks>
+            <param name="writer">Writer to send XML content to.</param>
+        </member>
+        <member name="P:Microsoft.Practices.Unity.Configuration.ContainerElement.Name">
+            <summary>
+            Name for this container configuration as given in the config file.
+            </summary>
+        </member>
+        <member name="P:Microsoft.Practices.Unity.Configuration.ContainerElement.Registrations">
+            <summary>
+            The type registrations in this container.
+            </summary>
+        </member>
+        <member name="P:Microsoft.Practices.Unity.Configuration.ContainerElement.Instances">
+            <summary>
+            Any instances to register in the container.
+            </summary>
+        </member>
+        <member name="P:Microsoft.Practices.Unity.Configuration.ContainerElement.Extensions">
+            <summary>
+            Any extensions to add to the container.
+            </summary>
+        </member>
+        <member name="P:Microsoft.Practices.Unity.Configuration.ContainerElement.ConfiguringElements">
+            <summary>
+            Set of any extra configuration elements that were added by a
+            section extension.
+            </summary>
+            <remarks>
+            This is not marked as a configuration property because we don't want
+            the actual property to show up as a nested element in the configuration.</remarks>
+        </member>
+        <member name="T:Microsoft.Practices.Unity.Configuration.ContainerExtensionElement">
+            <summary>
+            Configuration element representing an extension to add to a container.
+            </summary>
+        </member>
+        <member name="M:Microsoft.Practices.Unity.Configuration.ContainerExtensionElement.ConfigureContainer(Microsoft.Practices.Unity.IUnityContainer)">
+            <summary>
+            Add the extension specified in this element to the container.
+            </summary>
+            <param name="container">Container to configure.</param>
+        </member>
+        <member name="M:Microsoft.Practices.Unity.Configuration.ContainerExtensionElement.SerializeContent(System.Xml.XmlWriter)">
+            <summary>
+            Write the contents of this element to the given <see cref="T:System.Xml.XmlWriter"/>.
+            </summary>
+            <remarks>The caller of this method has already written the start element tag before
+            calling this method, so deriving classes only need to write the element content, not
+            the start or end tags.</remarks>
+            <param name="writer">Writer to send XML content to.</param>
+        </member>
+        <member name="P:Microsoft.Practices.Unity.Configuration.ContainerExtensionElement.TypeName">
+            <summary>
+            Type of the extension to add.
+            </summary>
+        </member>
+        <member name="T:Microsoft.Practices.Unity.Configuration.ContainerExtensionElementCollection">
+            <summary>
+            A collection of <see cref="T:Microsoft.Practices.Unity.Configuration.ContainerExtensionElement"/>s.
+            </summary>
+        </member>
+        <member name="M:Microsoft.Practices.Unity.Configuration.ContainerExtensionElementCollection.GetElementKey(System.Configuration.ConfigurationElement)">
+            <summary>
+            Gets the element key for a specified configuration element when overridden in a derived class.
+            </summary>
+            <returns>
+            An <see cref="T:System.Object"/> that acts as the key for the specified <see cref="T:System.Configuration.ConfigurationElement"/>.
+            </returns>
+            <param name="element">The <see cref="T:System.Configuration.ConfigurationElement"/> to return the key for. </param>
+        </member>
+        <member name="T:Microsoft.Practices.Unity.Configuration.DependencyElement">
+            <summary>
+            A <see cref="T:Microsoft.Practices.Unity.Configuration.ParameterValueElement"/> derived class that describes
+            a parameter that should be resolved through the container.
+            </summary>
+        </member>
+        <member name="M:Microsoft.Practices.Unity.Configuration.DependencyElement.#ctor">
+            <summary>
+            Create a new instance of <see cref="T:Microsoft.Practices.Unity.Configuration.DependencyElement"/>.
+            </summary>
+        </member>
+        <member name="M:Microsoft.Practices.Unity.Configuration.DependencyElement.#ctor(System.Collections.Generic.IDictionary{System.String,System.String})">
+            <summary>
+            Create a new instance of <see cref="T:Microsoft.Practices.Unity.Configuration.DependencyElement"/> with
+            properties initialized from the contents of 
+            <paramref name="attributeValues"/>.
+            </summary>
+            <param name="attributeValues">Dictionary of name/value pairs to
+            initialize this object with.</param>
+        </member>
+        <member name="M:Microsoft.Practices.Unity.Configuration.DependencyElement.Microsoft#Practices#Unity#Configuration#IAttributeOnlyElement#SerializeContent(System.Xml.XmlWriter)">
+            <summary>
+            Write the contents of this element to the given <see cref="T:System.Xml.XmlWriter"/>.
+            </summary>
+            <remarks>The caller of this method has already written the start element tag before
+            calling this method, so deriving classes only need to write the element content, not
+            the start or end tags.</remarks>
+            <param name="writer">Writer to send XML content to.</param>
+        </member>
+        <member name="M:Microsoft.Practices.Unity.Configuration.DependencyElement.SerializeContent(System.Xml.XmlWriter)">
+            <summary>
+            Write the contents of this element to the given <see cref="T:System.Xml.XmlWriter"/>. This
+            method always outputs an explicit &lt;dependency&gt; tag, instead of providing
+            attributes to the parent method.
+            </summary>
+            <param name="writer">Writer to send XML content to.</param>
+        </member>
+        <member name="M:Microsoft.Practices.Unity.Configuration.DependencyElement.GetInjectionParameterValue(Microsoft.Practices.Unity.IUnityContainer,System.Type)">
+            <summary>
+            Generate an <see cref="T:Microsoft.Practices.Unity.InjectionParameterValue"/> object
+            that will be used to configure the container for a type registration.
+            </summary>
+            <param name="container">Container that is being configured. Supplied in order
+            to let custom implementations retrieve services; do not configure the container
+            directly in this method.</param>
+            <param name="parameterType">Type of the </param>
+            <returns></returns>
+        </member>
+        <member name="P:Microsoft.Practices.Unity.Configuration.DependencyElement.Name">
+            <summary>
+            Name to use to when resolving. If empty, resolves the default.
+            </summary>
+        </member>
+        <member name="P:Microsoft.Practices.Unity.Configuration.DependencyElement.TypeName">
+            <summary>
+            Name of type this dependency should resolve to. This is optional;
+            without it the container will resolve the type of whatever
+            property or parameter this element is contained in.
+            </summary>
+        </member>
+        <member name="T:Microsoft.Practices.Unity.Configuration.ExtensionConfigurationElement">
+            <summary>
+            Base class used to derive new elements that can occur
+            directly within a container element.
+            </summary>
+        </member>
+        <member name="M:Microsoft.Practices.Unity.Configuration.ExtensionConfigurationElement.#ctor">
+            <summary>
+            Initialize a new instance of <see cref="T:Microsoft.Practices.Unity.Configuration.ExtensionConfigurationElement"/>.
+            </summary>
+        </member>
+        <member name="M:Microsoft.Practices.Unity.Configuration.ExtensionConfigurationElement.ConfigureContainer(Microsoft.Practices.Unity.IUnityContainer)">
+            <summary>
+            When overridden in a derived class, this method will make configuration
+            calls into the given <paramref name="container"/> according to its contents.
+            </summary>
+            <param name="container">The container to configure.</param>
+        </member>
+        <member name="P:Microsoft.Practices.Unity.Configuration.ExtensionConfigurationElement.Key">
+            <summary>
+            Unique key generated for use in the collection class.
+            </summary>
+        </member>
+        <member name="T:Microsoft.Practices.Unity.Configuration.ExtensionConfigurationElementCollection">
+            <summary>
+            A collection of <see cref="T:Microsoft.Practices.Unity.Configuration.ExtensionConfigurationElement"/>s.
+            </summary>
+        </member>
+        <member name="M:Microsoft.Practices.Unity.Configuration.ExtensionConfigurationElementCollection.CreateNewElement">
+            <summary>
+            When overridden in a derived class, creates a new <see cref="T:System.Configuration.ConfigurationElement"/>.
+            </summary>
+            <returns>
+            A new <see cref="T:System.Configuration.ConfigurationElement"/>.
+            </returns>
+        </member>
+        <member name="M:Microsoft.Practices.Unity.Configuration.ExtensionConfigurationElementCollection.GetElementKey(System.Configuration.ConfigurationElement)">
+            <summary>
+            Gets the element key for a specified configuration element when overridden in a derived class.
+            </summary>
+            <returns>
+            An <see cref="T:System.Object"/> that acts as the key for the specified <see cref="T:System.Configuration.ConfigurationElement"/>.
+            </returns>
+            <param name="element">The <see cref="T:System.Configuration.ConfigurationElement"/> to return the key for. </param>
+        </member>
+        <member name="T:Microsoft.Practices.Unity.Configuration.ExtensionElementMap">
+            <summary>
+            This class manages the set of extension elements
+            added by section elements.
+            </summary>
+        </member>
+        <member name="M:Microsoft.Practices.Unity.Configuration.ExtensionElementMap.Clear">
+            <summary>
+            Clear the current set of extension elements.
+            </summary>
+        </member>
+        <member name="M:Microsoft.Practices.Unity.Configuration.ExtensionElementMap.AddContainerConfiguringElement(System.String,System.String,System.Type)">
+            <summary>
+            Register a new ContainerExtensionConfigurationElement with he section so it
+            can be read.
+            </summary>
+            <param name="prefix">prefix if any.</param>
+            <param name="tag">tag name.</param>
+            <param name="elementType">Type of element to register.</param>
+        </member>
+        <member name="M:Microsoft.Practices.Unity.Configuration.ExtensionElementMap.AddInjectionMemberElement(System.String,System.String,System.Type)">
+            <summary>
+            Register a new <see cref="T:Microsoft.Practices.Unity.Configuration.InjectionMemberElement"/> with the section
+            so it can be read.
+            </summary>
+            <param name="prefix">prefix if any.</param>
+            <param name="tag">Tag name.</param>
+            <param name="elementType">Type of element to register.</param>
+        </member>
+        <member name="M:Microsoft.Practices.Unity.Configuration.ExtensionElementMap.AddParameterValueElement(System.String,System.String,System.Type)">
+            <summary>
+            Register a new <see cref="T:Microsoft.Practices.Unity.Configuration.ParameterValueElement"/> with the section
+            so it can be read.
+            </summary>
+            <param name="prefix">prefix if any.</param>
+            <param name="tag">Tag name.</param>
+            <param name="elementType">Type of element to register.</param>
+        </member>
+        <member name="M:Microsoft.Practices.Unity.Configuration.ExtensionElementMap.GetContainerConfiguringElementType(System.String)">
+            <summary>
+            Retrieve the <see cref="T:Microsoft.Practices.Unity.Configuration.ContainerConfiguringElement"/> registered for the given
+            tag.
+            </summary>
+            <param name="tag">Tag to look up.</param>
+            <returns>Type of element, or null if not registered.</returns>
+        </member>
+        <member name="M:Microsoft.Practices.Unity.Configuration.ExtensionElementMap.GetInjectionMemberElementType(System.String)">
+            <summary>
+            Retrieve the ContainerExtensionConfigurationElement registered for the given
+            tag.
+            </summary>
+            <param name="tag">Tag to look up.</param>
+            <returns>Type of element, or null if not registered.</returns>
+        </member>
+        <member name="M:Microsoft.Practices.Unity.Configuration.ExtensionElementMap.GetParameterValueElementType(System.String)">
+            <summary>
+            Retrieve the ContainerExtensionConfigurationElement registered for the given
+            tag.
+            </summary>
+            <param name="tag">Tag to look up.</param>
+            <returns>Type of element, or null if not registered.</returns>
+        </member>
+        <member name="M:Microsoft.Practices.Unity.Configuration.ExtensionElementMap.GetTagForExtensionElement(System.Configuration.ConfigurationElement)">
+            <summary>
+            Retrieve the correct tag to use when serializing the given
+            <paramref name="element"/> to XML.
+            </summary>
+            <param name="element">Element to be serialized.</param>
+            <returns>The tag for that element type.</returns>
+            <exception cref="T:System.ArgumentException"> if the element is of a type that
+            is not registered with the section already.</exception>
+        </member>
+        <member name="T:Microsoft.Practices.Unity.Configuration.InjectionMemberElementCollection">
+            <summary>
+            A polymorphic collection of <see cref="T:Microsoft.Practices.Unity.Configuration.InjectionMemberElement"/>s.
+            </summary>
+        </member>
+        <member name="M:Microsoft.Practices.Unity.Configuration.InjectionMemberElementCollection.OnDeserializeUnrecognizedElement(System.String,System.Xml.XmlReader)">
+            <summary>
+            Causes the configuration system to throw an exception.
+            </summary>
+            <returns>
+            true if the unrecognized element was deserialized successfully; otherwise, false. The default is false.
+            </returns>
+            <param name="elementName">The name of the unrecognized element.</param>
+            <param name="reader">An input stream that reads XML from the configuration file. </param>
+            <exception cref="T:System.Configuration.ConfigurationErrorsException">The element specified in <paramref name="elementName"/> is the &lt;clear&gt; element.
+            </exception><exception cref="T:System.ArgumentException"><paramref name="elementName"/> starts with the reserved prefix "config" or "lock".
+            </exception>
+        </member>
+        <member name="M:Microsoft.Practices.Unity.Configuration.InjectionMemberElementCollection.CreateNewElement">
+            <summary>
+            When overridden in a derived class, creates a new <see cref="T:System.Configuration.ConfigurationElement"/>.
+            </summary>
+            <returns>
+            A new <see cref="T:System.Configuration.ConfigurationElement"/>.
+            </returns>
+        </member>
+        <member name="M:Microsoft.Practices.Unity.Configuration.InjectionMemberElementCollection.GetElementKey(System.Configuration.ConfigurationElement)">
+            <summary>
+            Gets the element key for a specified configuration element when overridden in a derived class.
+            </summary>
+            <returns>
+            An <see cref="T:System.Object"/> that acts as the key for the specified <see cref="T:System.Configuration.ConfigurationElement"/>.
+            </returns>
+            <param name="element">The <see cref="T:System.Configuration.ConfigurationElement"/> to return the key for. </param>
+        </member>
+        <member name="P:Microsoft.Practices.Unity.Configuration.InjectionMemberElementCollection.Item(System.String)">
+            <summary>
+            Indexer that lets you access elements by their key.
+            </summary>
+            <param name="key">Key to retrieve element with.</param>
+            <returns>The element.</returns>
+        </member>
+        <member name="T:Microsoft.Practices.Unity.Configuration.InstanceElement">
+            <summary>
+            A configuration element that describes an instance to add to the container.
+            </summary>
+        </member>
+        <member name="M:Microsoft.Practices.Unity.Configuration.InstanceElement.SerializeContent(System.Xml.XmlWriter)">
+            <summary>
+            Write the contents of this element to the given <see cref="T:System.Xml.XmlWriter"/>.
+            </summary>
+            <remarks>The caller of this method has already written the start element tag before
+            calling this method, so deriving classes only need to write the element content, not
+            the start or end tags.</remarks>
+            <param name="writer">Writer to send XML content to.</param>
+        </member>
+        <member name="M:Microsoft.Practices.Unity.Configuration.InstanceElement.ConfigureContainer(Microsoft.Practices.Unity.IUnityContainer)">
+            <summary>
+            Add the instance defined by this element to the given container.
+            </summary>
+            <param name="container">Container to configure.</param>
+        </member>
+        <member name="P:Microsoft.Practices.Unity.Configuration.InstanceElement.Name">
+            <summary>
+            Name to register instance under
+            </summary>
+        </member>
+        <member name="P:Microsoft.Practices.Unity.Configuration.InstanceElement.Value">
+            <summary>
+            Value for this instance
+            </summary>
+        </member>
+        <member name="P:Microsoft.Practices.Unity.Configuration.InstanceElement.TypeName">
+            <summary>
+            Type of the instance. If not given, defaults to string
+            </summary>
+        </member>
+        <member name="P:Microsoft.Practices.Unity.Configuration.InstanceElement.TypeConverterTypeName">
+            <summary>
+            Type name for the type converter to use to create the instance. If not
+            given, defaults to the default type converter for this instance type.
+            </summary>
+        </member>
+        <member name="P:Microsoft.Practices.Unity.Configuration.InstanceElement.Key">
+            <summary>
+            Key used to keep these instances unique in the config collection.
+            </summary>
+        </member>
+        <member name="T:Microsoft.Practices.Unity.Configuration.InstanceElementCollection">
+            <summary>
+            A collection of <see cref="T:Microsoft.Practices.Unity.Configuration.InstanceElement"/>s.
+            </summary>
+        </member>
+        <member name="M:Microsoft.Practices.Unity.Configuration.InstanceElementCollection.GetElementKey(System.Configuration.ConfigurationElement)">
+            <summary>
+            Gets the element key for a specified configuration element when overridden in a derived class.
+            </summary>
+            <returns>
+            An <see cref="T:System.Object"/> that acts as the key for the specified <see cref="T:System.Configuration.ConfigurationElement"/>.
+            </returns>
+            <param name="element">The <see cref="T:System.Configuration.ConfigurationElement"/> to return the key for. </param>
+        </member>
+        <member name="T:Microsoft.Practices.Unity.Configuration.ConfigurationHelpers.IValueProvidingElement">
+            <summary>
+            An element that has a child Value property.
+            </summary>
+        </member>
+        <member name="P:Microsoft.Practices.Unity.Configuration.ConfigurationHelpers.IValueProvidingElement.Value">
+            <summary>
+            String that will be deserialized to provide the value.
+            </summary>
+        </member>
+        <member name="P:Microsoft.Practices.Unity.Configuration.ConfigurationHelpers.IValueProvidingElement.DestinationName">
+            <summary>
+            A string describing where the value this element contains
+            is being used. For example, if setting a property Prop1,
+            this should return "property Prop1" (in english).
+            </summary>
+        </member>
+        <member name="T:Microsoft.Practices.Unity.Configuration.LifetimeElement">
+            <summary>
+            A configuration element that represents lifetime managers.
+            </summary>
+        </member>
+        <member name="M:Microsoft.Practices.Unity.Configuration.LifetimeElement.CreateLifetimeManager">
+            <summary>
+            Create the <see cref="T:Microsoft.Practices.Unity.LifetimeManager"/> described by
+            this element.
+            </summary>
+            <returns>A <see cref="T:Microsoft.Practices.Unity.LifetimeManager"/> instance.</returns>
+        </member>
+        <member name="M:Microsoft.Practices.Unity.Configuration.LifetimeElement.SerializeContent(System.Xml.XmlWriter)">
+            <summary>
+            Write the contents of this element to the given <see cref="T:System.Xml.XmlWriter"/>.
+            </summary>
+            <remarks>The caller of this method has already written the start element tag before
+            calling this method, so deriving classes only need to write the element content, not
+            the start or end tags.</remarks>
+            <param name="writer">Writer to send XML content to.</param>
+        </member>
+        <member name="P:Microsoft.Practices.Unity.Configuration.LifetimeElement.TypeName">
+            <summary>
+            Type of the lifetime manager.
+            </summary>
+        </member>
+        <member name="P:Microsoft.Practices.Unity.Configuration.LifetimeElement.Value">
+            <summary>
+            Extra initialization information used by the type converter for this lifetime manager.
+            </summary>
+        </member>
+        <member name="P:Microsoft.Practices.Unity.Configuration.LifetimeElement.TypeConverterTypeName">
+            <summary>
+            Type of <see cref="T:System.ComponentModel.TypeConverter"/> to use to create the
+            lifetime manager.
+            </summary>
+        </member>
+        <member name="T:Microsoft.Practices.Unity.Configuration.MethodElement">
+            <summary>
+            A configuration element representing a method to call.
+            </summary>
+        </member>
+        <member name="M:Microsoft.Practices.Unity.Configuration.MethodElement.#ctor">
+            <summary>
+            Construct a new instance of <see cref="T:Microsoft.Practices.Unity.Configuration.MethodElement"/>.
+            </summary>
+        </member>
+        <member name="M:Microsoft.Practices.Unity.Configuration.MethodElement.SerializeContent(System.Xml.XmlWriter)">
+            <summary>
+            Write the contents of this element to the given <see cref="T:System.Xml.XmlWriter"/>.
+            </summary>
+            <remarks>The caller of this method has already written the start element tag before
+            calling this method, so deriving classes only need to write the element content, not
+            the start or end tags.</remarks>
+            <param name="writer">Writer to send XML content to.</param>
+        </member>
+        <member name="M:Microsoft.Practices.Unity.Configuration.MethodElement.GetInjectionMembers(Microsoft.Practices.Unity.IUnityContainer,System.Type,System.Type,System.String)">
+            <summary>
+            Return the set of <see cref="T:Microsoft.Practices.Unity.InjectionMember"/>s that are needed
+            to configure the container according to this configuration element.
+            </summary>
+            <param name="container">Container that is being configured.</param>
+            <param name="fromType">Type that is being registered.</param>
+            <param name="toType">Type that <paramref name="fromType"/> is being mapped to.</param>
+            <param name="name">Name this registration is under.</param>
+            <returns>One or more <see cref="T:Microsoft.Practices.Unity.InjectionMember"/> objects that should be
+            applied to the container registration.</returns>
+        </member>
+        <member name="P:Microsoft.Practices.Unity.Configuration.MethodElement.Name">
+            <summary>
+            Name of the method to call.
+            </summary>
+        </member>
+        <member name="P:Microsoft.Practices.Unity.Configuration.MethodElement.Parameters">
+            <summary>
+            Parameters to the method call.
+            </summary>
+        </member>
+        <member name="P:Microsoft.Practices.Unity.Configuration.MethodElement.Key">
+            <summary>
+            Each element must have a unique key, which is generated by the subclasses.
+            </summary>
+        </member>
+        <member name="P:Microsoft.Practices.Unity.Configuration.MethodElement.ElementName">
+            <summary>
+            Element name to use to serialize this into XML.
+            </summary>
+        </member>
+        <member name="T:Microsoft.Practices.Unity.Configuration.NamespaceElement">
+            <summary>
+            A configuration element representing the namespace
+            tags in the config file.
+            </summary>
+        </member>
+        <member name="T:Microsoft.Practices.Unity.Configuration.NamespaceElementCollection">
+            <summary>
+            A collection of <see cref="T:Microsoft.Practices.Unity.Configuration.NamespaceElement"/>s in configuration.
+            </summary>
+        </member>
+        <member name="M:Microsoft.Practices.Unity.Configuration.NamespaceElementCollection.GetElementKey(System.Configuration.ConfigurationElement)">
+            <summary>
+            Gets the element key for a specified configuration element when overridden in a derived class.
+            </summary>
+            <returns>
+            An <see cref="T:System.Object"/> that acts as the key for the specified <see cref="T:System.Configuration.ConfigurationElement"/>.
+            </returns>
+            <param name="element">The <see cref="T:System.Configuration.ConfigurationElement"/> to return the key for</param>
+        </member>
+        <member name="T:Microsoft.Practices.Unity.Configuration.OptionalElement">
+            <summary>
+            A configuration element that specifies that a value
+            is optional.
+            </summary>
+        </member>
+        <member name="M:Microsoft.Practices.Unity.Configuration.OptionalElement.SerializeContent(System.Xml.XmlWriter)">
+            <summary>
+            Write the contents of this element to the given <see cref="T:System.Xml.XmlWriter"/>.
+            </summary>
+            <remarks>The caller of this method has already written the start element tag before
+            calling this method, so deriving classes only need to write the element content, not
+            the start or end tags.</remarks>
+            <param name="writer">Writer to send XML content to.</param>
+        </member>
+        <member name="M:Microsoft.Practices.Unity.Configuration.OptionalElement.GetInjectionParameterValue(Microsoft.Practices.Unity.IUnityContainer,System.Type)">
+            <summary>
+            Generate an <see cref="T:Microsoft.Practices.Unity.InjectionParameterValue"/> object
+            that will be used to configure the container for a type registration.
+            </summary>
+            <param name="container">Container that is being configured. Supplied in order
+            to let custom implementations retrieve services; do not configure the container
+            directly in this method.</param>
+            <param name="parameterType">Type of the </param>
+            <returns></returns>
+        </member>
+        <member name="P:Microsoft.Practices.Unity.Configuration.OptionalElement.Name">
+            <summary>
+            Name used to resolve the dependency, leave out or blank to resolve default.
+            </summary>
+        </member>
+        <member name="P:Microsoft.Practices.Unity.Configuration.OptionalElement.TypeName">
+            <summary>
+            Type of dependency to resolve. If left out, resolved the type of
+            the containing parameter or property.
+            </summary>
+        </member>
+        <member name="T:Microsoft.Practices.Unity.Configuration.ParameterElement">
+            <summary>
+            Configuration element representing a parameter passed to a constructor
+            or a method.
+            </summary>
+        </member>
+        <member name="M:Microsoft.Practices.Unity.Configuration.ParameterElement.#ctor">
+            <summary>
+            Construct a new instance of <see cref="T:Microsoft.Practices.Unity.Configuration.ParameterElement"/>.
+            </summary>
+        </member>
+        <member name="M:Microsoft.Practices.Unity.Configuration.ParameterElement.GetParameterValue(Microsoft.Practices.Unity.IUnityContainer,System.Type)">
+            <summary>
+            Returns the required <see cref="T:Microsoft.Practices.Unity.InjectionParameterValue"/> needed
+            to configure the container so that the correct value is injected.
+            </summary>
+            <param name="container">Container being configured.</param>
+            <param name="parameterType">Type of the parameter.</param>
+            <returns>The value to use to configure the container.</returns>
+        </member>
+        <member name="M:Microsoft.Practices.Unity.Configuration.ParameterElement.Matches(System.Reflection.ParameterInfo)">
+            <summary>
+            Does the information in this <see cref="T:Microsoft.Practices.Unity.Configuration.ParameterElement"/> match
+            up with the given <paramref name="parameterInfo"/>?
+            </summary>
+            <param name="parameterInfo">Information about the parameter.</param>
+            <returns>True if this is a match, false if not.</returns>
+        </member>
+        <member name="M:Microsoft.Practices.Unity.Configuration.ParameterElement.DeserializeElement(System.Xml.XmlReader,System.Boolean)">
+            <summary>
+            Reads XML from the configuration file.
+            </summary>
+            <param name="reader">The <see cref="T:System.Xml.XmlReader"/> that reads from the configuration file.</param>
+            <param name="serializeCollectionKey">true to serialize only the collection key properties; otherwise, false.</param>
+            <exception cref="T:System.Configuration.ConfigurationErrorsException">The element to read is locked.
+            - or -
+            An attribute of the current node is not recognized.
+            - or -
+            The lock status of the current node cannot be determined.  
+            </exception>
+        </member>
+        <member name="M:Microsoft.Practices.Unity.Configuration.ParameterElement.OnDeserializeUnrecognizedAttribute(System.String,System.String)">
+            <summary>
+            Gets a value indicating whether an unknown attribute is encountered during deserialization.
+            </summary>
+            <returns>
+            true when an unknown attribute is encountered while deserializing; otherwise, false.
+            </returns>
+            <param name="name">The name of the unrecognized attribute.</param>
+            <param name="value">The value of the unrecognized attribute.</param>
+        </member>
+        <member name="M:Microsoft.Practices.Unity.Configuration.ParameterElement.OnDeserializeUnrecognizedElement(System.String,System.Xml.XmlReader)">
+            <summary>
+            Gets a value indicating whether an unknown element is encountered during deserialization.
+            </summary>
+            <returns>
+            true when an unknown element is encountered while deserializing; otherwise, false.
+            </returns>
+            <param name="elementName">The name of the unknown subelement.</param>
+            <param name="reader">The <see cref="T:System.Xml.XmlReader"/> being used for deserialization.</param>
+            <exception cref="T:System.Configuration.ConfigurationErrorsException">The element identified by <paramref name="elementName"/> is locked.
+            - or -
+            One or more of the element's attributes is locked.
+            - or -
+            <paramref name="elementName"/> is unrecognized, or the element has an unrecognized attribute.
+            - or -
+            The element has a Boolean attribute with an invalid value.
+            - or -
+            An attempt was made to deserialize a property more than once.
+            - or -
+            An attempt was made to deserialize a property that is not a valid member of the element.
+            - or -
+            The element cannot contain a CDATA or text element.
+            </exception>
+        </member>
+        <member name="M:Microsoft.Practices.Unity.Configuration.ParameterElement.SerializeContent(System.Xml.XmlWriter)">
+            <summary>
+            Write the contents of this element to the given <see cref="T:System.Xml.XmlWriter"/>.
+            </summary>
+            <remarks>The caller of this method has already written the start element tag before
+            calling this method, so deriving classes only need to write the element content, not
+            the start or end tags.</remarks>
+            <param name="writer">Writer to send XML content to.</param>
+        </member>
+        <member name="P:Microsoft.Practices.Unity.Configuration.ParameterElement.Name">
+            <summary>
+            Name of this parameter.
+            </summary>
+        </member>
+        <member name="P:Microsoft.Practices.Unity.Configuration.ParameterElement.TypeName">
+            <summary>
+            Type of this parameter.
+            </summary>
+            <remarks>This is only needed in order to disambiguate method overloads. Normally
+            the parameter name is sufficient.</remarks>
+        </member>
+        <member name="P:Microsoft.Practices.Unity.Configuration.ParameterElement.Value">
+            <summary>
+            Element that describes the value for this property.
+            </summary>
+            <remarks>
+            This is NOT marked as a ConfigurationProperty because this
+            child element is polymorphic, and the element tag determines
+            the type. Standard configuration properties only let you do
+            this if it's a collection, but we only want one value. Thus
+            the separate property. The element is deserialized in 
+            <see cref="M:Microsoft.Practices.Unity.Configuration.ParameterElement.OnDeserializeUnrecognizedElement(System.String,System.Xml.XmlReader)"/>.</remarks>
+        </member>
+        <member name="P:Microsoft.Practices.Unity.Configuration.ParameterElement.DestinationName">
+            <summary>
+            A string describing where the value this element contains
+            is being used. For example, if setting a property Prop1,
+            this should return "property Prop1" (in english).
+            </summary>
+        </member>
+        <member name="T:Microsoft.Practices.Unity.Configuration.ParameterElementCollection">
+            <summary>
+            A collection of <see cref="T:Microsoft.Practices.Unity.Configuration.ParameterElement"/> objects.
+            </summary>
+        </member>
+        <member name="M:Microsoft.Practices.Unity.Configuration.ParameterElementCollection.GetElementKey(System.Configuration.ConfigurationElement)">
+            <summary>
+            Gets the element key for a specified configuration element when overridden in a derived class.
+            </summary>
+            <returns>
+            An <see cref="T:System.Object"/> that acts as the key for the specified <see cref="T:System.Configuration.ConfigurationElement"/>.
+            </returns>
+            <param name="element">The <see cref="T:System.Configuration.ConfigurationElement"/> to return the key for. </param>
+        </member>
+        <member name="T:Microsoft.Practices.Unity.Configuration.ParameterValueElementCollection">
+            <summary>
+            A collection of <see cref="T:Microsoft.Practices.Unity.Configuration.ParameterValueElement"/> objects.
+            </summary>
+        </member>
+        <member name="M:Microsoft.Practices.Unity.Configuration.ParameterValueElementCollection.OnDeserializeUnrecognizedElement(System.String,System.Xml.XmlReader)">
+            <summary>
+            Causes the configuration system to throw an exception.
+            </summary>
+            <returns>
+            true if the unrecognized element was deserialized successfully; otherwise, false. The default is false.
+            </returns>
+            <param name="elementName">The name of the unrecognized element. </param>
+            <param name="reader">An input stream that reads XML from the configuration file. </param>
+            <exception cref="T:System.Configuration.ConfigurationErrorsException">The element specified in <paramref name="elementName"/> is the &lt;clear&gt; element.</exception>
+            <exception cref="T:System.ArgumentException"><paramref name="elementName"/> starts with the reserved prefix "config" or "lock".</exception>
+        </member>
+        <member name="M:Microsoft.Practices.Unity.Configuration.ParameterValueElementCollection.CreateNewElement">
+            <summary>
+            When overridden in a derived class, creates a new <see cref="T:System.Configuration.ConfigurationElement"/>.
+            </summary>
+            <returns>
+            A new <see cref="T:System.Configuration.ConfigurationElement"/>.
+            </returns>
+        </member>
+        <member name="M:Microsoft.Practices.Unity.Configuration.ParameterValueElementCollection.GetElementKey(System.Configuration.ConfigurationElement)">
+            <summary>
+            Gets the element key for a specified configuration element when overridden in a derived class.
+            </summary>
+            <returns>
+            An <see cref="T:System.Object"/> that acts as the key for the specified <see cref="T:System.Configuration.ConfigurationElement"/>.
+            </returns>
+            <param name="element">The <see cref="T:System.Configuration.ConfigurationElement"/> to return the key for. </param>
+        </member>
+        <member name="P:Microsoft.Practices.Unity.Configuration.ParameterValueElementCollection.DeserializedElementHolder.Value">
+            <summary>
+            String that will be deserialized to provide the value.
+            </summary>
+        </member>
+        <member name="P:Microsoft.Practices.Unity.Configuration.ParameterValueElementCollection.DeserializedElementHolder.DestinationName">
+            <summary>
+            A string describing where the value this element contains
+            is being used. For example, if setting a property Prop1,
+            this should return "property Prop1" (in english).
+            </summary>
+        </member>
+        <member name="T:Microsoft.Practices.Unity.Configuration.Properties.Resources">
+            <summary>
+              A strongly-typed resource class, for looking up localized strings, etc.
+            </summary>
+        </member>
+        <member name="P:Microsoft.Practices.Unity.Configuration.Properties.Resources.ResourceManager">
+            <summary>
+              Returns the cached ResourceManager instance used by this class.
+            </summary>
+        </member>
+        <member name="P:Microsoft.Practices.Unity.Configuration.Properties.Resources.Culture">
+            <summary>
+              Overrides the current thread's CurrentUICulture property for all
+              resource lookups using this strongly typed resource class.
+            </summary>
+        </member>
+        <member name="P:Microsoft.Practices.Unity.Configuration.Properties.Resources.CannotCreateContainerConfiguringElement">
+            <summary>
+              Looks up a localized string similar to An abstract ContainerConfiguringElement cannot be created. Please specify a concrete type..
+            </summary>
+        </member>
+        <member name="P:Microsoft.Practices.Unity.Configuration.Properties.Resources.CannotCreateExtensionConfigurationElement">
+            <summary>
+              Looks up a localized string similar to An abstract ExtensionConfigurationElement object cannot be created. Please specify a concrete type..
+            </summary>
+        </member>
+        <member name="P:Microsoft.Practices.Unity.Configuration.Properties.Resources.CannotCreateInjectionMemberElement">
+            <summary>
+              Looks up a localized string similar to An abstract InjectionMemberElement object cannot be created. Please specify a concrete type..
+            </summary>
+        </member>
+        <member name="P:Microsoft.Practices.Unity.Configuration.Properties.Resources.CannotCreateParameterValueElement">
+            <summary>
+              Looks up a localized string similar to An abstract ParameterValueElement object cannot be created. Please specify a concrete type..
+            </summary>
+        </member>
+        <member name="P:Microsoft.Practices.Unity.Configuration.Properties.Resources.CouldNotResolveType">
+            <summary>
+              Looks up a localized string similar to The type name or alias {0} could not be resolved. Please check your configuration file and verify this type name..
+            </summary>
+        </member>
+        <member name="P:Microsoft.Practices.Unity.Configuration.Properties.Resources.DependencyForGenericParameterWithTypeSet">
+            <summary>
+              Looks up a localized string similar to The dependency element for generic parameter {0} must not have an explicit type name but has &apos;{1}&apos;..
+            </summary>
+        </member>
+        <member name="P:Microsoft.Practices.Unity.Configuration.Properties.Resources.DependencyForOptionalGenericParameterWithTypeSet">
+            <summary>
+              Looks up a localized string similar to The optional dependency element for generic parameter {0} must not have an explicit type name but has &apos;{1}&apos;..
+            </summary>
+        </member>
+        <member name="P:Microsoft.Practices.Unity.Configuration.Properties.Resources.DestinationNameFormat">
+            <summary>
+              Looks up a localized string similar to {0} {1}.
+            </summary>
+        </member>
+        <member name="P:Microsoft.Practices.Unity.Configuration.Properties.Resources.DuplicateParameterValueElement">
+            <summary>
+              Looks up a localized string similar to The injection configuration for {0} has multiple values..
+            </summary>
+        </member>
+        <member name="P:Microsoft.Practices.Unity.Configuration.Properties.Resources.ElementTypeNotRegistered">
+            <summary>
+              Looks up a localized string similar to The configuration element type {0} has not been registered with the section..
+            </summary>
+        </member>
+        <member name="P:Microsoft.Practices.Unity.Configuration.Properties.Resources.ElementWithAttributesAndParameterValueElements">
+            <summary>
+              Looks up a localized string similar to The injection configuration for {0} is specified through both attributes and child value elements..
+            </summary>
+        </member>
+        <member name="P:Microsoft.Practices.Unity.Configuration.Properties.Resources.ExtensionTypeNotFound">
+            <summary>
+              Looks up a localized string similar to Could not load section extension type {0}..
+            </summary>
+        </member>
+        <member name="P:Microsoft.Practices.Unity.Configuration.Properties.Resources.ExtensionTypeNotValid">
+            <summary>
+              Looks up a localized string similar to The extension type {0} does not derive from SectionExtension..
+            </summary>
+        </member>
+        <member name="P:Microsoft.Practices.Unity.Configuration.Properties.Resources.InvalidExtensionElementType">
+            <summary>
+              Looks up a localized string similar to The extension element type {0} that is being added does not derive from ContainerConfiguringElement, InjectionMemberElement, or ParameterValueElement. An extension element must derive from one of these types..
+            </summary>
+        </member>
+        <member name="P:Microsoft.Practices.Unity.Configuration.Properties.Resources.InvalidValueAttributes">
+            <summary>
+              Looks up a localized string similar to No valid attributes were found to construct the value for the {0}. Please check the configuration file..
+            </summary>
+        </member>
+        <member name="P:Microsoft.Practices.Unity.Configuration.Properties.Resources.NoMatchingConstructor">
+            <summary>
+              Looks up a localized string similar to Configuration is incorrect, the type {0} does not have a constructor that takes parameters named {1}..
+            </summary>
+        </member>
+        <member name="P:Microsoft.Practices.Unity.Configuration.Properties.Resources.NoMatchingMethod">
+            <summary>
+              Looks up a localized string similar to Configuration is incorrect, the type {0} does not have a method named {1} that takes parameters named {2}..
+            </summary>
+        </member>
+        <member name="P:Microsoft.Practices.Unity.Configuration.Properties.Resources.NoSuchContainer">
+            <summary>
+              Looks up a localized string similar to The container named &quot;{0}&quot; is not defined in this configuration section..
+            </summary>
+        </member>
+        <member name="P:Microsoft.Practices.Unity.Configuration.Properties.Resources.NoSuchProperty">
+            <summary>
+              Looks up a localized string similar to The type {0} does not have a property named {1}..
+            </summary>
+        </member>
+        <member name="P:Microsoft.Practices.Unity.Configuration.Properties.Resources.NotAnArray">
+            <summary>
+              Looks up a localized string similar to The configuration is set to inject an array, but the type {0} is not an array type..
+            </summary>
+        </member>
+        <member name="P:Microsoft.Practices.Unity.Configuration.Properties.Resources.Parameter">
+            <summary>
+              Looks up a localized string similar to parameter.
+            </summary>
+        </member>
+        <member name="P:Microsoft.Practices.Unity.Configuration.Properties.Resources.Property">
+            <summary>
+              Looks up a localized string similar to property.
+            </summary>
+        </member>
+        <member name="P:Microsoft.Practices.Unity.Configuration.Properties.Resources.RequiredPropertyMissing">
+            <summary>
+              Looks up a localized string similar to The attribute {0} must be present and non-empty..
+            </summary>
+        </member>
+        <member name="P:Microsoft.Practices.Unity.Configuration.Properties.Resources.ValueNotAllowedForGenericArrayType">
+            <summary>
+              Looks up a localized string similar to The value element for {1} was specified for the generic array type {0}. Value elements are not allowed for generic array types..
+            </summary>
+        </member>
+        <member name="P:Microsoft.Practices.Unity.Configuration.Properties.Resources.ValueNotAllowedForGenericParameterType">
+            <summary>
+              Looks up a localized string similar to The value element for {1} was specified for the generic parameter type {0}. Value elements are not allowed for generic parameter types..
+            </summary>
+        </member>
+        <member name="P:Microsoft.Practices.Unity.Configuration.Properties.Resources.ValueNotAllowedForOpenGenericType">
+            <summary>
+              Looks up a localized string similar to The value element for {1} was specified for the generic type {0}. Value elements are not allowed for generic types..
+            </summary>
+        </member>
+        <member name="T:Microsoft.Practices.Unity.Configuration.PropertyElement">
+            <summary>
+            A class representing a property configuration element.
+            </summary>
+        </member>
+        <member name="M:Microsoft.Practices.Unity.Configuration.PropertyElement.#ctor">
+            <summary>
+            Construct a new instance of <see cref="T:Microsoft.Practices.Unity.Configuration.PropertyElement"/>
+            </summary>
+        </member>
+        <member name="M:Microsoft.Practices.Unity.Configuration.PropertyElement.DeserializeElement(System.Xml.XmlReader,System.Boolean)">
+            <summary>
+            Reads XML from the configuration file.
+            </summary>
+            <param name="reader">The <see cref="T:System.Xml.XmlReader"/> that reads from the configuration file.</param>
+            <param name="serializeCollectionKey">true to serialize only the collection key properties; otherwise, false.</param>
+            <exception cref="T:System.Configuration.ConfigurationErrorsException">The element to read is locked.
+            - or -
+            An attribute of the current node is not recognized.
+            - or -
+            The lock status of the current node cannot be determined.  
+            </exception>
+        </member>
+        <member name="M:Microsoft.Practices.Unity.Configuration.PropertyElement.OnDeserializeUnrecognizedAttribute(System.String,System.String)">
+            <summary>
+            Gets a value indicating whether an unknown attribute is encountered during deserialization.
+            </summary>
+            <returns>
+            true when an unknown attribute is encountered while deserializing; otherwise, false.
+            </returns>
+            <param name="name">The name of the unrecognized attribute.</param>
+            <param name="value">The value of the unrecognized attribute.</param>
+        </member>
+        <member name="M:Microsoft.Practices.Unity.Configuration.PropertyElement.OnDeserializeUnrecognizedElement(System.String,System.Xml.XmlReader)">
+            <summary>
+            Gets a value indicating whether an unknown element is encountered during deserialization.
+            </summary>
+            <returns>
+            true when an unknown element is encountered while deserializing; otherwise, false.
+            </returns>
+            <param name="elementName">The name of the unknown subelement.</param>
+            <param name="reader">The <see cref="T:System.Xml.XmlReader"/> being used for deserialization.</param>
+            <exception cref="T:System.Configuration.ConfigurationErrorsException">The element identified by <paramref name="elementName"/> is locked.
+            - or -
+            One or more of the element's attributes is locked.
+            - or -
+            <paramref name="elementName"/> is unrecognized, or the element has an unrecognized attribute.
+            - or -
+            The element has a Boolean attribute with an invalid value.
+            - or -
+            An attempt was made to deserialize a property more than once.
+            - or -
+            An attempt was made to deserialize a property that is not a valid member of the element.
+            - or -
+            The element cannot contain a CDATA or text element.
+            </exception>
+        </member>
+        <member name="M:Microsoft.Practices.Unity.Configuration.PropertyElement.SerializeContent(System.Xml.XmlWriter)">
+            <summary>
+            Write the contents of this element to the given <see cref="T:System.Xml.XmlWriter"/>.
+            </summary>
+            <remarks>The caller of this method has already written the start element tag before
+            calling this method, so deriving classes only need to write the element content, not
+            the start or end tags.</remarks>
+            <param name="writer">Writer to send XML content to.</param>
+        </member>
+        <member name="M:Microsoft.Practices.Unity.Configuration.PropertyElement.GetInjectionMembers(Microsoft.Practices.Unity.IUnityContainer,System.Type,System.Type,System.String)">
+            <summary>
+            Return the set of <see cref="T:Microsoft.Practices.Unity.InjectionMember"/>s that are needed
+            to configure the container according to this configuration element.
+            </summary>
+            <param name="container">Container that is being configured.</param>
+            <param name="fromType">Type that is being registered.</param>
+            <param name="toType">Type that <paramref name="fromType"/> is being mapped to.</param>
+            <param name="name">Name this registration is under.</param>
+            <returns>One or more <see cref="T:Microsoft.Practices.Unity.InjectionMember"/> objects that should be
+            applied to the container registration.</returns>
+        </member>
+        <member name="P:Microsoft.Practices.Unity.Configuration.PropertyElement.Name">
+            <summary>
+            Name of the property that will be set.
+            </summary>
+        </member>
+        <member name="P:Microsoft.Practices.Unity.Configuration.PropertyElement.Key">
+            <summary>
+            Each element must have a unique key, which is generated by the subclasses.
+            </summary>
+        </member>
+        <member name="P:Microsoft.Practices.Unity.Configuration.PropertyElement.Value">
+            <summary>
+            String that will be deserialized to provide the value.
+            </summary>
+        </member>
+        <member name="P:Microsoft.Practices.Unity.Configuration.PropertyElement.DestinationName">
+            <summary>
+            A string describing where the value this element contains
+            is being used. For example, if setting a property Prop1,
+            this should return "property Prop1" (in english).
+            </summary>
+        </member>
+        <member name="P:Microsoft.Practices.Unity.Configuration.PropertyElement.ElementName">
+            <summary>
+            Element name to use to serialize this into XML.
+            </summary>
+        </member>
+        <member name="T:Microsoft.Practices.Unity.Configuration.RegisterElement">
+            <summary>
+            A configuration element representing a single container type registration.
+            </summary>
+        </member>
+        <member name="M:Microsoft.Practices.Unity.Configuration.RegisterElement.ConfigureContainer(Microsoft.Practices.Unity.IUnityContainer)">
+            <summary>
+            Apply the registrations from this element to the given container.
+            </summary>
+            <param name="container">Container to configure.</param>
+        </member>
+        <member name="M:Microsoft.Practices.Unity.Configuration.RegisterElement.SerializeContent(System.Xml.XmlWriter)">
+            <summary>
+            Write the contents of this element to the given <see cref="T:System.Xml.XmlWriter"/>.
+            </summary>
+            <remarks>The caller of this method has already written the start element tag before
+            calling this method, so deriving classes only need to write the element content, not
+            the start or end tags.</remarks>
+            <param name="writer">Writer to send XML content to.</param>
+        </member>
+        <member name="P:Microsoft.Practices.Unity.Configuration.RegisterElement.TypeName">
+            <summary>
+            The type that is registered.
+            </summary>
+        </member>
+        <member name="P:Microsoft.Practices.Unity.Configuration.RegisterElement.Name">
+            <summary>
+            Name registered under.
+            </summary>
+        </member>
+        <member name="P:Microsoft.Practices.Unity.Configuration.RegisterElement.MapToName">
+            <summary>
+            Type that is mapped to.
+            </summary>
+        </member>
+        <member name="P:Microsoft.Practices.Unity.Configuration.RegisterElement.Lifetime">
+            <summary>
+            Lifetime manager to register for this registration.
+            </summary>
+        </member>
+        <member name="P:Microsoft.Practices.Unity.Configuration.RegisterElement.InjectionMembers">
+            <summary>
+            Any injection members (constructor, properties, etc.) that are specified for
+            this registration.
+            </summary>
+        </member>
+        <member name="T:Microsoft.Practices.Unity.Configuration.RegisterElementCollection">
+            <summary>
+            A collection of <see cref="T:Microsoft.Practices.Unity.Configuration.RegisterElement"/>s.
+            </summary>
+        </member>
+        <member name="M:Microsoft.Practices.Unity.Configuration.RegisterElementCollection.OnDeserializeUnrecognizedElement(System.String,System.Xml.XmlReader)">
+            <summary>
+            Causes the configuration system to throw an exception.
+            </summary>
+            <returns>
+            true if the unrecognized element was deserialized successfully; otherwise, false. The default is false.
+            </returns>
+            <param name="elementName">The name of the unrecognized element. </param>
+            <param name="reader">An input stream that reads XML from the configuration file. </param>
+            <exception cref="T:System.Configuration.ConfigurationErrorsException">The element specified in <paramref name="elementName"/> is the &lt;clear&gt; element.</exception>
+            <exception cref="T:System.ArgumentException"><paramref name="elementName"/> starts with the reserved prefix "config" or "lock".</exception>
+        </member>
+        <member name="M:Microsoft.Practices.Unity.Configuration.RegisterElementCollection.GetElementKey(System.Configuration.ConfigurationElement)">
+            <summary>
+            Gets the element key for a specified configuration element when overridden in a derived class.
+            </summary>
+            <returns>
+            An <see cref="T:System.Object"/> that acts as the key for the specified <see cref="T:System.Configuration.ConfigurationElement"/>.
+            </returns>
+            <param name="element">The <see cref="T:System.Configuration.ConfigurationElement"/> to return the key for. </param>
+        </member>
+        <member name="T:Microsoft.Practices.Unity.Configuration.SectionExtension">
+            <summary>
+            Base class for Unity configuration section extensions.
+            Derived classes are used to add custom elements and aliases
+            into the configuration section being loaded.
+            </summary>
+        </member>
+        <member name="M:Microsoft.Practices.Unity.Configuration.SectionExtension.AddExtensions(Microsoft.Practices.Unity.Configuration.SectionExtensionContext)">
+            <summary>
+            Add the extensions to the section via the context.
+            </summary>
+            <param name="context">Context object that can be used to add elements and aliases.</param>
+        </member>
+        <member name="T:Microsoft.Practices.Unity.Configuration.SectionExtensionContext">
+            <summary>
+            An object that gives the ability to add
+            elements and aliases to a configuration section.
+            </summary>
+        </member>
+        <member name="M:Microsoft.Practices.Unity.Configuration.SectionExtensionContext.AddAlias(System.String,System.Type)">
+            <summary>
+            Add a new alias to the configuration section. This is useful
+            for those extensions that add commonly used types to configuration
+            so users don't have to alias them repeatedly.
+            </summary>
+            <param name="newAlias">The alias to use.</param>
+            <param name="aliasedType">Type the alias maps to.</param>
+        </member>
+        <member name="M:Microsoft.Practices.Unity.Configuration.SectionExtensionContext.AddAlias``1(System.String)">
+            <summary>
+            Add a new alias to the configuration section. This is useful
+            for those extensions that add commonly used types to configuration
+            so users don't have to alias them repeatedly.
+            </summary>
+            <typeparam name="TAliased">Type the alias maps to.</typeparam>
+            <param name="alias">The alias to use</param>
+        </member>
+        <member name="M:Microsoft.Practices.Unity.Configuration.SectionExtensionContext.AddElement(System.String,System.Type)">
+            <summary>
+            Add a new element to the configuration section schema.
+            </summary>
+            <param name="tag">Tag name in the XML.</param>
+            <param name="elementType">Type the tag maps to.</param>
+        </member>
+        <member name="M:Microsoft.Practices.Unity.Configuration.SectionExtensionContext.AddElement``1(System.String)">
+            <summary>
+            Add a new element to the configuration section schema.
+            </summary>
+            <typeparam name="TElement">Type the tag maps to.</typeparam>
+            <param name="tag">Tag name in the XML.</param>
+        </member>
+        <member name="T:Microsoft.Practices.Unity.Configuration.SectionExtensionElement">
+            <summary>
+            A configuration element used to specify which extensions to
+            add to the configuration schema.
+            </summary>
+        </member>
+        <member name="M:Microsoft.Practices.Unity.Configuration.SectionExtensionElement.DeserializeElement(System.Xml.XmlReader,System.Boolean)">
+            <summary>
+            Reads XML from the configuration file.
+            </summary>
+            <param name="reader">The <see cref="T:System.Xml.XmlReader"/> that reads from the configuration file.</param>
+            <param name="serializeCollectionKey">true to serialize only the collection key properties; otherwise, false.</param>
+            <exception cref="T:System.Configuration.ConfigurationErrorsException">The element to read is locked.
+            - or -
+            An attribute of the current node is not recognized.
+            - or -
+            The lock status of the current node cannot be determined.  
+            </exception>
+        </member>
+        <member name="M:Microsoft.Practices.Unity.Configuration.SectionExtensionElement.SerializeContent(System.Xml.XmlWriter)">
+            <summary>
+            Write the contents of this element to the given <see cref="T:System.Xml.XmlWriter"/>.
+            </summary>
+            <remarks>The caller of this method has already written the start element tag before
+            calling this method, so deriving classes only need to write the element content, not
+            the start or end tags.</remarks>
+            <param name="writer">Writer to send XML content to.</param>
+        </member>
+        <member name="P:Microsoft.Practices.Unity.Configuration.SectionExtensionElement.TypeName">
+            <summary>
+            Type of the section extender object that will provide new elements to the schema.
+            </summary>
+        </member>
+        <member name="P:Microsoft.Practices.Unity.Configuration.SectionExtensionElement.Prefix">
+            <summary>
+            Optional prefix that will be added to the element names added by this
+            section extender. If left out, no prefix will be added.
+            </summary>
+        </member>
+        <member name="P:Microsoft.Practices.Unity.Configuration.SectionExtensionElement.ExtensionObject">
+            <summary>
+            The extension object represented by this element.
+            </summary>
+        </member>
+        <member name="T:Microsoft.Practices.Unity.Configuration.SectionExtensionElementCollection">
+            <summary>
+            A collection of <see cref="T:Microsoft.Practices.Unity.Configuration.SectionExtensionElement"/>s.
+            </summary>
+        </member>
+        <member name="M:Microsoft.Practices.Unity.Configuration.SectionExtensionElementCollection.GetElementKey(System.Configuration.ConfigurationElement)">
+            <summary>
+            Gets the element key for a specified configuration element when overridden in a derived class.
+            </summary>
+            <returns>
+            An <see cref="T:System.Object"/> that acts as the key for the specified <see cref="T:System.Configuration.ConfigurationElement"/>.
+            </returns>
+            <param name="element">The <see cref="T:System.Configuration.ConfigurationElement"/> to return the key for. </param>
+        </member>
+        <member name="T:Microsoft.Practices.Unity.Configuration.ConfigurationHelpers.TypeResolver">
+            <summary>
+            Type that manage access to a set of type aliases and implements
+            the logic for converting aliases to their actual types.
+            </summary>
+        </member>
+        <member name="M:Microsoft.Practices.Unity.Configuration.ConfigurationHelpers.TypeResolver.SetAliases(Microsoft.Practices.Unity.Configuration.UnityConfigurationSection)">
+            <summary>
+            Set the set of aliases to use for resolution.
+            </summary>
+            <param name="section">Configuration section containing the various
+            type aliases, namespaces and assemblies.</param>
+        </member>
+        <member name="M:Microsoft.Practices.Unity.Configuration.ConfigurationHelpers.TypeResolver.ResolveType(System.String)">
+            <summary>
+            Resolves a type alias or type FullName to a concrete type.
+            </summary>
+            <param name="typeNameOrAlias">Type alias or type FullName</param>
+            <returns>Type object or null if resolve fails.</returns>
+            <exception cref="T:System.InvalidOperationException">Thrown if alias lookup fails.</exception>
+        </member>
+        <member name="M:Microsoft.Practices.Unity.Configuration.ConfigurationHelpers.TypeResolver.ResolveType(System.String,System.Boolean)">
+            <summary>
+            Resolves a type alias or type FullName to a concrete type.
+            </summary>
+            <param name="typeNameOrAlias">Alias or name to resolve.</param>
+            <param name="throwIfResolveFails">if true and the alias does not
+            resolve, throw an <see cref="T:System.InvalidOperationException"/>, otherwise 
+            return null on failure.</param>
+            <returns>The type object or null if resolve fails and 
+            <paramref name="throwIfResolveFails"/> is false.</returns>
+        </member>
+        <member name="M:Microsoft.Practices.Unity.Configuration.ConfigurationHelpers.TypeResolver.ResolveTypeWithDefault(System.String,System.Type)">
+            <summary>
+            Resolve a type alias or type full name to a concrete type.
+            If <paramref name="typeNameOrAlias"/> is null or empty, return the
+            given <paramref name="defaultValue"/> instead.
+            </summary>
+            <param name="typeNameOrAlias">Type alias or full name to resolve.</param>
+            <param name="defaultValue">Value to return if typeName is null or empty.</param>
+            <returns>The concrete <see cref="T:System.Type"/>.</returns>
+            <exception cref="T:System.InvalidOperationException">Thrown if alias lookup fails.</exception>
+        </member>
+        <member name="M:Microsoft.Practices.Unity.Configuration.ConfigurationHelpers.TypeResolver.ResolveTypeWithDefault(System.String,System.Type,System.Boolean)">
+            <summary>
+            Resolve a type alias or type full name to a concrete type.
+            If <paramref name="typeNameOrAlias"/> is null or empty, return the
+            given <paramref name="defaultValue"/> instead.
+            </summary>
+            <param name="typeNameOrAlias">Type alias or full name to resolve.</param>
+            <param name="defaultValue">Value to return if typeName is null or empty.</param>
+            <param name="throwIfResolveFails">if true and the alias does not
+            resolve, throw an <see cref="T:System.InvalidOperationException"/>, otherwise 
+            return null on failure.</param>
+            <returns>The concrete <see cref="T:System.Type"/>.</returns>
+            <exception cref="T:System.InvalidOperationException">Thrown if alias lookup fails and
+            <paramref name="throwIfResolveFails"/> is true.</exception>
+        </member>
+        <member name="T:Microsoft.Practices.Unity.Configuration.ConfigurationHelpers.TypeResolverImpl">
+            <summary>
+            A helper class that implements the actual logic for resolving a shorthand
+            type name (alias or raw type name) into an actual type object.
+            </summary>
+        </member>
+        <member name="M:Microsoft.Practices.Unity.Configuration.ConfigurationHelpers.TypeResolverImpl.#ctor(System.Collections.Generic.IEnumerable{System.Collections.Generic.KeyValuePair{System.String,System.String}},System.Collections.Generic.IEnumerable{System.String},System.Collections.Generic.IEnumerable{System.String})">
+            <summary>
+            Construct a new <see cref="T:Microsoft.Practices.Unity.Configuration.ConfigurationHelpers.TypeResolverImpl"/> that uses the given
+            sequence of alias, type name pairs to resolve types.
+            </summary>
+            <param name="aliasesSequence">Type aliases from the configuration file.</param>
+            <param name="assemblies">Assembly names to search.</param>
+            <param name="namespaces">Namespaces to search.</param>
+        </member>
+        <member name="M:Microsoft.Practices.Unity.Configuration.ConfigurationHelpers.TypeResolverImpl.ResolveType(System.String,System.Boolean)">
+            <summary>
+            Resolves a type alias or type FullName to a concrete type.
+            </summary>
+            <param name="typeNameOrAlias">Alias or name to resolve.</param>
+            <param name="throwIfResolveFails">if true and the alias does not
+            resolve, throw an <see cref="T:System.InvalidOperationException"/>, otherwise 
+            return null on failure.</param>
+            <returns>The type object or null if resolve fails and 
+            <paramref name="throwIfResolveFails"/> is false.</returns>
+        </member>
+        <member name="M:Microsoft.Practices.Unity.Configuration.ConfigurationHelpers.TypeResolverImpl.ResolveTypeWithDefault(System.String,System.Type,System.Boolean)">
+            <summary>
+            Resolve a type alias or type full name to a concrete type.
+            If <paramref name="typeNameOrAlias"/> is null or empty, return the
+            given <paramref name="defaultValue"/> instead.
+            </summary>
+            <param name="typeNameOrAlias">Type alias or full name to resolve.</param>
+            <param name="defaultValue">Value to return if typeName is null or empty.</param>
+            <param name="throwIfResolveFails">if true and the alias does not
+            resolve, throw an <see cref="T:System.InvalidOperationException"/>, otherwise 
+            return null on failure.</param>
+            <returns>
+            <para>If <paramref name="typeNameOrAlias"/> is null or an empty string,
+            then return <paramref name="defaultValue"/>.</para>
+            <para>Otherwise, return the resolved type object. If the resolution fails
+            and <paramref name="throwIfResolveFails"/> is false, then return null.</para>
+            </returns>
+        </member>
+        <member name="T:Microsoft.Practices.Unity.Configuration.UnityConfigurationSection">
+            <summary>
+            A configuration section describing configuration for an <see cref="T:Microsoft.Practices.Unity.IUnityContainer"/>.
+            </summary>
+        </member>
+        <member name="F:Microsoft.Practices.Unity.Configuration.UnityConfigurationSection.SectionName">
+            <summary>
+            The name of the section where unity configuration is expected to be found.
+            </summary>
+        </member>
+        <member name="F:Microsoft.Practices.Unity.Configuration.UnityConfigurationSection.XmlNamespace">
+            <summary>
+            XML Namespace string used for IntelliSense in this section.
+            </summary>
+        </member>
+        <member name="M:Microsoft.Practices.Unity.Configuration.UnityConfigurationSection.Configure(Microsoft.Practices.Unity.IUnityContainer)">
+            <summary>
+            Apply the configuration in the default container element to the given container.
+            </summary>
+            <param name="container">Container to configure.</param>
+            <returns>The passed in <paramref name="container"/>.</returns>
+        </member>
+        <member name="M:Microsoft.Practices.Unity.Configuration.UnityConfigurationSection.Configure(Microsoft.Practices.Unity.IUnityContainer,System.String)">
+            <summary>
+            Apply the configuration in the default container element to the given container.
+            </summary>
+            <param name="container">Container to configure.</param>
+            <param name="configuredContainerName">Name of the container element to use to configure the container.</param>
+            <returns>The passed in <paramref name="container"/>.</returns>
+        </member>
+        <member name="M:Microsoft.Practices.Unity.Configuration.UnityConfigurationSection.DeserializeSection(System.Xml.XmlReader)">
+            <summary>
+            Reads XML from the configuration file.
+            </summary>
+            <param name="reader">The <see cref="T:System.Xml.XmlReader"/> object, which reads from the configuration file. </param>
+            <exception cref="T:System.Configuration.ConfigurationErrorsException"><paramref name="reader"/> found no elements in the configuration file.</exception>
+        </member>
+        <member name="M:Microsoft.Practices.Unity.Configuration.UnityConfigurationSection.OnDeserializeUnrecognizedElement(System.String,System.Xml.XmlReader)">
+            <summary>
+            Gets a value indicating whether an unknown element is encountered during deserialization.
+            </summary>
+            <returns>
+            true when an unknown element is encountered while deserializing; otherwise, false.
+            </returns>
+            <param name="elementName">The name of the unknown subelement.</param>
+            <param name="reader">The <see cref="T:System.Xml.XmlReader"/> being used for deserialization.</param>
+            <exception cref="T:System.Configuration.ConfigurationErrorsException">
+                                The element identified by <paramref name="elementName"/> is locked.
+                                - or -
+                                One or more of the element's attributes is locked.
+                                - or -
+                            <paramref name="elementName"/> is unrecognized, or the element has an unrecognized attribute.
+                                - or -
+                                The element has a Boolean attribute with an invalid value.
+                                - or -
+                                An attempt was made to deserialize a property more than once.
+                                - or -
+                                An attempt was made to deserialize a property that is not a valid member of the element.
+                                - or -
+                                The element cannot contain a CDATA or text element.
+            </exception>
+        </member>
+        <member name="M:Microsoft.Practices.Unity.Configuration.UnityConfigurationSection.SerializeSection(System.Configuration.ConfigurationElement,System.String,System.Configuration.ConfigurationSaveMode)">
+            <summary>
+            Creates an XML string containing an unmerged view of the <see cref="T:System.Configuration.ConfigurationSection"/> object as a single section to write to a file.
+            </summary>
+            <returns>
+            An XML string containing an unmerged view of the <see cref="T:System.Configuration.ConfigurationSection"/> object.
+            </returns>
+            <param name="parentElement">The <see cref="T:System.Configuration.ConfigurationElement"/> instance to use as the parent when performing the un-merge.</param>
+            <param name="name">The name of the section to create.</param>
+            <param name="saveMode">The <see cref="T:System.Configuration.ConfigurationSaveMode"/> instance to use when writing to a string.</param>
+        </member>
+        <member name="P:Microsoft.Practices.Unity.Configuration.UnityConfigurationSection.CurrentSection">
+            <summary>
+            The current <see cref="T:Microsoft.Practices.Unity.Configuration.UnityConfigurationSection"/> that is being deserialized
+            or being configured from.
+            </summary>
+        </member>
+        <member name="P:Microsoft.Practices.Unity.Configuration.UnityConfigurationSection.Xmlns">
+            <summary>
+            Storage for XML namespace. The namespace isn't used or validated by config, but
+            it is useful for Visual Studio XML IntelliSense to kick in.
+            </summary>
+        </member>
+        <member name="P:Microsoft.Practices.Unity.Configuration.UnityConfigurationSection.Containers">
+            <summary>
+            The set of containers defined in this configuration section.
+            </summary>
+        </member>
+        <member name="P:Microsoft.Practices.Unity.Configuration.UnityConfigurationSection.TypeAliases">
+            <summary>
+            The set of type aliases defined in this configuration file.
+            </summary>
+        </member>
+        <member name="P:Microsoft.Practices.Unity.Configuration.UnityConfigurationSection.SectionExtensions">
+            <summary>
+            Any schema extensions that are added.
+            </summary>
+        </member>
+        <member name="P:Microsoft.Practices.Unity.Configuration.UnityConfigurationSection.Namespaces">
+            <summary>
+            Any namespaces added to the type search list.
+            </summary>
+        </member>
+        <member name="P:Microsoft.Practices.Unity.Configuration.UnityConfigurationSection.Assemblies">
+            <summary>
+            Any assemblies added to the type search list.
+            </summary>
+        </member>
+        <member name="M:Microsoft.Practices.Unity.Configuration.UnityConfigurationSection.ExtensionContext.AddAlias(System.String,System.Type)">
+            <summary>
+            Add a new alias to the configuration section. This is useful
+            for those extensions that add commonly used types to configuration
+            so users don't have to alias them repeatedly.
+            </summary>
+            <param name="newAlias">The alias to use.</param>
+            <param name="aliasedType">Type the alias maps to.</param>
+        </member>
+        <member name="M:Microsoft.Practices.Unity.Configuration.UnityConfigurationSection.ExtensionContext.AddElement(System.String,System.Type)">
+            <summary>
+            Add a new element to the configuration section schema.
+            </summary>
+            <param name="tag">Tag name in the XML.</param>
+            <param name="elementType">Type the tag maps to.</param>
+        </member>
+        <member name="T:Microsoft.Practices.Unity.Configuration.ContainerElementCollection">
+            <summary>
+            Collection element for <see cref="T:Microsoft.Practices.Unity.Configuration.ContainerElement"/>s.
+            </summary>
+        </member>
+        <member name="M:Microsoft.Practices.Unity.Configuration.ContainerElementCollection.GetElement(System.Int32)">
+            <summary>
+            Plug point to get objects out of the collection.
+            </summary>
+            <param name="index">Index in the collection to retrieve the item from.</param>
+            <returns>Item at that index or null if not present.</returns>
+        </member>
+        <member name="M:Microsoft.Practices.Unity.Configuration.ContainerElementCollection.GetElement(System.Object)">
+            <summary>
+            Plug point to get objects out of the collection.
+            </summary>
+            <param name="key">Key to look up the object by.</param>
+            <returns>Item with that key or null if not present.</returns>
+        </member>
+        <member name="M:Microsoft.Practices.Unity.Configuration.ContainerElementCollection.GetElementKey(System.Configuration.ConfigurationElement)">
+            <summary>
+            Gets the element key for a specified configuration element when overridden in a derived class.
+            </summary>
+            <returns>
+            An <see cref="T:System.Object"/> that acts as the key for the specified <see cref="T:System.Configuration.ConfigurationElement"/>.
+            </returns>
+            <param name="element">The <see cref="T:System.Configuration.ConfigurationElement"/> to return the key for. </param>
+        </member>
+        <member name="P:Microsoft.Practices.Unity.Configuration.ContainerElementCollection.Item(System.String)">
+            <summary>
+            Retrieve a stored <see cref="T:Microsoft.Practices.Unity.Configuration.ContainerElement"/> by name.
+            </summary>
+            <param name="name">Name to look up.</param>
+            <returns>The stored container or null if not in the collection.</returns>
+        </member>
+        <member name="P:Microsoft.Practices.Unity.Configuration.ContainerElementCollection.Default">
+            <summary>
+            Return the default container in the collection. The default is the one without a name.
+            </summary>
+        </member>
+        <member name="T:Microsoft.Practices.Unity.Configuration.UnityContainerExtensions">
+            <summary>
+            Extensions to <see cref="T:Microsoft.Practices.Unity.IUnityContainer"/> to simplify
+            loading configuration into a container.
+            </summary>
+        </member>
+        <member name="M:Microsoft.Practices.Unity.Configuration.UnityContainerExtensions.LoadConfiguration(Microsoft.Practices.Unity.IUnityContainer,Microsoft.Practices.Unity.Configuration.UnityConfigurationSection,System.String)">
+            <summary>
+            Apply configuration from the given section and named container
+            into the given container.
+            </summary>
+            <param name="container">Unity container to configure.</param>
+            <param name="section">Configuration section with config information.</param>
+            <param name="containerName">Named container.</param>
+            <returns><paramref name="container"/>.</returns>
+        </member>
+        <member name="M:Microsoft.Practices.Unity.Configuration.UnityContainerExtensions.LoadConfiguration(Microsoft.Practices.Unity.IUnityContainer,System.String)">
+            <summary>
+            Apply configuration from the default section (named "unity" pulled out of
+            ConfigurationManager) and the named container.
+            </summary>
+            <param name="container">Unity container to configure.</param>
+            <param name="containerName">Named container element in configuration.</param>
+            <returns><paramref name="container"/>.</returns>
+        </member>
+        <member name="M:Microsoft.Practices.Unity.Configuration.UnityContainerExtensions.LoadConfiguration(Microsoft.Practices.Unity.IUnityContainer)">
+            <summary>
+            Apply configuration from the default section and unnamed container element.
+            </summary>
+            <param name="container">Container to configure.</param>
+            <returns><paramref name="container"/>.</returns>
+        </member>
+        <member name="M:Microsoft.Practices.Unity.Configuration.UnityContainerExtensions.LoadConfiguration(Microsoft.Practices.Unity.IUnityContainer,Microsoft.Practices.Unity.Configuration.UnityConfigurationSection)">
+            <summary>
+            Apply configuration from the default container in the given section.
+            </summary>
+            <param name="container">Unity container to configure.</param>
+            <param name="section">Configuration section.</param>
+            <returns><paramref name="container"/>.</returns>
+        </member>
+        <member name="T:Microsoft.Practices.Unity.Configuration.ValueElement">
+            <summary>
+            Element that describes a constant value that will be
+            injected into the container.
+            </summary>
+        </member>
+        <member name="M:Microsoft.Practices.Unity.Configuration.ValueElement.#ctor">
+            <summary>
+            Construct a new <see cref="T:Microsoft.Practices.Unity.Configuration.ValueElement"/> object.
+            </summary>
+        </member>
+        <member name="M:Microsoft.Practices.Unity.Configuration.ValueElement.#ctor(System.Collections.Generic.IDictionary{System.String,System.String})">
+            <summary>
+            Construct a new <see cref="T:Microsoft.Practices.Unity.Configuration.ValueElement"/> object,
+            initializing properties from the contents of
+            <paramref name="propertyValues"/>.
+            </summary>
+            <param name="propertyValues">Name/value pairs which
+            contain the values to initialize properties to.</param>
+        </member>
+        <member name="M:Microsoft.Practices.Unity.Configuration.ValueElement.Microsoft#Practices#Unity#Configuration#IAttributeOnlyElement#SerializeContent(System.Xml.XmlWriter)">
+            <summary>
+            Write the contents of this element to the given <see cref="T:System.Xml.XmlWriter"/>.
+            </summary>
+            <remarks>The caller of this method has already written the start element tag before
+            calling this method, so deriving classes only need to write the element content, not
+            the start or end tags.</remarks>
+            <param name="writer">Writer to send XML content to.</param>
+        </member>
+        <member name="M:Microsoft.Practices.Unity.Configuration.ValueElement.SerializeContent(System.Xml.XmlWriter)">
+            <summary>
+            Write the contents of this element to the given <see cref="T:System.Xml.XmlWriter"/>. This
+            method always outputs an explicit &lt;dependency&gt; tag, instead of providing
+            attributes to the parent method.
+            </summary>
+            <param name="writer">Writer to send XML content to.</param>
+        </member>
+        <member name="M:Microsoft.Practices.Unity.Configuration.ValueElement.GetInjectionParameterValue(Microsoft.Practices.Unity.IUnityContainer,System.Type)">
+            <summary>
+            Generate an <see cref="T:Microsoft.Practices.Unity.InjectionParameterValue"/> object
+            that will be used to configure the container for a type registration.
+            </summary>
+            <param name="container">Container that is being configured. Supplied in order
+            to let custom implementations retrieve services; do not configure the container
+            directly in this method.</param>
+            <param name="parameterType">Type of the parameter to get the value for.</param>
+            <returns>The required <see cref="T:Microsoft.Practices.Unity.InjectionParameterValue"/> object.</returns>
+        </member>
+        <member name="P:Microsoft.Practices.Unity.Configuration.ValueElement.Value">
+            <summary>
+            Value for this element
+            </summary>
+        </member>
+        <member name="P:Microsoft.Practices.Unity.Configuration.ValueElement.TypeConverterTypeName">
+            <summary>
+            
+            </summary>
+        </member>
+    </members>
+</doc>
