Copying //UE4/Dev-Core to //UE4/Main

diff --git a/Engine/Source/Programs/CrashReporter/CrashReportWebSite/packages/Unity.4.0.1/lib/net45/Microsoft.Practices.Unity.RegistrationByConvention.XML b/Engine/Source/Programs/CrashReporter/CrashReportWebSite/packages/Unity.4.0.1/lib/net45/Microsoft.Practices.Unity.RegistrationByConvention.XML
new file mode 100644
--- /dev/null
+++ b/Engine/Source/Programs/CrashReporter/CrashReportWebSite/packages/Unity.4.0.1/lib/net45/Microsoft.Practices.Unity.RegistrationByConvention.XML
@@ -0,0 +1,328 @@
+<?xml version="1.0"?>
+<doc>
+    <assembly>
+        <name>Microsoft.Practices.Unity.RegistrationByConvention</name>
+    </assembly>
+    <members>
+        <member name="T:Microsoft.Practices.Unity.AllClasses">
+            <summary>
+            Provides helper methods to retrieve classes from assemblies.
+            </summary>
+        </member>
+        <member name="M:Microsoft.Practices.Unity.AllClasses.FromAssemblies(System.Reflection.Assembly[])">
+            <summary>
+            Returns all visible, non-abstract classes from <paramref name="assemblies"/>.
+            </summary>
+            <param name="assemblies">The assemblies.</param>
+            <returns>All visible, non-abstract classes found in the assemblies.</returns>
+            <exception cref="T:System.ArgumentNullException"><paramref name="assemblies"/> is <see langword="null"/>.</exception>
+            <exception cref="T:System.ArgumentException"><paramref name="assemblies"/> contains <see langword="null"/> elements.</exception>
+            <remarks>All exceptions thrown while getting types from the assemblies are ignored, and the types that can be retrieved are returned.</remarks>
+        </member>
+        <member name="M:Microsoft.Practices.Unity.AllClasses.FromAssemblies(System.Boolean,System.Reflection.Assembly[])">
+            <summary>
+            Returns all visible, non-abstract classes from <paramref name="assemblies"/>, and optionally skips errors.
+            </summary>
+            <param name="skipOnError"><see langword="true"/> to skip errors; otherwise, <see langword="true"/>.</param>
+            <param name="assemblies">The assemblies.</param>
+            <returns>
+            All visible, non-abstract classes.
+            </returns>
+            <exception cref="T:System.ArgumentNullException"><paramref name="assemblies"/> is <see langword="null"/>.</exception>
+            <exception cref="T:System.ArgumentException"><paramref name="assemblies"/> contains <see langword="null"/> elements.</exception>
+            <remarks>
+            If <paramref name="skipOnError"/> is <see langword="true"/>, all exceptions thrown while getting types from the assemblies are ignored, and the types 
+            that can be retrieved are returned; otherwise, the original exception is thrown.
+            </remarks>
+        </member>
+        <member name="M:Microsoft.Practices.Unity.AllClasses.FromAssemblies(System.Collections.Generic.IEnumerable{System.Reflection.Assembly},System.Boolean)">
+            <summary>
+            Returns all visible, non-abstract classes from <paramref name="assemblies"/>.
+            </summary>
+            <param name="skipOnError"><see langword="true"/> to skip errors; otherwise, <see langword="true"/>.</param>
+            <param name="assemblies">The assemblies.</param>
+            <returns>
+            All visible, non-abstract classes.
+            </returns>
+            <exception cref="T:System.ArgumentNullException"><paramref name="assemblies"/> is <see langword="null"/>.</exception>
+            <exception cref="T:System.ArgumentException"><paramref name="assemblies"/> contains <see langword="null"/> elements.</exception>
+            <remarks>
+            If <paramref name="skipOnError"/> is <see langword="true"/>, all exceptions thrown while getting types from the assemblies are ignored, and the types 
+            that can be retrieved are returned; otherwise, the original exception is thrown.
+            </remarks>
+        </member>
+        <member name="M:Microsoft.Practices.Unity.AllClasses.FromLoadedAssemblies(System.Boolean,System.Boolean,System.Boolean,System.Boolean)">
+            <summary>
+            Returns all visible, non-abstract classes from all assemblies that are loaded in the current application domain.
+            </summary>
+            <param name="includeSystemAssemblies"><see langword="false" /> to include system assemblies; otherwise, <see langword="false" />. Defaults to <see langword="false" />.</param>
+            <param name="includeUnityAssemblies"><see langword="false" /> to include the Unity assemblies; otherwise, <see langword="false" />. Defaults to <see langword="false" />.</param>
+            <param name="includeDynamicAssemblies"><see langword="false" /> to include dynamic assemblies; otherwise, <see langword="false" />. Defaults to <see langword="false" />.</param>
+            <param name="skipOnError"><see langword="true"/> to skip errors; otherwise, <see langword="true"/>.</param>
+            <returns>
+            All visible, non-abstract classes in the loaded assemblies.
+            </returns>
+            <remarks>
+            If <paramref name="skipOnError" /> is <see langword="true" />, all exceptions thrown while getting types from the assemblies are ignored, and the types
+            that can be retrieved are returned; otherwise, the original exception is thrown.
+            </remarks>
+        </member>
+        <member name="M:Microsoft.Practices.Unity.AllClasses.FromAssembliesInBasePath(System.Boolean,System.Boolean,System.Boolean)">
+            <summary>
+            Returns all visible, non-abstract classes from all assemblies that are located in the base folder of the current application domain.
+            </summary>
+            <param name="includeSystemAssemblies"><see langword="false" /> to include system assemblies; otherwise, <see langword="false" />. Defaults to <see langword="false" />.</param>
+            <param name="includeUnityAssemblies"><see langword="false" /> to include the Unity assemblies; otherwise, <see langword="false" />. Defaults to <see langword="false" />.</param>
+            <param name="skipOnError"><see langword="true"/> to skip errors; otherwise, <see langword="true"/>.</param>
+            <returns>
+            All visible, non-abstract classes.
+            </returns>
+            <remarks>
+            If <paramref name="skipOnError" /> is <see langword="true" />, all exceptions thrown while loading assemblies or getting types from the assemblies are ignored, and the types
+            that can be retrieved are returned; otherwise, the original exception is thrown.
+            </remarks>
+        </member>
+        <member name="T:Microsoft.Practices.Unity.DuplicateTypeMappingException">
+            <summary>
+            The exception that is thrown when registering multiple types would result in an type mapping being overwritten.
+            </summary>
+        </member>
+        <member name="M:Microsoft.Practices.Unity.DuplicateTypeMappingException.#ctor(System.String,System.Type,System.Type,System.Type)">
+            <summary>
+            Initializes a new instance of the <see cref="T:Microsoft.Practices.Unity.DuplicateTypeMappingException"/> class.
+            </summary>
+            <param name="name">The name for the mapping.</param>
+            <param name="mappedFromType">The source type for the mapping.</param>
+            <param name="currentMappedToType">The type currently mapped.</param>
+            <param name="newMappedToType">The new type to map.</param>
+        </member>
+        <member name="P:Microsoft.Practices.Unity.DuplicateTypeMappingException.Name">
+            <summary>
+            Gets the name for the mapping.
+            </summary>
+        </member>
+        <member name="P:Microsoft.Practices.Unity.DuplicateTypeMappingException.MappedFromType">
+            <summary>
+            Gets the source type for the mapping.
+            </summary>
+        </member>
+        <member name="P:Microsoft.Practices.Unity.DuplicateTypeMappingException.CurrentMappedToType">
+            <summary>
+            Gets the type currently mapped.
+            </summary>
+        </member>
+        <member name="P:Microsoft.Practices.Unity.DuplicateTypeMappingException.NewMappedToType">
+            <summary>
+            Gets the new type to map.
+            </summary>
+        </member>
+        <member name="T:Microsoft.Practices.Unity.Properties.Resources">
+            <summary>
+              A strongly-typed resource class, for looking up localized strings, etc.
+            </summary>
+        </member>
+        <member name="P:Microsoft.Practices.Unity.Properties.Resources.ResourceManager">
+            <summary>
+              Returns the cached ResourceManager instance used by this class.
+            </summary>
+        </member>
+        <member name="P:Microsoft.Practices.Unity.Properties.Resources.Culture">
+            <summary>
+              Overrides the current thread's CurrentUICulture property for all
+              resource lookups using this strongly typed resource class.
+            </summary>
+        </member>
+        <member name="P:Microsoft.Practices.Unity.Properties.Resources.DuplicateTypeMappingException">
+            <summary>
+              Looks up a localized string similar to An attempt to override an existing mapping was detected for type {1} with name &quot;{0}&quot;, currently mapped to type {2}, to type {3}..
+            </summary>
+        </member>
+        <member name="P:Microsoft.Practices.Unity.Properties.Resources.ExceptionNullAssembly">
+            <summary>
+              Looks up a localized string similar to The set of assemblies contains a null element..
+            </summary>
+        </member>
+        <member name="T:Microsoft.Practices.Unity.RegistrationConvention">
+            <summary>
+            Represents a set of types to register and their registration settings.
+            </summary>
+        </member>
+        <member name="M:Microsoft.Practices.Unity.RegistrationConvention.GetTypes">
+            <summary>
+            Gets types to register.
+            </summary>
+        </member>
+        <member name="M:Microsoft.Practices.Unity.RegistrationConvention.GetFromTypes">
+            <summary>
+            Gets a function to get the types that will be requested for each type to configure.
+            </summary>
+        </member>
+        <member name="M:Microsoft.Practices.Unity.RegistrationConvention.GetName">
+            <summary>
+            Gets a function to get the name to use for the registration of each type.
+            </summary>
+        </member>
+        <member name="M:Microsoft.Practices.Unity.RegistrationConvention.GetLifetimeManager">
+            <summary>
+            Gets a function to get the <see cref="T:Microsoft.Practices.Unity.LifetimeManager"/> for the registration of each type. Defaults to no lifetime management.
+            </summary>
+            <returns></returns>
+        </member>
+        <member name="M:Microsoft.Practices.Unity.RegistrationConvention.GetInjectionMembers">
+            <summary>
+            Gets a function to get the additional <see cref="T:Microsoft.Practices.Unity.InjectionMember"/> objects for the registration of each type. Defaults to no injection members.
+            </summary>
+            <returns></returns>
+        </member>
+        <member name="T:Microsoft.Practices.Unity.UnityContainerRegistrationByConventionExtensions">
+            <summary>
+            Provides a set of convenience overloads to the
+            <see cref="T:Microsoft.Practices.Unity.IUnityContainer"/> interface to support registration of multiple types.
+            </summary>
+        </member>
+        <member name="M:Microsoft.Practices.Unity.UnityContainerRegistrationByConventionExtensions.RegisterTypes(Microsoft.Practices.Unity.IUnityContainer,System.Collections.Generic.IEnumerable{System.Type},System.Func{System.Type,System.Collections.Generic.IEnumerable{System.Type}},System.Func{System.Type,System.String},System.Func{System.Type,Microsoft.Practices.Unity.LifetimeManager},System.Func{System.Type,System.Collections.Generic.IEnumerable{Microsoft.Practices.Unity.InjectionMember}},System.Boolean)">
+            <summary>
+            Registers the supplied types by using the specified rules for name, lifetime manager, injection members, and registration types.
+            </summary>
+            <param name="container">The container to configure.</param>
+            <param name="types">The types to register. The methods in the <see cref="T:Microsoft.Practices.Unity.AllClasses"/> class can be used to scan assemblies to get types, and further filtering can be performed using LINQ queries.</param>
+            <param name="getFromTypes">A function that gets the types that will be requested for each type to configure. It can be a method from the <see cref="T:Microsoft.Practices.Unity.WithMappings"/> class or a custom function. Defaults to no registration types, and registers only the supplied types.</param>
+            <param name="getName">A function that gets the name to use for the registration of each type. It can be a method from the <see cref="T:Microsoft.Practices.Unity.WithName"/> or a custom function. Defaults to no name.</param>
+            <param name="getLifetimeManager">A function that gets the <see cref="T:Microsoft.Practices.Unity.LifetimeManager"/> for the registration of each type. It can be a method from the <see cref="T:Microsoft.Practices.Unity.WithLifetime"/> class or a custom function. Defaults to no lifetime management.</param>
+            <param name="getInjectionMembers">A function that gets the additional <see cref="T:Microsoft.Practices.Unity.InjectionMember"/> objects for the registration of each type. Defaults to no injection members.</param>
+            <param name="overwriteExistingMappings"><see langword="true"/> to overwrite existing mappings; otherwise, <see langword="false"/>. Defaults to <see langword="false"/>.</param>
+            <returns>
+            The container that this method was called on.
+            </returns>
+            <exception cref="T:System.ArgumentException">A new registration would overwrite an existing mapping and <paramref name="overwriteExistingMappings"/> is <see langword="false"/>.</exception>
+        </member>
+        <member name="M:Microsoft.Practices.Unity.UnityContainerRegistrationByConventionExtensions.RegisterTypes(Microsoft.Practices.Unity.IUnityContainer,Microsoft.Practices.Unity.RegistrationConvention,System.Boolean)">
+            <summary>
+            Registers the types according to the <paramref name="convention"/>.
+            </summary>
+            <param name="container">The container to configure.</param>
+            <param name="convention">The convention to determine which types will be registered and how.</param>
+            <param name="overwriteExistingMappings"><see langword="true"/> to overwrite existing mappings; otherwise, <see langword="false"/>. Defaults to <see langword="false"/>.</param>
+            <returns>
+            The container that this method was called on.
+            </returns>
+        </member>
+        <member name="T:Microsoft.Practices.Unity.WithLifetime">
+            <summary>
+            Provides helper methods to specify the lifetime for a type with registration by convention.
+            </summary>
+        </member>
+        <member name="M:Microsoft.Practices.Unity.WithLifetime.None(System.Type)">
+            <summary>
+            Returns a <see langword="null"/> <see cref="T:Microsoft.Practices.Unity.LifetimeManager"/>.
+            </summary>
+            <param name="type">The type.</param>
+            <returns>A lifetime manager</returns>
+        </member>
+        <member name="M:Microsoft.Practices.Unity.WithLifetime.ContainerControlled(System.Type)">
+            <summary>
+            Returns a <see cref="T:Microsoft.Practices.Unity.ContainerControlledLifetimeManager"/>.
+            </summary>
+            <param name="type">The type.</param>
+            <returns>A container controlled lifetime manager.</returns>
+        </member>
+        <member name="M:Microsoft.Practices.Unity.WithLifetime.ExternallyControlled(System.Type)">
+            <summary>
+            Returns a <see cref="T:Microsoft.Practices.Unity.ExternallyControlledLifetimeManager"/>.
+            </summary>
+            <param name="type">The type.</param>
+            <returns>An externally controlled lifetime manager.</returns>
+        </member>
+        <member name="M:Microsoft.Practices.Unity.WithLifetime.Hierarchical(System.Type)">
+            <summary>
+            Returns a <see cref="T:Microsoft.Practices.Unity.HierarchicalLifetimeManager"/>.
+            </summary>
+            <param name="type">The type.</param>
+            <returns>A hierarchical lifetime manager.</returns>
+        </member>
+        <member name="M:Microsoft.Practices.Unity.WithLifetime.PerResolve(System.Type)">
+            <summary>
+            Returns a <see cref="T:Microsoft.Practices.Unity.PerResolveLifetimeManager"/>.
+            </summary>
+            <param name="type">The type.</param>
+            <returns>A per resolve lifetime manager.</returns>
+        </member>
+        <member name="M:Microsoft.Practices.Unity.WithLifetime.Transient(System.Type)">
+            <summary>
+            Returns a <see cref="T:Microsoft.Practices.Unity.TransientLifetimeManager"/>.
+            </summary>
+            <param name="type">The type.</param>
+            <returns>A transient lifetime manager.</returns>
+        </member>
+        <member name="M:Microsoft.Practices.Unity.WithLifetime.Custom``1(System.Type)">
+            <summary>
+            Returns a <see cref="T:Microsoft.Practices.Unity.LifetimeManager"/>.
+            </summary>
+            <typeparam name="T">The custom <see cref="T:Microsoft.Practices.Unity.LifetimeManager"/> type.</typeparam>
+            <param name="type">The type.</param>
+            <returns>
+            A lifetime manager.
+            </returns>
+        </member>
+        <member name="M:Microsoft.Practices.Unity.WithLifetime.PerThread(System.Type)">
+            <summary>
+            Returns a <see cref="T:Microsoft.Practices.Unity.PerThreadLifetimeManager"/>.
+            </summary>
+            <param name="type">The type.</param>
+            <returns>A per thread lifetime manager.</returns>
+        </member>
+        <member name="T:Microsoft.Practices.Unity.WithMappings">
+            <summary>
+            Provides helper methods to map types to the types interfaces to which register them.
+            </summary>
+        </member>
+        <member name="M:Microsoft.Practices.Unity.WithMappings.None(System.Type)">
+            <summary>
+            Returns no types.
+            </summary>
+            <param name="implementationType">The type to register.</param>
+            <returns>An empty enumeration.</returns>
+        </member>
+        <member name="M:Microsoft.Practices.Unity.WithMappings.FromMatchingInterface(System.Type)">
+            <summary>
+            Returns an enumeration with the interface that matches the name of <paramref name="implementationType"/>.
+            </summary>
+            <param name="implementationType">The type to register.</param>
+            <returns>An enumeration with the first interface matching the name of <paramref name="implementationType"/> (for example, if type is MyType, a matching interface is IMyType),
+            or an empty enumeration if no such interface is found.</returns>
+        </member>
+        <member name="M:Microsoft.Practices.Unity.WithMappings.FromAllInterfaces(System.Type)">
+            <summary>
+            Returns an enumeration with all the interfaces implemented by <paramref name="implementationType"/>.
+            </summary>
+            <param name="implementationType">The type to register.</param>
+            <returns>An enumeration with all the interfaces implemented by the implementation type except <see cref="T:System.IDisposable"/>.</returns>
+        </member>
+        <member name="M:Microsoft.Practices.Unity.WithMappings.FromAllInterfacesInSameAssembly(System.Type)">
+            <summary>
+            Returns an enumeration with all the interfaces implemented by <paramref name="implementationType"/> that belong to the same assembly as implementationType.
+            </summary>
+            <param name="implementationType">The type to register.</param>
+            <returns>An enumeration with all the interfaces implemented by the implementation type that belong to the same assembly.</returns>
+        </member>
+        <member name="T:Microsoft.Practices.Unity.WithName">
+            <summary>
+            Provides helper methods to get type names.
+            </summary>
+        </member>
+        <member name="M:Microsoft.Practices.Unity.WithName.TypeName(System.Type)">
+            <summary>
+            Returns the type name.
+            </summary>
+            <param name="type">The type.</param>
+            <returns>The type name.</returns>
+        </member>
+        <member name="M:Microsoft.Practices.Unity.WithName.Default(System.Type)">
+            <summary>
+            Returns null for the registration name.
+            </summary>
+            <param name="type">The type.</param>
+            <returns><see langword="null"/></returns>
+        </member>
+    </members>
+</doc>
