Copying //UE4/Dev-Core to //UE4/Main

diff --git a/Engine/Source/Programs/CrashReporter/CrashReportWebSite/packages/Unity.4.0.1/lib/net45/Microsoft.Practices.Unity.xml b/Engine/Source/Programs/CrashReporter/CrashReportWebSite/packages/Unity.4.0.1/lib/net45/Microsoft.Practices.Unity.xml
new file mode 100644
--- /dev/null
+++ b/Engine/Source/Programs/CrashReporter/CrashReportWebSite/packages/Unity.4.0.1/lib/net45/Microsoft.Practices.Unity.xml
@@ -0,0 +1,5916 @@
+<?xml version="1.0"?>
+<doc>
+    <assembly>
+        <name>Microsoft.Practices.Unity</name>
+    </assembly>
+    <members>
+        <member name="T:Microsoft.Practices.Unity.IRegisteredNamesPolicy">
+            <summary>
+            Provides access to the names registered for a container.
+            </summary>
+        </member>
+        <member name="T:Microsoft.Practices.ObjectBuilder2.IBuilderPolicy">
+            <summary>
+            Represents a builder policy interface. Since there are no fixed requirements
+            for policies, it acts as a marker interface from which to derive all other
+            policy interfaces.
+            </summary>
+        </member>
+        <member name="M:Microsoft.Practices.Unity.IRegisteredNamesPolicy.GetRegisteredNames(System.Type)">
+            <summary>
+            Gets the names registered for a type.
+            </summary>
+            <param name="type">The type.</param>
+            <returns>The names registered for <paramref name="type"/>.</returns>
+        </member>
+        <member name="T:Microsoft.Practices.ObjectBuilder2.LazyDynamicMethodBuildPlanCreatorPolicy">
+            <summary>
+            An <see cref="T:Microsoft.Practices.ObjectBuilder2.IBuildPlanCreatorPolicy"/> implementation
+            that constructs a build plan for creating <see cref="T:System.Lazy`1"/> objects.
+            </summary>
+        </member>
+        <member name="T:Microsoft.Practices.ObjectBuilder2.IBuildPlanCreatorPolicy">
+            <summary>
+            A <see cref="T:Microsoft.Practices.ObjectBuilder2.IBuilderPolicy"/> that can create and return an <see cref="T:Microsoft.Practices.ObjectBuilder2.IBuildPlanPolicy"/>
+            for the given build key.
+            </summary>
+        </member>
+        <member name="M:Microsoft.Practices.ObjectBuilder2.IBuildPlanCreatorPolicy.CreatePlan(Microsoft.Practices.ObjectBuilder2.IBuilderContext,Microsoft.Practices.ObjectBuilder2.NamedTypeBuildKey)">
+            <summary>
+            Create a build plan using the given context and build key.
+            </summary>
+            <param name="context">Current build context.</param>
+            <param name="buildKey">Current build key.</param>
+            <returns>The build plan.</returns>
+        </member>
+        <member name="M:Microsoft.Practices.ObjectBuilder2.LazyDynamicMethodBuildPlanCreatorPolicy.CreatePlan(Microsoft.Practices.ObjectBuilder2.IBuilderContext,Microsoft.Practices.ObjectBuilder2.NamedTypeBuildKey)">
+            <summary>
+            Creates a build plan using the given context and build key.
+            </summary>
+            <param name="context">Current build context.</param>
+            <param name="buildKey">Current build key.</param>
+            <returns>
+            The build plan.
+            </returns>
+        </member>
+        <member name="T:Microsoft.Practices.Unity.Utility.TypeReflectionExtensions">
+            <summary>
+            Provides extension methods to the <see cref="T:System.Type"/> class due to the introduction 
+            of <see cref="T:System.Reflection.TypeInfo"/> class in the .NET for Windows Store apps.
+            </summary>
+        </member>
+        <member name="M:Microsoft.Practices.Unity.Utility.TypeReflectionExtensions.GetConstructor(System.Type,System.Type[])">
+            <summary>
+            Returns the constructor in <paramref name="type"/> that matches the specified constructor parameter types.
+            </summary>
+            <param name="type">The type to inspect</param>
+            <param name="constructorParameters">The constructor parameter types.</param>
+            <returns>The constructor that matches the specified parameter types.</returns>
+        </member>
+        <member name="M:Microsoft.Practices.Unity.Utility.TypeReflectionExtensions.GetMethodsHierarchical(System.Type)">
+            <summary>
+            Returns the non-static declared methods of a type or its base types.
+            </summary>
+            <param name="type">The type to inspect</param>
+            <returns>An enumerable of the <see cref="T:System.Reflection.MethodInfo"/> objects.</returns>
+        </member>
+        <member name="M:Microsoft.Practices.Unity.Utility.TypeReflectionExtensions.GetMethodHierarchical(System.Type,System.String,System.Type[])">
+            <summary>
+            Returns the non-static method of a type or its based type.
+            </summary>
+            <param name="type">The type to inspect</param>
+            <param name="methodName">The name of the method to seek.</param>
+            <param name="closedParameters">The (closed) parameter type signature of the method.</param>
+            <returns>The discovered <see cref="T:System.Reflection.MethodInfo"/></returns>
+        </member>
+        <member name="M:Microsoft.Practices.Unity.Utility.TypeReflectionExtensions.GetPropertiesHierarchical(System.Type)">
+            <summary>
+            Returns the declared properties of a type or its base types.
+            </summary>
+            <param name="type">The type to inspect</param>
+            <returns>An enumerable of the <see cref="T:System.Reflection.PropertyInfo"/> objects.</returns>
+        </member>
+        <member name="M:Microsoft.Practices.Unity.Utility.TypeReflectionExtensions.ParametersMatch(System.Reflection.ParameterInfo[],System.Type[])">
+            <summary>
+            Determines if the types in a parameter set ordinally matches the set of supplied types.
+            </summary>
+            <param name="parameters"></param>
+            <param name="closedConstructorParameterTypes"></param>
+            <returns></returns>
+        </member>
+        <member name="T:Microsoft.Practices.Unity.DependencyResolutionAttribute">
+            <summary>
+            Base class for attributes that can be placed on parameters
+            or properties to specify how to resolve the value for
+            that parameter or property.
+            </summary>
+        </member>
+        <member name="M:Microsoft.Practices.Unity.DependencyResolutionAttribute.CreateResolver(System.Type)">
+            <summary>
+            Create an instance of <see cref="T:Microsoft.Practices.ObjectBuilder2.IDependencyResolverPolicy"/> that
+            will be used to get the value for the member this attribute is
+            applied to.
+            </summary>
+            <param name="typeToResolve">Type of parameter or property that
+            this attribute is decoration.</param>
+            <returns>The resolver object.</returns>
+        </member>
+        <member name="T:Microsoft.Practices.Unity.InjectionConstructorAttribute">
+            <summary>
+            This attribute is used to indicate which constructor to choose when
+            the container attempts to build a type.
+            </summary>
+        </member>
+        <member name="T:Microsoft.Practices.Unity.InjectionMethodAttribute">
+            <summary>
+            This attribute is used to mark methods that should be called when
+            the container is building an object.
+            </summary>
+        </member>
+        <member name="T:Microsoft.Practices.Unity.DependencyAttribute">
+            <summary>
+            This attribute is used to mark properties and parameters as targets for injection.
+            </summary>
+            <remarks>
+            For properties, this attribute is necessary for injection to happen. For parameters,
+            it's not needed unless you want to specify additional information to control how
+            the parameter is resolved.
+            </remarks>
+        </member>
+        <member name="M:Microsoft.Practices.Unity.DependencyAttribute.#ctor">
+            <summary>
+            Create an instance of <see cref="T:Microsoft.Practices.Unity.DependencyAttribute"/> with no name.
+            </summary>
+        </member>
+        <member name="M:Microsoft.Practices.Unity.DependencyAttribute.#ctor(System.String)">
+            <summary>
+            Create an instance of <see cref="T:Microsoft.Practices.Unity.DependencyAttribute"/> with the given name.
+            </summary>
+            <param name="name">Name to use when resolving this dependency.</param>
+        </member>
+        <member name="M:Microsoft.Practices.Unity.DependencyAttribute.CreateResolver(System.Type)">
+            <summary>
+            Create an instance of <see cref="T:Microsoft.Practices.ObjectBuilder2.IDependencyResolverPolicy"/> that
+            will be used to get the value for the member this attribute is
+            applied to.
+            </summary>
+            <param name="typeToResolve">Type of parameter or property that
+            this attribute is decoration.</param>
+            <returns>The resolver object.</returns>
+        </member>
+        <member name="P:Microsoft.Practices.Unity.DependencyAttribute.Name">
+            <summary>
+            The name specified in the constructor.
+            </summary>
+        </member>
+        <member name="T:Microsoft.Practices.Unity.OptionalDependencyAttribute">
+            <summary>
+            An <see cref="T:Microsoft.Practices.Unity.DependencyResolutionAttribute"/> used to mark a dependency
+            as optional - the container will try to resolve it, and return null
+            if the resolution fails rather than throw.
+            </summary>
+        </member>
+        <member name="M:Microsoft.Practices.Unity.OptionalDependencyAttribute.#ctor">
+            <summary>
+            Construct a new <see cref="T:Microsoft.Practices.Unity.OptionalDependencyAttribute"/> object.
+            </summary>
+        </member>
+        <member name="M:Microsoft.Practices.Unity.OptionalDependencyAttribute.#ctor(System.String)">
+            <summary>
+            Construct a new <see cref="T:Microsoft.Practices.Unity.OptionalDependencyAttribute"/> object that
+            specifies a named dependency.
+            </summary>
+            <param name="name">Name of the dependency.</param>
+        </member>
+        <member name="M:Microsoft.Practices.Unity.OptionalDependencyAttribute.CreateResolver(System.Type)">
+            <summary>
+            Create an instance of <see cref="T:Microsoft.Practices.ObjectBuilder2.IDependencyResolverPolicy"/> that
+            will be used to get the value for the member this attribute is
+            applied to.
+            </summary>
+            <param name="typeToResolve">Type of parameter or property that
+            this attribute is decoration.</param>
+            <returns>The resolver object.</returns>
+        </member>
+        <member name="P:Microsoft.Practices.Unity.OptionalDependencyAttribute.Name">
+            <summary>
+            Name of the dependency.
+            </summary>
+        </member>
+        <member name="T:Microsoft.Practices.Unity.CompositeResolverOverride">
+            <summary>
+            A <see cref="T:Microsoft.Practices.Unity.ResolverOverride"/> that composites other
+            ResolverOverride objects. The GetResolver operation then
+            returns the resolver from the first child override that
+            matches the current context and request.
+            </summary>
+        </member>
+        <member name="T:Microsoft.Practices.Unity.ResolverOverride">
+            <summary>
+            Base class for all override objects passed in the
+            <see cref="M:Microsoft.Practices.Unity.IUnityContainer.Resolve(System.Type,System.String,Microsoft.Practices.Unity.ResolverOverride[])"/> method.
+            </summary>
+        </member>
+        <member name="M:Microsoft.Practices.Unity.ResolverOverride.GetResolver(Microsoft.Practices.ObjectBuilder2.IBuilderContext,System.Type)">
+            <summary>
+            Return a <see cref="T:Microsoft.Practices.ObjectBuilder2.IDependencyResolverPolicy"/> that can be used to give a value
+            for the given desired dependency.
+            </summary>
+            <param name="context">Current build context.</param>
+            <param name="dependencyType">Type of dependency desired.</param>
+            <returns>a <see cref="T:Microsoft.Practices.ObjectBuilder2.IDependencyResolverPolicy"/> object if this override applies, null if not.</returns>
+        </member>
+        <member name="M:Microsoft.Practices.Unity.ResolverOverride.OnType``1">
+            <summary>
+            Wrap this resolver in one that verifies the type of the object being built.
+            This allows you to narrow any override down to a specific type easily.
+            </summary>
+            <typeparam name="T">Type to constrain the override to.</typeparam>
+            <returns>The new override.</returns>
+        </member>
+        <member name="M:Microsoft.Practices.Unity.ResolverOverride.OnType(System.Type)">
+            <summary>
+            Wrap this resolver in one that verifies the type of the object being built.
+            This allows you to narrow any override down to a specific type easily.
+            </summary>
+            <param name="typeToOverride">Type to constrain the override to.</param>
+            <returns>The new override.</returns>
+        </member>
+        <member name="M:Microsoft.Practices.Unity.CompositeResolverOverride.Add(Microsoft.Practices.Unity.ResolverOverride)">
+            <summary>
+            Add a new <see cref="T:Microsoft.Practices.Unity.ResolverOverride"/> to the collection
+            that is checked.
+            </summary>
+            <param name="newOverride">item to add.</param>
+        </member>
+        <member name="M:Microsoft.Practices.Unity.CompositeResolverOverride.AddRange(System.Collections.Generic.IEnumerable{Microsoft.Practices.Unity.ResolverOverride})">
+            <summary>
+            Add a set of <see cref="T:Microsoft.Practices.Unity.ResolverOverride"/>s to the collection.
+            </summary>
+            <param name="newOverrides">items to add.</param>
+        </member>
+        <member name="M:Microsoft.Practices.Unity.CompositeResolverOverride.System#Collections#IEnumerable#GetEnumerator">
+            <summary>
+            Returns an enumerator that iterates through a collection.
+            </summary>
+            <returns>
+            An <see cref="T:System.Collections.IEnumerator"/> object that can be used to iterate through the collection.
+            </returns>
+            <filterpriority>2</filterpriority>
+        </member>
+        <member name="M:Microsoft.Practices.Unity.CompositeResolverOverride.GetEnumerator">
+            <summary>
+            Returns an enumerator that iterates through the collection.
+            </summary>
+            <returns>
+            A <see cref="T:System.Collections.Generic.IEnumerator`1"/> that can be used to iterate through the collection.
+            </returns>
+            <filterpriority>1</filterpriority>
+        </member>
+        <member name="M:Microsoft.Practices.Unity.CompositeResolverOverride.GetResolver(Microsoft.Practices.ObjectBuilder2.IBuilderContext,System.Type)">
+            <summary>
+            Return a <see cref="T:Microsoft.Practices.ObjectBuilder2.IDependencyResolverPolicy"/> that can be used to give a value
+            for the given desired dependency.
+            </summary>
+            <param name="context">Current build context.</param>
+            <param name="dependencyType">Type of dependency desired.</param>
+            <returns>a <see cref="T:Microsoft.Practices.ObjectBuilder2.IDependencyResolverPolicy"/> object if this override applies, null if not.</returns>
+        </member>
+        <member name="T:Microsoft.Practices.Unity.ContainerRegistration">
+            <summary>
+            Class that returns information about the types registered in a container.
+            </summary>
+        </member>
+        <member name="P:Microsoft.Practices.Unity.ContainerRegistration.RegisteredType">
+            <summary>
+            The type that was passed to the <see cref="M:Microsoft.Practices.Unity.IUnityContainer.RegisterType(System.Type,System.Type,System.String,Microsoft.Practices.Unity.LifetimeManager,Microsoft.Practices.Unity.InjectionMember[])"/> method
+            as the "from" type, or the only type if type mapping wasn't done.
+            </summary>
+        </member>
+        <member name="P:Microsoft.Practices.Unity.ContainerRegistration.MappedToType">
+            <summary>
+            The type that this registration is mapped to. If no type mapping was done, the
+            <see cref="P:Microsoft.Practices.Unity.ContainerRegistration.RegisteredType"/> property and this one will have the same value.
+            </summary>
+        </member>
+        <member name="P:Microsoft.Practices.Unity.ContainerRegistration.Name">
+            <summary>
+            Name the type was registered under. Null for default registration.
+            </summary>
+        </member>
+        <member name="P:Microsoft.Practices.Unity.ContainerRegistration.LifetimeManagerType">
+            <summary>
+            The registered lifetime manager instance.
+            </summary>
+        </member>
+        <member name="P:Microsoft.Practices.Unity.ContainerRegistration.LifetimeManager">
+            <summary>
+            The lifetime manager for this registration.
+            </summary>
+            <remarks>
+            This property will be null if this registration is for an open generic.</remarks>
+        </member>
+        <member name="T:Microsoft.Practices.Unity.DependencyOverride">
+            <summary>
+            A <see cref="T:Microsoft.Practices.Unity.ResolverOverride"/> class that overrides
+            the value injected whenever there is a dependency of the
+            given type, regardless of where it appears in the object graph.
+            </summary>
+        </member>
+        <member name="M:Microsoft.Practices.Unity.DependencyOverride.#ctor(System.Type,System.Object)">
+            <summary>
+            Create an instance of <see cref="T:Microsoft.Practices.Unity.DependencyOverride"/> to override
+            the given type with the given value.
+            </summary>
+            <param name="typeToConstruct">Type of the dependency.</param>
+            <param name="dependencyValue">Value to use.</param>
+        </member>
+        <member name="M:Microsoft.Practices.Unity.DependencyOverride.GetResolver(Microsoft.Practices.ObjectBuilder2.IBuilderContext,System.Type)">
+            <summary>
+            Return a <see cref="T:Microsoft.Practices.ObjectBuilder2.IDependencyResolverPolicy"/> that can be used to give a value
+            for the given desired dependency.
+            </summary>
+            <param name="context">Current build context.</param>
+            <param name="dependencyType">Type of dependency desired.</param>
+            <returns>a <see cref="T:Microsoft.Practices.ObjectBuilder2.IDependencyResolverPolicy"/> object if this override applies, null if not.</returns>
+        </member>
+        <member name="T:Microsoft.Practices.Unity.DependencyOverride`1">
+            <summary>
+            A convenience version of <see cref="T:Microsoft.Practices.Unity.DependencyOverride"/> that lets you
+            specify the dependency type using generic syntax.
+            </summary>
+            <typeparam name="T">Type of the dependency to override.</typeparam>
+        </member>
+        <member name="M:Microsoft.Practices.Unity.DependencyOverride`1.#ctor(System.Object)">
+            <summary>
+            Construct a new <see cref="T:Microsoft.Practices.Unity.DependencyOverride`1"/> object that will
+            override the given dependency, and pass the given value.
+            </summary>
+        </member>
+        <member name="T:Microsoft.Practices.Unity.DependencyOverrides">
+            <summary>
+            A convenience form of <see cref="T:Microsoft.Practices.Unity.DependencyOverride"/> that lets you
+            specify multiple parameter overrides in one shot rather than having
+            to construct multiple objects.
+            </summary>
+            <remarks>
+            This class isn't really a collection, it just implements IEnumerable
+            so that we get use of the nice C# collection initializer syntax.
+            </remarks>
+        </member>
+        <member name="T:Microsoft.Practices.Unity.OverrideCollection`3">
+            <summary>
+            Base helper class for creating collections of <see cref="T:Microsoft.Practices.Unity.ResolverOverride"/> objects
+            for use in passing a bunch of them to the resolve call. This base class provides
+            the mechanics needed to allow you to use the C# collection initializer syntax.
+            </summary>
+            <typeparam name="TOverride">Concrete type of the <see cref="T:Microsoft.Practices.Unity.ResolverOverride"/> this class collects.</typeparam>
+            <typeparam name="TKey">Key used to create the underlying override object.</typeparam>
+            <typeparam name="TValue">Value that the override returns.</typeparam>
+        </member>
+        <member name="M:Microsoft.Practices.Unity.OverrideCollection`3.Add(`1,`2)">
+            <summary>
+            Add a new override to the collection with the given key and value.
+            </summary>
+            <param name="key">Key - for example, a parameter or property name.</param>
+            <param name="value">Value - the value to be returned by the override.</param>
+        </member>
+        <member name="M:Microsoft.Practices.Unity.OverrideCollection`3.GetResolver(Microsoft.Practices.ObjectBuilder2.IBuilderContext,System.Type)">
+            <summary>
+            Return a <see cref="T:Microsoft.Practices.ObjectBuilder2.IDependencyResolverPolicy"/> that can be used to give a value
+            for the given desired dependency.
+            </summary>
+            <param name="context">Current build context.</param>
+            <param name="dependencyType">Type of dependency desired.</param>
+            <returns>a <see cref="T:Microsoft.Practices.ObjectBuilder2.IDependencyResolverPolicy"/> object if this override applies, null if not.</returns>
+        </member>
+        <member name="M:Microsoft.Practices.Unity.OverrideCollection`3.System#Collections#IEnumerable#GetEnumerator">
+            <summary>
+            Returns an enumerator that iterates through a collection.
+            </summary>
+            <returns>
+            An <see cref="T:System.Collections.IEnumerator"/> object that can be used to iterate through the collection.
+            </returns>
+            <filterpriority>2</filterpriority>
+        </member>
+        <member name="M:Microsoft.Practices.Unity.OverrideCollection`3.GetEnumerator">
+            <summary>
+            Returns an enumerator that iterates through the collection.
+            </summary>
+            <returns>
+            A <see cref="T:System.Collections.Generic.IEnumerator`1"/> that can be used to iterate through the collection.
+            </returns>
+            <filterpriority>1</filterpriority>
+        </member>
+        <member name="M:Microsoft.Practices.Unity.OverrideCollection`3.MakeOverride(`1,`2)">
+            <summary>
+            When implemented in derived classes, this method is called from the <see cref="M:Microsoft.Practices.Unity.OverrideCollection`3.Add(`1,`2)"/>
+            method to create the actual <see cref="T:Microsoft.Practices.Unity.ResolverOverride"/> objects.
+            </summary>
+            <param name="key">Key value to create the resolver.</param>
+            <param name="value">Value to store in the resolver.</param>
+            <returns>The created <see cref="T:Microsoft.Practices.Unity.ResolverOverride"/>.</returns>
+        </member>
+        <member name="M:Microsoft.Practices.Unity.DependencyOverrides.MakeOverride(System.Type,System.Object)">
+            <summary>
+            When implemented in derived classes, this method is called from the <see cref="M:Microsoft.Practices.Unity.OverrideCollection`3.Add(`1,`2)"/>
+            method to create the actual <see cref="T:Microsoft.Practices.Unity.ResolverOverride"/> objects.
+            </summary>
+            <param name="key">Key value to create the resolver.</param>
+            <param name="value">Value to store in the resolver.</param>
+            <returns>The created <see cref="T:Microsoft.Practices.Unity.ResolverOverride"/>.</returns>
+        </member>
+        <member name="T:Microsoft.Practices.Unity.ChildContainerCreatedEventArgs">
+            <summary>
+            Event argument class for the <see cref="E:Microsoft.Practices.Unity.ExtensionContext.ChildContainerCreated"/> event.
+            </summary>
+        </member>
+        <member name="M:Microsoft.Practices.Unity.ChildContainerCreatedEventArgs.#ctor(Microsoft.Practices.Unity.ExtensionContext)">
+            <summary>
+            Construct a new <see cref="T:Microsoft.Practices.Unity.ChildContainerCreatedEventArgs"/> object with the
+            given child container object.
+            </summary>
+            <param name="childContext">An <see cref="T:Microsoft.Practices.Unity.ExtensionContext"/> for the newly created child
+            container.</param>
+        </member>
+        <member name="P:Microsoft.Practices.Unity.ChildContainerCreatedEventArgs.ChildContainer">
+            <summary>
+            The newly created child container.
+            </summary>
+        </member>
+        <member name="P:Microsoft.Practices.Unity.ChildContainerCreatedEventArgs.ChildContext">
+            <summary>
+            An extension context for the created child container.
+            </summary>
+        </member>
+        <member name="T:Microsoft.Practices.Unity.GenericParameterBase">
+            <summary>
+            Base class for <see cref="T:Microsoft.Practices.Unity.InjectionParameterValue"/> subclasses that let you specify that
+            an instance of a generic type parameter should be resolved.
+            </summary>
+        </member>
+        <member name="T:Microsoft.Practices.Unity.InjectionParameterValue">
+            <summary>
+            Base type for objects that are used to configure parameters for
+            constructor or method injection, or for getting the value to
+            be injected into a property.
+            </summary>
+        </member>
+        <member name="M:Microsoft.Practices.Unity.InjectionParameterValue.MatchesType(System.Type)">
+            <summary>
+            Test to see if this parameter value has a matching type for the given type.
+            </summary>
+            <param name="t">Type to check.</param>
+            <returns>True if this parameter value is compatible with type <paramref name="t"/>,
+            false if not.</returns>
+        </member>
+        <member name="M:Microsoft.Practices.Unity.InjectionParameterValue.GetResolverPolicy(System.Type)">
+            <summary>
+            Return a <see cref="T:Microsoft.Practices.ObjectBuilder2.IDependencyResolverPolicy"/> instance that will
+            return this types value for the parameter.
+            </summary>
+            <param name="typeToBuild">Type that contains the member that needs this parameter. Used
+            to resolve open generic parameters.</param>
+            <returns>The <see cref="T:Microsoft.Practices.ObjectBuilder2.IDependencyResolverPolicy"/>.</returns>
+        </member>
+        <member name="M:Microsoft.Practices.Unity.InjectionParameterValue.ToParameters(System.Object[])">
+            <summary>
+            Convert the given set of arbitrary values to a sequence of InjectionParameterValue
+            objects. The rules are: If it's already an InjectionParameterValue, return it. If
+            it's a Type, return a ResolvedParameter object for that type. Otherwise return
+            an InjectionParameter object for that value.
+            </summary>
+            <param name="values">The values to build the sequence from.</param>
+            <returns>The resulting converted sequence.</returns>
+        </member>
+        <member name="M:Microsoft.Practices.Unity.InjectionParameterValue.ToParameter(System.Object)">
+            <summary>
+            Convert an arbitrary value to an InjectionParameterValue object. The rules are: 
+            If it's already an InjectionParameterValue, return it. If it's a Type, return a
+            ResolvedParameter object for that type. Otherwise return an InjectionParameter
+            object for that value.
+            </summary>
+            <param name="value">The value to convert.</param>
+            <returns>The resulting <see cref="T:Microsoft.Practices.Unity.InjectionParameterValue"/>.</returns>
+        </member>
+        <member name="P:Microsoft.Practices.Unity.InjectionParameterValue.ParameterTypeName">
+            <summary>
+            Name for the type represented by this <see cref="T:Microsoft.Practices.Unity.InjectionParameterValue"/>.
+            This may be an actual type name or a generic argument name.
+            </summary>
+        </member>
+        <member name="M:Microsoft.Practices.Unity.GenericParameterBase.#ctor(System.String)">
+            <summary>
+            Create a new <see cref="T:Microsoft.Practices.Unity.GenericParameter"/> instance that specifies
+            that the given named generic parameter should be resolved.
+            </summary>
+            <param name="genericParameterName">The generic parameter name to resolve.</param>
+        </member>
+        <member name="M:Microsoft.Practices.Unity.GenericParameterBase.#ctor(System.String,System.String)">
+            <summary>
+            Create a new <see cref="T:Microsoft.Practices.Unity.GenericParameter"/> instance that specifies
+            that the given named generic parameter should be resolved.
+            </summary>
+            <param name="genericParameterName">The generic parameter name to resolve.</param>
+            <param name="resolutionKey">name to use when looking up in the container.</param>
+        </member>
+        <member name="M:Microsoft.Practices.Unity.GenericParameterBase.MatchesType(System.Type)">
+            <summary>
+            Test to see if this parameter value has a matching type for the given type.
+            </summary>
+            <param name="t">Type to check.</param>
+            <returns>True if this parameter value is compatible with type <paramref name="t"/>,
+            false if not.</returns>
+        </member>
+        <member name="M:Microsoft.Practices.Unity.GenericParameterBase.GetResolverPolicy(System.Type)">
+            <summary>
+            Return a <see cref="T:Microsoft.Practices.ObjectBuilder2.IDependencyResolverPolicy"/> instance that will
+            return this types value for the parameter.
+            </summary>
+            <param name="typeToBuild">Type that contains the member that needs this parameter. Used
+            to resolve open generic parameters.</param>
+            <returns>The <see cref="T:Microsoft.Practices.ObjectBuilder2.IDependencyResolverPolicy"/>.</returns>
+        </member>
+        <member name="M:Microsoft.Practices.Unity.GenericParameterBase.DoGetResolverPolicy(System.Type,System.String)">
+            <summary>
+            Return a <see cref="T:Microsoft.Practices.ObjectBuilder2.IDependencyResolverPolicy"/> instance that will
+            return this types value for the parameter.
+            </summary>
+            <param name="typeToResolve">The actual type to resolve.</param>
+            <param name="resolutionKey">The resolution key.</param>
+            <returns>The <see cref="T:Microsoft.Practices.ObjectBuilder2.IDependencyResolverPolicy"/>.</returns>
+        </member>
+        <member name="P:Microsoft.Practices.Unity.GenericParameterBase.ParameterTypeName">
+            <summary>
+            Name for the type represented by this <see cref="T:Microsoft.Practices.Unity.InjectionParameterValue"/>.
+            This may be an actual type name or a generic argument name.
+            </summary>
+        </member>
+        <member name="T:Microsoft.Practices.Unity.OptionalGenericParameter">
+            <summary>
+            A <see cref="T:Microsoft.Practices.Unity.InjectionParameterValue"/> that lets you specify that
+            an instance of a generic type parameter should be resolved, providing the <see langword="null"/>
+            value if resolving fails.
+            </summary>
+        </member>
+        <member name="M:Microsoft.Practices.Unity.OptionalGenericParameter.#ctor(System.String)">
+            <summary>
+            Create a new <see cref="T:Microsoft.Practices.Unity.GenericParameter"/> instance that specifies
+            that the given named generic parameter should be resolved.
+            </summary>
+            <param name="genericParameterName">The generic parameter name to resolve.</param>
+        </member>
+        <member name="M:Microsoft.Practices.Unity.OptionalGenericParameter.#ctor(System.String,System.String)">
+            <summary>
+            Create a new <see cref="T:Microsoft.Practices.Unity.GenericParameter"/> instance that specifies
+            that the given named generic parameter should be resolved.
+            </summary>
+            <param name="genericParameterName">The generic parameter name to resolve.</param>
+            <param name="resolutionKey">name to use when looking up in the container.</param>
+        </member>
+        <member name="M:Microsoft.Practices.Unity.OptionalGenericParameter.DoGetResolverPolicy(System.Type,System.String)">
+            <summary>
+            Return a <see cref="T:Microsoft.Practices.ObjectBuilder2.IDependencyResolverPolicy"/> instance that will
+            return this types value for the parameter.
+            </summary>
+            <param name="typeToResolve">The actual type to resolve.</param>
+            <param name="resolutionKey">The resolution key.</param>
+            <returns>The <see cref="T:Microsoft.Practices.ObjectBuilder2.IDependencyResolverPolicy"/>.</returns>
+        </member>
+        <member name="T:Microsoft.Practices.Unity.InjectionFactory">
+            <summary>
+            A class that lets you specify a factory method the container
+            will use to create the object.
+            </summary>
+            <remarks>This is a significantly easier way to do the same
+            thing the old static factory extension was used for.</remarks>
+        </member>
+        <member name="T:Microsoft.Practices.Unity.InjectionMember">
+            <summary>
+            Base class for objects that can be used to configure what
+            class members get injected by the container.
+            </summary>
+        </member>
+        <member name="M:Microsoft.Practices.Unity.InjectionMember.AddPolicies(System.Type,Microsoft.Practices.ObjectBuilder2.IPolicyList)">
+            <summary>
+            Add policies to the <paramref name="policies"/> to configure the
+            container to call this constructor with the appropriate parameter values.
+            </summary>
+            <param name="typeToCreate">Type to register.</param>
+            <param name="policies">Policy list to add policies to.</param>
+        </member>
+        <member name="M:Microsoft.Practices.Unity.InjectionMember.AddPolicies(System.Type,System.Type,System.String,Microsoft.Practices.ObjectBuilder2.IPolicyList)">
+            <summary>
+            Add policies to the <paramref name="policies"/> to configure the
+            container to call this constructor with the appropriate parameter values.
+            </summary>
+            <param name="serviceType">Type of interface being registered. If no interface,
+            this will be null.</param>
+            <param name="implementationType">Type of concrete type being registered.</param>
+            <param name="name">Name used to resolve the type object.</param>
+            <param name="policies">Policy list to add policies to.</param>
+        </member>
+        <member name="M:Microsoft.Practices.Unity.InjectionFactory.#ctor(System.Func{Microsoft.Practices.Unity.IUnityContainer,System.Object})">
+            <summary>
+            Create a new instance of <see cref="T:Microsoft.Practices.Unity.InjectionFactory"/> with
+            the given factory function.
+            </summary>
+            <param name="factoryFunc">Factory function.</param>
+        </member>
+        <member name="M:Microsoft.Practices.Unity.InjectionFactory.#ctor(System.Func{Microsoft.Practices.Unity.IUnityContainer,System.Type,System.String,System.Object})">
+            <summary>
+            Create a new instance of <see cref="T:Microsoft.Practices.Unity.InjectionFactory"/> with
+            the given factory function.
+            </summary>
+            <param name="factoryFunc">Factory function.</param>
+        </member>
+        <member name="M:Microsoft.Practices.Unity.InjectionFactory.AddPolicies(System.Type,System.Type,System.String,Microsoft.Practices.ObjectBuilder2.IPolicyList)">
+            <summary>
+            Add policies to the <paramref name="policies"/> to configure the
+            container to call this constructor with the appropriate parameter values.
+            </summary>
+            <param name="serviceType">Type of interface being registered. If no interface,
+            this will be null. This parameter is ignored in this implementation.</param>
+            <param name="implementationType">Type of concrete type being registered.</param>
+            <param name="name">Name used to resolve the type object.</param>
+            <param name="policies">Policy list to add policies to.</param>
+        </member>
+        <member name="T:Microsoft.Practices.Unity.OptionalParameter">
+            <summary>
+            A <see cref="T:Microsoft.Practices.Unity.InjectionParameterValue"/> that can be passed to
+            <see cref="M:Microsoft.Practices.Unity.IUnityContainer.RegisterType(System.Type,System.Type,System.String,Microsoft.Practices.Unity.LifetimeManager,Microsoft.Practices.Unity.InjectionMember[])"/> to configure a
+            parameter or property as an optional dependency.
+            </summary>
+        </member>
+        <member name="T:Microsoft.Practices.Unity.TypedInjectionValue">
+            <summary>
+            A base class for implementing <see cref="T:Microsoft.Practices.Unity.InjectionParameterValue"/> classes
+            that deal in explicit types.
+            </summary>
+        </member>
+        <member name="M:Microsoft.Practices.Unity.TypedInjectionValue.#ctor(System.Type)">
+            <summary>
+            Create a new <see cref="T:Microsoft.Practices.Unity.TypedInjectionValue"/> that exposes
+            information about the given <paramref name="parameterType"/>.
+            </summary>
+            <param name="parameterType">Type of the parameter.</param>
+        </member>
+        <member name="M:Microsoft.Practices.Unity.TypedInjectionValue.MatchesType(System.Type)">
+            <summary>
+            Test to see if this parameter value has a matching type for the given type.
+            </summary>
+            <param name="t">Type to check.</param>
+            <returns>True if this parameter value is compatible with type <paramref name="t"/>,
+            false if not.</returns>
+        </member>
+        <member name="P:Microsoft.Practices.Unity.TypedInjectionValue.ParameterType">
+            <summary>
+            The type of parameter this object represents.
+            </summary>
+        </member>
+        <member name="P:Microsoft.Practices.Unity.TypedInjectionValue.ParameterTypeName">
+            <summary>
+            Name for the type represented by this <see cref="T:Microsoft.Practices.Unity.InjectionParameterValue"/>.
+            This may be an actual type name or a generic argument name.
+            </summary>
+        </member>
+        <member name="M:Microsoft.Practices.Unity.OptionalParameter.#ctor(System.Type)">
+            <summary>
+            Construct a new <see cref="T:Microsoft.Practices.Unity.OptionalParameter"/> object that
+            specifies the given <paramref name="type"/>.
+            </summary>
+            <param name="type">Type of the dependency.</param>
+        </member>
+        <member name="M:Microsoft.Practices.Unity.OptionalParameter.#ctor(System.Type,System.String)">
+            <summary>
+            Construct a new <see cref="T:Microsoft.Practices.Unity.OptionalParameter"/> object that
+            specifies the given <paramref name="type"/> and <paramref name="name"/>.
+            </summary>
+            <param name="type">Type of the dependency.</param>
+            <param name="name">Name for the dependency.</param>
+        </member>
+        <member name="M:Microsoft.Practices.Unity.OptionalParameter.GetResolverPolicy(System.Type)">
+            <summary>
+            Return a <see cref="T:Microsoft.Practices.ObjectBuilder2.IDependencyResolverPolicy"/> instance that will
+            return this types value for the parameter.
+            </summary>
+            <param name="typeToBuild">Type that contains the member that needs this parameter. Used
+            to resolve open generic parameters.</param>
+            <returns>The <see cref="T:Microsoft.Practices.ObjectBuilder2.IDependencyResolverPolicy"/>.</returns>
+        </member>
+        <member name="T:Microsoft.Practices.Unity.OptionalParameter`1">
+            <summary>
+            A generic version of <see cref="T:Microsoft.Practices.Unity.OptionalParameter"></see> that lets you
+            specify the type of the dependency using generics syntax.
+            </summary>
+            <typeparam name="T">Type of the dependency.</typeparam>
+        </member>
+        <member name="M:Microsoft.Practices.Unity.OptionalParameter`1.#ctor">
+            <summary>
+            Construct a new <see cref="T:Microsoft.Practices.Unity.OptionalParameter`1"/>.
+            </summary>
+        </member>
+        <member name="M:Microsoft.Practices.Unity.OptionalParameter`1.#ctor(System.String)">
+            <summary>
+            Construct a new <see cref="T:Microsoft.Practices.Unity.OptionalParameter`1"/> with the given
+            <paramref name="name"/>.
+            </summary>
+            <param name="name">Name of the dependency.</param>
+        </member>
+        <member name="T:Microsoft.Practices.Unity.HierarchicalLifetimeManager">
+            <summary>
+            A special lifetime manager which works like <see cref="T:Microsoft.Practices.Unity.ContainerControlledLifetimeManager"/>,
+            except that in the presence of child containers, each child gets it's own instance
+            of the object, instead of sharing one in the common parent.
+            </summary>
+        </member>
+        <member name="T:Microsoft.Practices.Unity.ContainerControlledLifetimeManager">
+            <summary>
+            A <see cref="T:Microsoft.Practices.Unity.LifetimeManager"/> that holds onto the instance given to it.
+            When the <see cref="T:Microsoft.Practices.Unity.ContainerControlledLifetimeManager"/> is disposed,
+            the instance is disposed with it.
+            </summary>
+        </member>
+        <member name="T:Microsoft.Practices.Unity.SynchronizedLifetimeManager">
+            <summary>
+            Base class for Lifetime managers which need to synchronize calls to
+            <see cref="M:Microsoft.Practices.Unity.SynchronizedLifetimeManager.GetValue"/>.
+            </summary>
+            <remarks>
+            <para>
+            The purpose of this class is to provide a basic implementation of the lifetime manager synchronization pattern.
+            </para>
+            <para>
+            Calls to the <see cref="M:Microsoft.Practices.Unity.SynchronizedLifetimeManager.GetValue"/> method of a <see cref="T:Microsoft.Practices.Unity.SynchronizedLifetimeManager"/> 
+            instance acquire a lock, and if the instance has not been initialized with a value yet the lock will only be released 
+            when such an initialization takes place by calling the <see cref="M:Microsoft.Practices.Unity.SynchronizedLifetimeManager.SetValue(System.Object)"/> method or if 
+            the build request which resulted in the call to the GetValue method fails.
+            </para>
+            </remarks>
+            <see cref="T:Microsoft.Practices.Unity.LifetimeManager"/>
+        </member>
+        <member name="T:Microsoft.Practices.Unity.LifetimeManager">
+            <summary>
+            Base class for Lifetime managers - classes that control how
+            and when instances are created by the Unity container.
+            </summary>
+        </member>
+        <member name="T:Microsoft.Practices.ObjectBuilder2.ILifetimePolicy">
+            <summary>
+            A <see cref="T:Microsoft.Practices.ObjectBuilder2.IBuilderPolicy"/> that controls how instances are
+            persisted and recovered from an external store. Used to implement
+            things like singletons and per-http-request lifetime.
+            </summary>
+        </member>
+        <member name="M:Microsoft.Practices.ObjectBuilder2.ILifetimePolicy.GetValue">
+            <summary>
+            Retrieve a value from the backing store associated with this Lifetime policy.
+            </summary>
+            <returns>the object desired, or null if no such object is currently stored.</returns>
+        </member>
+        <member name="M:Microsoft.Practices.ObjectBuilder2.ILifetimePolicy.SetValue(System.Object)">
+            <summary>
+            Stores the given value into backing store for retrieval later.
+            </summary>
+            <param name="newValue">The object to store.</param>
+        </member>
+        <member name="M:Microsoft.Practices.ObjectBuilder2.ILifetimePolicy.RemoveValue">
+            <summary>
+            Remove the value this lifetime policy is managing from backing store.
+            </summary>
+        </member>
+        <member name="M:Microsoft.Practices.Unity.LifetimeManager.GetValue">
+            <summary>
+            Retrieve a value from the backing store associated with this Lifetime policy.
+            </summary>
+            <returns>the object desired, or null if no such object is currently stored.</returns>
+        </member>
+        <member name="M:Microsoft.Practices.Unity.LifetimeManager.SetValue(System.Object)">
+            <summary>
+            Stores the given value into backing store for retrieval later.
+            </summary>
+            <param name="newValue">The object being stored.</param>
+        </member>
+        <member name="M:Microsoft.Practices.Unity.LifetimeManager.RemoveValue">
+            <summary>
+            Remove the given object from backing store.
+            </summary>
+        </member>
+        <member name="T:Microsoft.Practices.ObjectBuilder2.IRequiresRecovery">
+            <summary>
+            This interface provides a hook for the builder context to
+            implement error recovery when a builder strategy throws
+            an exception. Since we can't get try/finally blocks onto
+            the call stack for later stages in the chain, we instead
+            add these objects to the context. If there's an exception,
+            all the current IRequiresRecovery instances will have
+            their Recover methods called.
+            </summary>
+        </member>
+        <member name="M:Microsoft.Practices.ObjectBuilder2.IRequiresRecovery.Recover">
+            <summary>
+            A method that does whatever is needed to clean up
+            as part of cleaning up after an exception.
+            </summary>
+            <remarks>
+            Don't do anything that could throw in this method,
+            it will cause later recover operations to get skipped
+            and play real havoc with the stack trace.
+            </remarks>
+        </member>
+        <member name="M:Microsoft.Practices.Unity.SynchronizedLifetimeManager.GetValue">
+            <summary>
+            Retrieve a value from the backing store associated with this Lifetime policy.
+            </summary>
+            <returns>the object desired, or null if no such object is currently stored.</returns>
+            <remarks>Calls to this method acquire a lock which is released only if a non-null value
+            has been set for the lifetime manager.</remarks>
+        </member>
+        <member name="M:Microsoft.Practices.Unity.SynchronizedLifetimeManager.SynchronizedGetValue">
+            <summary>
+            Performs the actual retrieval of a value from the backing store associated 
+            with this Lifetime policy.
+            </summary>
+            <returns>the object desired, or null if no such object is currently stored.</returns>
+            <remarks>This method is invoked by <see cref="M:Microsoft.Practices.Unity.SynchronizedLifetimeManager.GetValue"/>
+            after it has acquired its lock.</remarks>
+        </member>
+        <member name="M:Microsoft.Practices.Unity.SynchronizedLifetimeManager.SetValue(System.Object)">
+            <summary>
+            Stores the given value into backing store for retrieval later.
+            </summary>
+            <param name="newValue">The object being stored.</param>
+            <remarks>Setting a value will attempt to release the lock acquired by 
+            <see cref="M:Microsoft.Practices.Unity.SynchronizedLifetimeManager.GetValue"/>.</remarks>
+        </member>
+        <member name="M:Microsoft.Practices.Unity.SynchronizedLifetimeManager.SynchronizedSetValue(System.Object)">
+            <summary>
+            Performs the actual storage of the given value into backing store for retrieval later.
+            </summary>
+            <param name="newValue">The object being stored.</param>
+            <remarks>This method is invoked by <see cref="M:Microsoft.Practices.Unity.SynchronizedLifetimeManager.SetValue(System.Object)"/>
+            before releasing its lock.</remarks>
+        </member>
+        <member name="M:Microsoft.Practices.Unity.SynchronizedLifetimeManager.RemoveValue">
+            <summary>
+            Remove the given object from backing store.
+            </summary>
+        </member>
+        <member name="M:Microsoft.Practices.Unity.SynchronizedLifetimeManager.Recover">
+            <summary>
+            A method that does whatever is needed to clean up
+            as part of cleaning up after an exception.
+            </summary>
+            <remarks>
+            Don't do anything that could throw in this method,
+            it will cause later recover operations to get skipped
+            and play real havoc with the stack trace.
+            </remarks>
+        </member>
+        <member name="M:Microsoft.Practices.Unity.ContainerControlledLifetimeManager.SynchronizedGetValue">
+            <summary>
+            Retrieve a value from the backing store associated with this Lifetime policy.
+            </summary>
+            <returns>the object desired, or null if no such object is currently stored.</returns>
+        </member>
+        <member name="M:Microsoft.Practices.Unity.ContainerControlledLifetimeManager.SynchronizedSetValue(System.Object)">
+            <summary>
+            Stores the given value into backing store for retrieval later.
+            </summary>
+            <param name="newValue">The object being stored.</param>
+        </member>
+        <member name="M:Microsoft.Practices.Unity.ContainerControlledLifetimeManager.RemoveValue">
+            <summary>
+            Remove the given object from backing store.
+            </summary>
+        </member>
+        <member name="M:Microsoft.Practices.Unity.ContainerControlledLifetimeManager.Dispose">
+            <summary>
+            Performs application-defined tasks associated with freeing, releasing, or resetting unmanaged resources.
+            </summary>
+        </member>
+        <member name="M:Microsoft.Practices.Unity.ContainerControlledLifetimeManager.Dispose(System.Boolean)">
+            <summary>
+            Standard Dispose pattern implementation. Not needed, but it keeps FxCop happy.
+            </summary>
+            <param name="disposing">Always true, since we don't have a finalizer.</param>
+        </member>
+        <member name="T:Microsoft.Practices.Unity.PerResolveLifetimeManager">
+            <summary>
+            This is a custom lifetime manager that acts like <see cref="T:Microsoft.Practices.Unity.TransientLifetimeManager"/>,
+            but also provides a signal to the default build plan, marking the type so that
+            instances are reused across the build up object graph.
+            </summary>
+        </member>
+        <member name="M:Microsoft.Practices.Unity.PerResolveLifetimeManager.#ctor">
+            <summary>
+            Construct a new <see cref="T:Microsoft.Practices.Unity.PerResolveLifetimeManager"/> object that does not
+            itself manage an instance.
+            </summary>
+        </member>
+        <member name="M:Microsoft.Practices.Unity.PerResolveLifetimeManager.#ctor(System.Object)">
+            <summary>
+            Construct a new <see cref="T:Microsoft.Practices.Unity.PerResolveLifetimeManager"/> object that stores the
+            give value. This value will be returned by <see cref="M:Microsoft.Practices.Unity.LifetimeManager.GetValue"/>
+            but is not stored in the lifetime manager, nor is the value disposed.
+            This Lifetime manager is intended only for internal use, which is why the
+            normal <see cref="M:Microsoft.Practices.Unity.LifetimeManager.SetValue(System.Object)"/> method is not used here.
+            </summary>
+            <param name="value">Value to store.</param>
+        </member>
+        <member name="M:Microsoft.Practices.Unity.PerResolveLifetimeManager.GetValue">
+            <summary>
+            Retrieve a value from the backing store associated with this Lifetime policy.
+            </summary>
+            <returns>the object desired, or null if no such object is currently stored.</returns>
+        </member>
+        <member name="M:Microsoft.Practices.Unity.PerResolveLifetimeManager.SetValue(System.Object)">
+            <summary>
+            Stores the given value into backing store for retrieval later. In this class,
+            this is a noop, since it has special hooks down in the guts.
+            </summary>
+            <param name="newValue">The object being stored.</param>
+        </member>
+        <member name="M:Microsoft.Practices.Unity.PerResolveLifetimeManager.RemoveValue">
+            <summary>
+            Remove the given object from backing store. Noop in this class.
+            </summary>
+        </member>
+        <member name="T:Microsoft.Practices.Unity.HierarchicalLifetimeStrategy">
+            <summary>
+            A strategy that handles Hierarchical lifetimes across a set of parent/child
+            containers.
+            </summary>
+        </member>
+        <member name="T:Microsoft.Practices.ObjectBuilder2.BuilderStrategy">
+            <summary>
+            Represents a strategy in the chain of responsibility.
+            Strategies are required to support both BuildUp and TearDown.
+            </summary>
+        </member>
+        <member name="T:Microsoft.Practices.ObjectBuilder2.IBuilderStrategy">
+            <summary>
+            Represents a strategy in the chain of responsibility.
+            Strategies are required to support both BuildUp and TearDown. Although you
+            can implement this interface directly, you may also choose to use
+            <see cref="T:Microsoft.Practices.ObjectBuilder2.BuilderStrategy"/> as the base class for your strategies, as
+            this class provides useful helper methods and makes support BuildUp and TearDown
+            optional.
+            </summary>
+        </member>
+        <member name="M:Microsoft.Practices.ObjectBuilder2.IBuilderStrategy.PreBuildUp(Microsoft.Practices.ObjectBuilder2.IBuilderContext)">
+            <summary>
+            Called during the chain of responsibility for a build operation. The
+            PreBuildUp method is called when the chain is being executed in the
+            forward direction.
+            </summary>
+            <param name="context">Context of the build operation.</param>
+        </member>
+        <member name="M:Microsoft.Practices.ObjectBuilder2.IBuilderStrategy.PostBuildUp(Microsoft.Practices.ObjectBuilder2.IBuilderContext)">
+            <summary>
+            Called during the chain of responsibility for a build operation. The
+            PostBuildUp method is called when the chain has finished the PreBuildUp
+            phase and executes in reverse order from the PreBuildUp calls.
+            </summary>
+            <param name="context">Context of the build operation.</param>
+        </member>
+        <member name="M:Microsoft.Practices.ObjectBuilder2.IBuilderStrategy.PreTearDown(Microsoft.Practices.ObjectBuilder2.IBuilderContext)">
+            <summary>
+            Called during the chain of responsibility for a teardown operation. The
+            PreTearDown method is called when the chain is being executed in the
+            forward direction.
+            </summary>
+            <param name="context">Context of the teardown operation.</param>
+        </member>
+        <member name="M:Microsoft.Practices.ObjectBuilder2.IBuilderStrategy.PostTearDown(Microsoft.Practices.ObjectBuilder2.IBuilderContext)">
+            <summary>
+            Called during the chain of responsibility for a teardown operation. The
+            PostTearDown method is called when the chain has finished the PreTearDown
+            phase and executes in reverse order from the PreTearDown calls.
+            </summary>
+            <param name="context">Context of the teardown operation.</param>
+        </member>
+        <member name="M:Microsoft.Practices.ObjectBuilder2.BuilderStrategy.PreBuildUp(Microsoft.Practices.ObjectBuilder2.IBuilderContext)">
+            <summary>
+            Called during the chain of responsibility for a build operation. The
+            PreBuildUp method is called when the chain is being executed in the
+            forward direction.
+            </summary>
+            <param name="context">Context of the build operation.</param>
+        </member>
+        <member name="M:Microsoft.Practices.ObjectBuilder2.BuilderStrategy.PostBuildUp(Microsoft.Practices.ObjectBuilder2.IBuilderContext)">
+            <summary>
+            Called during the chain of responsibility for a build operation. The
+            PostBuildUp method is called when the chain has finished the PreBuildUp
+            phase and executes in reverse order from the PreBuildUp calls.
+            </summary>
+            <param name="context">Context of the build operation.</param>
+        </member>
+        <member name="M:Microsoft.Practices.ObjectBuilder2.BuilderStrategy.PreTearDown(Microsoft.Practices.ObjectBuilder2.IBuilderContext)">
+            <summary>
+            Called during the chain of responsibility for a teardown operation. The
+            PreTearDown method is called when the chain is being executed in the
+            forward direction.
+            </summary>
+            <param name="context">Context of the teardown operation.</param>
+        </member>
+        <member name="M:Microsoft.Practices.ObjectBuilder2.BuilderStrategy.PostTearDown(Microsoft.Practices.ObjectBuilder2.IBuilderContext)">
+            <summary>
+            Called during the chain of responsibility for a teardown operation. The
+            PostTearDown method is called when the chain has finished the PreTearDown
+            phase and executes in reverse order from the PreTearDown calls.
+            </summary>
+            <param name="context">Context of the teardown operation.</param>
+        </member>
+        <member name="M:Microsoft.Practices.Unity.HierarchicalLifetimeStrategy.PreBuildUp(Microsoft.Practices.ObjectBuilder2.IBuilderContext)">
+            <summary>
+            Called during the chain of responsibility for a build operation. The
+            PreBuildUp method is called when the chain is being executed in the
+            forward direction.
+            </summary>
+            <param name="context">Context of the build operation.</param>
+        </member>
+        <member name="T:Microsoft.Practices.Unity.OptionalDependencyResolverPolicy">
+            <summary>
+            A <see cref="T:Microsoft.Practices.ObjectBuilder2.IDependencyResolverPolicy"/> that will attempt to
+            resolve a value, and return null if it cannot rather than throwing.
+            </summary>
+        </member>
+        <member name="T:Microsoft.Practices.ObjectBuilder2.IDependencyResolverPolicy">
+            <summary>
+            A strategy that is used at build plan execution time
+            to resolve a dependent value.
+            </summary>
+        </member>
+        <member name="M:Microsoft.Practices.ObjectBuilder2.IDependencyResolverPolicy.Resolve(Microsoft.Practices.ObjectBuilder2.IBuilderContext)">
+            <summary>
+            Get the value for a dependency.
+            </summary>
+            <param name="context">Current build context.</param>
+            <returns>The value for the dependency.</returns>
+        </member>
+        <member name="M:Microsoft.Practices.Unity.OptionalDependencyResolverPolicy.#ctor(System.Type,System.String)">
+            <summary>
+            Construct a new <see cref="T:Microsoft.Practices.Unity.OptionalDependencyResolverPolicy"/> object
+            that will attempt to resolve the given name and type from the container.
+            </summary>
+            <param name="type">Type to resolve. Must be a reference type.</param>
+            <param name="name">Name to resolve with.</param>
+        </member>
+        <member name="M:Microsoft.Practices.Unity.OptionalDependencyResolverPolicy.#ctor(System.Type)">
+            <summary>
+            Construct a new <see cref="T:Microsoft.Practices.Unity.OptionalDependencyResolverPolicy"/> object
+            that will attempt to resolve the given type from the container.
+            </summary>
+            <param name="type">Type to resolve. Must be a reference type.</param>
+        </member>
+        <member name="M:Microsoft.Practices.Unity.OptionalDependencyResolverPolicy.Resolve(Microsoft.Practices.ObjectBuilder2.IBuilderContext)">
+            <summary>
+            Get the value for a dependency.
+            </summary>
+            <param name="context">Current build context.</param>
+            <returns>The value for the dependency.</returns>
+        </member>
+        <member name="P:Microsoft.Practices.Unity.OptionalDependencyResolverPolicy.DependencyType">
+            <summary>
+            Type this resolver will resolve.
+            </summary>
+        </member>
+        <member name="P:Microsoft.Practices.Unity.OptionalDependencyResolverPolicy.Name">
+            <summary>
+            Name this resolver will resolve.
+            </summary>
+        </member>
+        <member name="T:Microsoft.Practices.ObjectBuilder2.PolicyListExtensions">
+            <summary>
+            Extension methods on <see cref="T:Microsoft.Practices.ObjectBuilder2.IPolicyList"/> to provide convenience
+            overloads (generic versions, mostly).
+            </summary>
+        </member>
+        <member name="M:Microsoft.Practices.ObjectBuilder2.PolicyListExtensions.Clear``1(Microsoft.Practices.ObjectBuilder2.IPolicyList,System.Object)">
+            <summary>
+            Removes an individual policy type for a build key.
+            </summary>
+            <typeparam name="TPolicyInterface">The type the policy was registered as.</typeparam>
+            <param name="policies"><see cref="T:Microsoft.Practices.ObjectBuilder2.IPolicyList"/> to remove the policy from.</param>
+            <param name="buildKey">The key the policy applies.</param>
+        </member>
+        <member name="M:Microsoft.Practices.ObjectBuilder2.PolicyListExtensions.ClearDefault``1(Microsoft.Practices.ObjectBuilder2.IPolicyList)">
+            <summary>
+            Removes a default policy.
+            </summary>
+            <typeparam name="TPolicyInterface">The type the policy was registered as.</typeparam>
+            <param name="policies"><see cref="T:Microsoft.Practices.ObjectBuilder2.IPolicyList"/> to remove the policy from.</param>
+        </member>
+        <member name="M:Microsoft.Practices.ObjectBuilder2.PolicyListExtensions.Get``1(Microsoft.Practices.ObjectBuilder2.IPolicyList,System.Object)">
+            <summary>
+            Gets an individual policy.
+            </summary>
+            <typeparam name="TPolicyInterface">The interface the policy is registered under.</typeparam>
+            <param name="policies"><see cref="T:Microsoft.Practices.ObjectBuilder2.IPolicyList"/> to search.</param>
+            <param name="buildKey">The key the policy applies.</param>
+            <returns>The policy in the list, if present; returns null otherwise.</returns>
+        </member>
+        <member name="M:Microsoft.Practices.ObjectBuilder2.PolicyListExtensions.Get``1(Microsoft.Practices.ObjectBuilder2.IPolicyList,System.Object,Microsoft.Practices.ObjectBuilder2.IPolicyList@)">
+            <summary>
+            Gets an individual policy.
+            </summary>
+            <typeparam name="TPolicyInterface">The interface the policy is registered under.</typeparam>
+            <param name="policies"><see cref="T:Microsoft.Practices.ObjectBuilder2.IPolicyList"/> to search.</param>
+            <param name="buildKey">The key the policy applies.</param>
+            <param name="containingPolicyList">The policy list that actually contains the returned policy.</param>
+            <returns>The policy in the list, if present; returns null otherwise.</returns>
+        </member>
+        <member name="M:Microsoft.Practices.ObjectBuilder2.PolicyListExtensions.Get(Microsoft.Practices.ObjectBuilder2.IPolicyList,System.Type,System.Object)">
+            <summary>
+            Gets an individual policy.
+            </summary>
+            <param name="policies"><see cref="T:Microsoft.Practices.ObjectBuilder2.IPolicyList"/> to search.</param>
+            <param name="policyInterface">The interface the policy is registered under.</param>
+            <param name="buildKey">The key the policy applies.</param>
+            <returns>The policy in the list, if present; returns null otherwise.</returns>
+        </member>
+        <member name="M:Microsoft.Practices.ObjectBuilder2.PolicyListExtensions.Get(Microsoft.Practices.ObjectBuilder2.IPolicyList,System.Type,System.Object,Microsoft.Practices.ObjectBuilder2.IPolicyList@)">
+            <summary>
+            Gets an individual policy.
+            </summary>
+            <param name="policies"><see cref="T:Microsoft.Practices.ObjectBuilder2.IPolicyList"/> to search.</param>
+            <param name="policyInterface">The interface the policy is registered under.</param>
+            <param name="buildKey">The key the policy applies.</param>
+            <param name="containingPolicyList">The policy list that actually contains the returned policy.</param>
+            <returns>The policy in the list, if present; returns null otherwise.</returns>
+        </member>
+        <member name="M:Microsoft.Practices.ObjectBuilder2.PolicyListExtensions.Get``1(Microsoft.Practices.ObjectBuilder2.IPolicyList,System.Object,System.Boolean)">
+            <summary>
+            Gets an individual policy.
+            </summary>
+            <typeparam name="TPolicyInterface">The interface the policy is registered under.</typeparam>
+            <param name="policies"><see cref="T:Microsoft.Practices.ObjectBuilder2.IPolicyList"/> to search.</param>
+            <param name="buildKey">The key the policy applies.</param>
+            <param name="localOnly">true if the policy searches local only; otherwise false to search up the parent chain.</param>
+            <returns>The policy in the list, if present; returns null otherwise.</returns>
+        </member>
+        <member name="M:Microsoft.Practices.ObjectBuilder2.PolicyListExtensions.Get``1(Microsoft.Practices.ObjectBuilder2.IPolicyList,System.Object,System.Boolean,Microsoft.Practices.ObjectBuilder2.IPolicyList@)">
+            <summary>
+            Gets an individual policy.
+            </summary>
+            <typeparam name="TPolicyInterface">The interface the policy is registered under.</typeparam>
+            <param name="policies"><see cref="T:Microsoft.Practices.ObjectBuilder2.IPolicyList"/> to search.</param>
+            <param name="buildKey">The key the policy applies.</param>
+            <param name="localOnly">true if the policy searches local only; otherwise false to search up the parent chain.</param>
+            <param name="containingPolicyList">The policy list that actually contains the returned policy.</param>
+            <returns>The policy in the list, if present; returns null otherwise.</returns>
+        </member>
+        <member name="M:Microsoft.Practices.ObjectBuilder2.PolicyListExtensions.Get(Microsoft.Practices.ObjectBuilder2.IPolicyList,System.Type,System.Object,System.Boolean)">
+            <summary>
+            Gets an individual policy.
+            </summary>
+            <param name="policies"><see cref="T:Microsoft.Practices.ObjectBuilder2.IPolicyList"/> to search.</param>
+            <param name="policyInterface">The interface the policy is registered under.</param>
+            <param name="buildKey">The key the policy applies.</param>
+            <param name="localOnly">true if the policy searches local only; otherwise false to search up the parent chain.</param>
+            <returns>The policy in the list, if present; returns null otherwise.</returns>
+        </member>
+        <member name="M:Microsoft.Practices.ObjectBuilder2.PolicyListExtensions.GetNoDefault``1(Microsoft.Practices.ObjectBuilder2.IPolicyList,System.Object,System.Boolean)">
+            <summary>
+            Get the non default policy.
+            </summary>
+            <typeparam name="TPolicyInterface">The interface the policy is registered under.</typeparam>
+            <param name="policies"><see cref="T:Microsoft.Practices.ObjectBuilder2.IPolicyList"/> to search.</param>
+            <param name="buildKey">The key the policy applies.</param>
+            <param name="localOnly">true if the policy searches local only; otherwise false to search up the parent chain.</param>
+            <returns>The policy in the list, if present; returns null otherwise.</returns>
+        </member>
+        <member name="M:Microsoft.Practices.ObjectBuilder2.PolicyListExtensions.GetNoDefault``1(Microsoft.Practices.ObjectBuilder2.IPolicyList,System.Object,System.Boolean,Microsoft.Practices.ObjectBuilder2.IPolicyList@)">
+            <summary>
+            Get the non default policy.
+            </summary>
+            <typeparam name="TPolicyInterface">The interface the policy is registered under.</typeparam>
+            <param name="policies"><see cref="T:Microsoft.Practices.ObjectBuilder2.IPolicyList"/> to search.</param>
+            <param name="buildKey">The key the policy applies.</param>
+            <param name="localOnly">true if the policy searches local only; otherwise false to search up the parent chain.</param>
+            <param name="containingPolicyList">The policy list that actually contains the returned policy.</param>
+            <returns>The policy in the list, if present; returns null otherwise.</returns>
+        </member>
+        <member name="M:Microsoft.Practices.ObjectBuilder2.PolicyListExtensions.GetNoDefault(Microsoft.Practices.ObjectBuilder2.IPolicyList,System.Type,System.Object,System.Boolean)">
+            <summary>
+            Get the non default policy.
+            </summary>
+            <param name="policies"><see cref="T:Microsoft.Practices.ObjectBuilder2.IPolicyList"/> to search.</param>
+            <param name="policyInterface">The interface the policy is registered under.</param>
+            <param name="buildKey">The key the policy applies.</param>
+            <param name="localOnly">true if the policy searches local only; otherwise false to search up the parent chain.</param>
+            <returns>The policy in the list, if present; returns null otherwise.</returns>
+        </member>
+        <member name="M:Microsoft.Practices.ObjectBuilder2.PolicyListExtensions.Set``1(Microsoft.Practices.ObjectBuilder2.IPolicyList,``0,System.Object)">
+            <summary>
+            Sets an individual policy.
+            </summary>
+            <typeparam name="TPolicyInterface">The interface the policy is registered under.</typeparam>
+            <param name="policies"><see cref="T:Microsoft.Practices.ObjectBuilder2.IPolicyList"/> to add the policy to.</param>
+            <param name="policy">The policy to be registered.</param>
+            <param name="buildKey">The key the policy applies.</param>
+        </member>
+        <member name="M:Microsoft.Practices.ObjectBuilder2.PolicyListExtensions.SetDefault``1(Microsoft.Practices.ObjectBuilder2.IPolicyList,``0)">
+            <summary>
+            Sets a default policy. When checking for a policy, if no specific individual policy
+            is available, the default will be used.
+            </summary>
+            <typeparam name="TPolicyInterface">The interface to register the policy under.</typeparam>
+            <param name="policies"><see cref="T:Microsoft.Practices.ObjectBuilder2.IPolicyList"/> to add the policy to.</param>
+            <param name="policy">The default policy to be registered.</param>
+        </member>
+        <member name="T:Microsoft.Practices.ObjectBuilder2.BuildOperation">
+            <summary>
+            Base class for the current operation stored in the build context.
+            </summary>
+        </member>
+        <member name="M:Microsoft.Practices.ObjectBuilder2.BuildOperation.#ctor(System.Type)">
+            <summary>
+            Create a new <see cref="T:Microsoft.Practices.ObjectBuilder2.BuildOperation"/>.
+            </summary>
+            <param name="typeBeingConstructed">Type currently being built.</param>
+        </member>
+        <member name="P:Microsoft.Practices.ObjectBuilder2.BuildOperation.TypeBeingConstructed">
+            <summary>
+            The type that's currently being built.
+            </summary>
+        </member>
+        <member name="T:Microsoft.Practices.ObjectBuilder2.DeferredResolveBuildPlanPolicy">
+            <summary>
+            Build plan for <see cref="T:System.Func`1"/> that will return a Func that will resolve the requested type
+            through this container later.
+            </summary>
+        </member>
+        <member name="T:Microsoft.Practices.ObjectBuilder2.IBuildPlanPolicy">
+            <summary>
+            A build plan is an object that, when invoked, will create a new object
+            or fill in a given existing one. It encapsulates all the information
+            gathered by the strategies to construct a particular object.
+            </summary>
+        </member>
+        <member name="M:Microsoft.Practices.ObjectBuilder2.IBuildPlanPolicy.BuildUp(Microsoft.Practices.ObjectBuilder2.IBuilderContext)">
+            <summary>
+            Creates an instance of this build plan's type, or fills
+            in the existing type if passed in.
+            </summary>
+            <param name="context">Context used to build up the object.</param>
+        </member>
+        <member name="M:Microsoft.Practices.ObjectBuilder2.OverriddenBuildPlanMarkerPolicy.BuildUp(Microsoft.Practices.ObjectBuilder2.IBuilderContext)">
+            <summary>
+            Creates an instance of this build plan's type, or fills
+            in the existing type if passed in.
+            </summary>
+            <param name="context">Context used to build up the object.</param>
+        </member>
+        <member name="T:Microsoft.Practices.ObjectBuilder2.EnumerableExtensions">
+            <summary>
+            The almost inevitable collection of extra helper methods on
+            <see cref="T:System.Collections.Generic.IEnumerable`1"/> to augment the rich set of what
+            LINQ already gives us.
+            </summary>
+        </member>
+        <member name="M:Microsoft.Practices.ObjectBuilder2.EnumerableExtensions.ForEach``1(System.Collections.Generic.IEnumerable{``0},System.Action{``0})">
+            <summary>
+            Execute the provided <paramref name="action"/> on every item in <paramref name="sequence"/>.
+            </summary>
+            <typeparam name="TItem">Type of the items stored in <paramref name="sequence"/></typeparam>
+            <param name="sequence">Sequence of items to process.</param>
+            <param name="action">Code to run over each item.</param>
+        </member>
+        <member name="M:Microsoft.Practices.ObjectBuilder2.EnumerableExtensions.JoinStrings``1(System.Collections.Generic.IEnumerable{``0},System.String,System.Func{``0,System.String})">
+            <summary>
+            Create a single string from a sequence of items, separated by the provided <paramref name="separator"/>,
+            and with the conversion to string done by the given <paramref name="converter"/>.
+            </summary>
+            <remarks>This method does basically the same thing as <see cref="M:System.String.Join(System.String,System.String[])"/>,
+            but will work on any sequence of items, not just arrays.</remarks>
+            <typeparam name="TItem">Type of items in the sequence.</typeparam>
+            <param name="sequence">Sequence of items to convert.</param>
+            <param name="separator">Separator to place between the items in the string.</param>
+            <param name="converter">The conversion function to change TItem -&gt; string.</param>
+            <returns>The resulting string.</returns>
+        </member>
+        <member name="M:Microsoft.Practices.ObjectBuilder2.EnumerableExtensions.JoinStrings``1(System.Collections.Generic.IEnumerable{``0},System.String)">
+            <summary>
+            Create a single string from a sequence of items, separated by the provided <paramref name="separator"/>,
+            and with the conversion to string done by the item's <see cref="M:System.Object.ToString"/> method.
+            </summary>
+            <remarks>This method does basically the same thing as <see cref="M:System.String.Join(System.String,System.String[])"/>,
+            but will work on any sequence of items, not just arrays.</remarks>
+            <typeparam name="TItem">Type of items in the sequence.</typeparam>
+            <param name="sequence">Sequence of items to convert.</param>
+            <param name="separator">Separator to place between the items in the string.</param>
+            <returns>The resulting string.</returns>
+        </member>
+        <member name="T:Microsoft.Practices.Unity.ParameterOverride">
+            <summary>
+            A <see cref="T:Microsoft.Practices.Unity.ResolverOverride"/> class that lets you
+            override a named parameter passed to a constructor.
+            </summary>
+        </member>
+        <member name="M:Microsoft.Practices.Unity.ParameterOverride.#ctor(System.String,System.Object)">
+            <summary>
+            Construct a new <see cref="T:Microsoft.Practices.Unity.ParameterOverride"/> object that will
+            override the given named constructor parameter, and pass the given
+            value.
+            </summary>
+            <param name="parameterName">Name of the constructor parameter.</param>
+            <param name="parameterValue">Value to pass for the constructor.</param>
+        </member>
+        <member name="M:Microsoft.Practices.Unity.ParameterOverride.GetResolver(Microsoft.Practices.ObjectBuilder2.IBuilderContext,System.Type)">
+            <summary>
+            Return a <see cref="T:Microsoft.Practices.ObjectBuilder2.IDependencyResolverPolicy"/> that can be used to give a value
+            for the given desired dependency.
+            </summary>
+            <param name="context">Current build context.</param>
+            <param name="dependencyType">Type of dependency desired.</param>
+            <returns>a <see cref="T:Microsoft.Practices.ObjectBuilder2.IDependencyResolverPolicy"/> object if this override applies, null if not.</returns>
+        </member>
+        <member name="T:Microsoft.Practices.Unity.ParameterOverrides">
+            <summary>
+            A convenience form of <see cref="T:Microsoft.Practices.Unity.ParameterOverride"/> that lets you
+            specify multiple parameter overrides in one shot rather than having
+            to construct multiple objects.
+            </summary>
+        </member>
+        <member name="M:Microsoft.Practices.Unity.ParameterOverrides.MakeOverride(System.String,System.Object)">
+            <summary>
+            When implemented in derived classes, this method is called from the <see cref="M:Microsoft.Practices.Unity.OverrideCollection`3.Add(`1,`2)"/>
+            method to create the actual <see cref="T:Microsoft.Practices.Unity.ResolverOverride"/> objects.
+            </summary>
+            <param name="key">Key value to create the resolver.</param>
+            <param name="value">Value to store in the resolver.</param>
+            <returns>The created <see cref="T:Microsoft.Practices.Unity.ResolverOverride"/>.</returns>
+        </member>
+        <member name="T:Microsoft.Practices.Unity.PropertyOverride">
+            <summary>
+            A <see cref="T:Microsoft.Practices.Unity.ResolverOverride"/> that lets you override
+            the value for a specified property.
+            </summary>
+        </member>
+        <member name="M:Microsoft.Practices.Unity.PropertyOverride.#ctor(System.String,System.Object)">
+            <summary>
+            Create an instance of <see cref="T:Microsoft.Practices.Unity.PropertyOverride"/>.
+            </summary>
+            <param name="propertyName">The property name.</param>
+            <param name="propertyValue">Value to use for the property.</param>
+        </member>
+        <member name="M:Microsoft.Practices.Unity.PropertyOverride.GetResolver(Microsoft.Practices.ObjectBuilder2.IBuilderContext,System.Type)">
+            <summary>
+            Return a <see cref="T:Microsoft.Practices.ObjectBuilder2.IDependencyResolverPolicy"/> that can be used to give a value
+            for the given desired dependency.
+            </summary>
+            <param name="context">Current build context.</param>
+            <param name="dependencyType">Type of dependency desired.</param>
+            <returns>a <see cref="T:Microsoft.Practices.ObjectBuilder2.IDependencyResolverPolicy"/> object if this override applies, null if not.</returns>
+        </member>
+        <member name="T:Microsoft.Practices.Unity.PropertyOverrides">
+            <summary>
+            A convenience form of <see cref="T:Microsoft.Practices.Unity.PropertyOverride"/> that lets you
+            specify multiple property overrides in one shot rather than having
+            to construct multiple objects.
+            </summary>
+        </member>
+        <member name="M:Microsoft.Practices.Unity.PropertyOverrides.MakeOverride(System.String,System.Object)">
+            <summary>
+            When implemented in derived classes, this method is called from the <see cref="M:Microsoft.Practices.Unity.OverrideCollection`3.Add(`1,`2)"/>
+            method to create the actual <see cref="T:Microsoft.Practices.Unity.ResolverOverride"/> objects.
+            </summary>
+            <param name="key">Key value to create the resolver.</param>
+            <param name="value">Value to store in the resolver.</param>
+            <returns>The created <see cref="T:Microsoft.Practices.Unity.ResolverOverride"/>.</returns>
+        </member>
+        <member name="T:Microsoft.Practices.Unity.StaticFactory.IStaticFactoryConfiguration">
+            <summary>
+            Interface defining the configuration interface exposed by the
+            Static Factory extension.
+            </summary>
+        </member>
+        <member name="T:Microsoft.Practices.Unity.IUnityContainerExtensionConfigurator">
+            <summary>
+            Base interface for all extension configuration interfaces.
+            </summary>
+        </member>
+        <member name="P:Microsoft.Practices.Unity.IUnityContainerExtensionConfigurator.Container">
+            <summary>
+            Retrieve the container instance that we are currently configuring.
+            </summary>
+        </member>
+        <member name="M:Microsoft.Practices.Unity.StaticFactory.IStaticFactoryConfiguration.RegisterFactory``1(System.Func{Microsoft.Practices.Unity.IUnityContainer,System.Object})">
+            <summary>
+            Register the given factory delegate to be called when the container is
+            asked to resolve <typeparamref name="TTypeToBuild"/>.
+            </summary>
+            <typeparam name="TTypeToBuild">Type that will be requested from the container.</typeparam>
+            <param name="factoryMethod">Delegate to invoke to create the instance.</param>
+            <returns>The container extension object this method was invoked on.</returns>
+        </member>
+        <member name="M:Microsoft.Practices.Unity.StaticFactory.IStaticFactoryConfiguration.RegisterFactory``1(System.String,System.Func{Microsoft.Practices.Unity.IUnityContainer,System.Object})">
+            <summary>
+            Register the given factory delegate to be called when the container is
+            asked to resolve <typeparamref name="TTypeToBuild"/> and <paramref name="name"/>.
+            </summary>
+            <typeparam name="TTypeToBuild">Type that will be requested from the container.</typeparam>
+            <param name="name">The name that will be used when requesting to resolve this type.</param>
+            <param name="factoryMethod">Delegate to invoke to create the instance.</param>
+            <returns>The container extension object this method was invoked on.</returns>
+        </member>
+        <member name="T:Microsoft.Practices.ObjectBuilder2.BuilderContext">
+            <summary>
+            Represents the context in which a build-up or tear-down operation runs.
+            </summary>
+        </member>
+        <member name="T:Microsoft.Practices.ObjectBuilder2.IBuilderContext">
+            <summary>
+            Represents the context in which a build-up or tear-down operation runs.
+            </summary>
+        </member>
+        <member name="M:Microsoft.Practices.ObjectBuilder2.IBuilderContext.AddResolverOverrides(System.Collections.Generic.IEnumerable{Microsoft.Practices.Unity.ResolverOverride})">
+            <summary>
+            Add a new set of resolver override objects to the current build operation.
+            </summary>
+            <param name="newOverrides"><see cref="T:Microsoft.Practices.Unity.ResolverOverride"/> objects to add.</param>
+        </member>
+        <member name="M:Microsoft.Practices.ObjectBuilder2.IBuilderContext.GetOverriddenResolver(System.Type)">
+            <summary>
+            Get a <see cref="T:Microsoft.Practices.ObjectBuilder2.IDependencyResolverPolicy"/> object for the given <paramref name="dependencyType"/>
+            or null if that dependency hasn't been overridden.
+            </summary>
+            <param name="dependencyType">Type of the dependency.</param>
+            <returns>Resolver to use, or null if no override matches for the current operation.</returns>
+        </member>
+        <member name="M:Microsoft.Practices.ObjectBuilder2.IBuilderContext.NewBuildUp(Microsoft.Practices.ObjectBuilder2.NamedTypeBuildKey)">
+            <summary>
+            A convenience method to do a new buildup operation on an existing context.
+            </summary>
+            <param name="newBuildKey">Key to use to build up.</param>
+            <returns>Created object.</returns>
+        </member>
+        <member name="M:Microsoft.Practices.ObjectBuilder2.IBuilderContext.NewBuildUp(Microsoft.Practices.ObjectBuilder2.NamedTypeBuildKey,System.Action{Microsoft.Practices.ObjectBuilder2.IBuilderContext})">
+            <summary>
+            A convenience method to do a new buildup operation on an existing context. This
+            overload allows you to specify extra policies which will be in effect for the duration
+            of the build.
+            </summary>
+            <param name="newBuildKey">Key defining what to build up.</param>
+            <param name="childCustomizationBlock">A delegate that takes a <see cref="T:Microsoft.Practices.ObjectBuilder2.IBuilderContext"/>. This
+            is invoked with the new child context before the build up process starts. This gives callers
+            the opportunity to customize the context for the build process.</param>
+            <returns>Created object.</returns>
+        </member>
+        <member name="P:Microsoft.Practices.ObjectBuilder2.IBuilderContext.Strategies">
+            <summary>
+            Gets the head of the strategy chain.
+            </summary>
+            <returns>
+            The strategy that's first in the chain; returns null if there are no
+            strategies in the chain.
+            </returns>
+        </member>
+        <member name="P:Microsoft.Practices.ObjectBuilder2.IBuilderContext.Lifetime">
+            <summary>
+            Gets the <see cref="T:Microsoft.Practices.ObjectBuilder2.ILifetimeContainer"/> associated with the build.
+            </summary>
+            <value>
+            The <see cref="T:Microsoft.Practices.ObjectBuilder2.ILifetimeContainer"/> associated with the build.
+            </value>
+        </member>
+        <member name="P:Microsoft.Practices.ObjectBuilder2.IBuilderContext.OriginalBuildKey">
+            <summary>
+            Gets the original build key for the build operation.
+            </summary>
+            <value>
+            The original build key for the build operation.
+            </value>
+        </member>
+        <member name="P:Microsoft.Practices.ObjectBuilder2.IBuilderContext.BuildKey">
+            <summary>
+            Get the current build key for the current build operation.
+            </summary>
+        </member>
+        <member name="P:Microsoft.Practices.ObjectBuilder2.IBuilderContext.PersistentPolicies">
+            <summary>
+            The set of policies that were passed into this context.
+            </summary>
+            <remarks>This returns the policies passed into the context.
+            Policies added here will remain after buildup completes.</remarks>
+            <value>The persistent policies for the current context.</value>
+        </member>
+        <member name="P:Microsoft.Practices.ObjectBuilder2.IBuilderContext.Policies">
+            <summary>
+            Gets the policies for the current context. 
+            </summary>
+            <remarks>Any policies added to this object are transient
+            and will be erased at the end of the buildup.</remarks>
+            <value>
+            The policies for the current context.
+            </value>
+        </member>
+        <member name="P:Microsoft.Practices.ObjectBuilder2.IBuilderContext.RecoveryStack">
+            <summary>
+            Gets the collection of <see cref="T:Microsoft.Practices.ObjectBuilder2.IRequiresRecovery"/> objects
+            that need to execute in event of an exception.
+            </summary>
+        </member>
+        <member name="P:Microsoft.Practices.ObjectBuilder2.IBuilderContext.Existing">
+            <summary>
+            The current object being built up or torn down.
+            </summary>
+            <value>
+            The current object being manipulated by the build operation. May
+            be null if the object hasn't been created yet.</value>
+        </member>
+        <member name="P:Microsoft.Practices.ObjectBuilder2.IBuilderContext.BuildComplete">
+            <summary>
+            Flag indicating if the build operation should continue.
+            </summary>
+            <value>true means that building should not call any more
+            strategies, false means continue to the next strategy.</value>
+        </member>
+        <member name="P:Microsoft.Practices.ObjectBuilder2.IBuilderContext.CurrentOperation">
+            <summary>
+            An object representing what is currently being done in the
+            build chain. Used to report back errors if there's a failure.
+            </summary>
+        </member>
+        <member name="P:Microsoft.Practices.ObjectBuilder2.IBuilderContext.ChildContext">
+            <summary>
+            The build context used to resolve a dependency during the build operation represented by this context.
+            </summary>
+        </member>
+        <member name="M:Microsoft.Practices.ObjectBuilder2.BuilderContext.#ctor">
+            <summary>
+            Initialize a new instance of the <see cref="T:Microsoft.Practices.ObjectBuilder2.BuilderContext"/> class.
+            </summary>
+        </member>
+        <member name="M:Microsoft.Practices.ObjectBuilder2.BuilderContext.#ctor(Microsoft.Practices.ObjectBuilder2.IStrategyChain,Microsoft.Practices.ObjectBuilder2.ILifetimeContainer,Microsoft.Practices.ObjectBuilder2.IPolicyList,Microsoft.Practices.ObjectBuilder2.NamedTypeBuildKey,System.Object)">
+            <summary>
+            Initialize a new instance of the <see cref="T:Microsoft.Practices.ObjectBuilder2.BuilderContext"/> class with a <see cref="T:Microsoft.Practices.ObjectBuilder2.IStrategyChain"/>, 
+            <see cref="T:Microsoft.Practices.ObjectBuilder2.ILifetimeContainer"/>, <see cref="T:Microsoft.Practices.ObjectBuilder2.IPolicyList"/> and the 
+            build key used to start this build operation. 
+            </summary>
+            <param name="chain">The <see cref="T:Microsoft.Practices.ObjectBuilder2.IStrategyChain"/> to use for this context.</param>
+            <param name="lifetime">The <see cref="T:Microsoft.Practices.ObjectBuilder2.ILifetimeContainer"/> to use for this context.</param>
+            <param name="policies">The <see cref="T:Microsoft.Practices.ObjectBuilder2.IPolicyList"/> to use for this context.</param>
+            <param name="originalBuildKey">Build key to start building.</param>
+            <param name="existing">The existing object to build up.</param>
+        </member>
+        <member name="M:Microsoft.Practices.ObjectBuilder2.BuilderContext.#ctor(Microsoft.Practices.ObjectBuilder2.IStrategyChain,Microsoft.Practices.ObjectBuilder2.ILifetimeContainer,Microsoft.Practices.ObjectBuilder2.IPolicyList,Microsoft.Practices.ObjectBuilder2.IPolicyList,Microsoft.Practices.ObjectBuilder2.NamedTypeBuildKey,System.Object)">
+            <summary>
+            Create a new <see cref="T:Microsoft.Practices.ObjectBuilder2.BuilderContext"/> using the explicitly provided
+            values.
+            </summary>
+            <param name="chain">The <see cref="T:Microsoft.Practices.ObjectBuilder2.IStrategyChain"/> to use for this context.</param>
+            <param name="lifetime">The <see cref="T:Microsoft.Practices.ObjectBuilder2.ILifetimeContainer"/> to use for this context.</param>
+            <param name="persistentPolicies">The set of persistent policies to use for this context.</param>
+            <param name="transientPolicies">The set of transient policies to use for this context. It is
+            the caller's responsibility to ensure that the transient and persistent policies are properly
+            combined.</param>
+            <param name="buildKey">Build key for this context.</param>
+            <param name="existing">Existing object to build up.</param>
+        </member>
+        <member name="M:Microsoft.Practices.ObjectBuilder2.BuilderContext.#ctor(Microsoft.Practices.ObjectBuilder2.IStrategyChain,Microsoft.Practices.ObjectBuilder2.ILifetimeContainer,Microsoft.Practices.ObjectBuilder2.IPolicyList,Microsoft.Practices.ObjectBuilder2.IPolicyList,Microsoft.Practices.ObjectBuilder2.NamedTypeBuildKey,Microsoft.Practices.Unity.CompositeResolverOverride)">
+            <summary>
+            Create a new <see cref="T:Microsoft.Practices.ObjectBuilder2.BuilderContext"/> using the explicitly provided
+            values.
+            </summary>
+            <param name="chain">The <see cref="T:Microsoft.Practices.ObjectBuilder2.IStrategyChain"/> to use for this context.</param>
+            <param name="lifetime">The <see cref="T:Microsoft.Practices.ObjectBuilder2.ILifetimeContainer"/> to use for this context.</param>
+            <param name="persistentPolicies">The set of persistent policies to use for this context.</param>
+            <param name="transientPolicies">The set of transient policies to use for this context. It is
+            the caller's responsibility to ensure that the transient and persistent policies are properly
+            combined.</param>
+            <param name="buildKey">Build key for this context.</param>
+            <param name="resolverOverrides">The resolver overrides.</param>
+        </member>
+        <member name="M:Microsoft.Practices.ObjectBuilder2.BuilderContext.AddResolverOverrides(System.Collections.Generic.IEnumerable{Microsoft.Practices.Unity.ResolverOverride})">
+            <summary>
+            Add a new set of resolver override objects to the current build operation.
+            </summary>
+            <param name="newOverrides"><see cref="T:Microsoft.Practices.Unity.ResolverOverride"/> objects to add.</param>
+        </member>
+        <member name="M:Microsoft.Practices.ObjectBuilder2.BuilderContext.GetOverriddenResolver(System.Type)">
+            <summary>
+            Get a <see cref="T:Microsoft.Practices.ObjectBuilder2.IDependencyResolverPolicy"/> object for the given <paramref name="dependencyType"/>
+            or null if that dependency hasn't been overridden.
+            </summary>
+            <param name="dependencyType">Type of the dependency.</param>
+            <returns>Resolver to use, or null if no override matches for the current operation.</returns>
+        </member>
+        <member name="M:Microsoft.Practices.ObjectBuilder2.BuilderContext.NewBuildUp(Microsoft.Practices.ObjectBuilder2.NamedTypeBuildKey)">
+            <summary>
+            A convenience method to do a new buildup operation on an existing context.
+            </summary>
+            <param name="newBuildKey">Key to use to build up.</param>
+            <returns>Created object.</returns>
+        </member>
+        <member name="M:Microsoft.Practices.ObjectBuilder2.BuilderContext.NewBuildUp(Microsoft.Practices.ObjectBuilder2.NamedTypeBuildKey,System.Action{Microsoft.Practices.ObjectBuilder2.IBuilderContext})">
+            <summary>
+            A convenience method to do a new buildup operation on an existing context. This
+            overload allows you to specify extra policies which will be in effect for the duration
+            of the build.
+            </summary>
+            <param name="newBuildKey">Key defining what to build up.</param>
+            <param name="childCustomizationBlock">A delegate that takes a <see cref="T:Microsoft.Practices.ObjectBuilder2.IBuilderContext"/>. This
+            is invoked with the new child context before the build up process starts. This gives callers
+            the opportunity to customize the context for the build process.</param>
+            <returns>Created object.</returns>
+        </member>
+        <member name="P:Microsoft.Practices.ObjectBuilder2.BuilderContext.Strategies">
+            <summary>
+            Gets the head of the strategy chain.
+            </summary>
+            <returns>
+            The strategy that's first in the chain; returns null if there are no
+            strategies in the chain.
+            </returns>
+        </member>
+        <member name="P:Microsoft.Practices.ObjectBuilder2.BuilderContext.BuildKey">
+            <summary>
+            Get the current build key for the current build operation.
+            </summary>
+        </member>
+        <member name="P:Microsoft.Practices.ObjectBuilder2.BuilderContext.Existing">
+            <summary>
+            The current object being built up or torn down.
+            </summary>
+            <value>
+            The current object being manipulated by the build operation. May
+            be null if the object hasn't been created yet.</value>
+        </member>
+        <member name="P:Microsoft.Practices.ObjectBuilder2.BuilderContext.Lifetime">
+            <summary>
+            Gets the <see cref="T:Microsoft.Practices.ObjectBuilder2.ILifetimeContainer"/> associated with the build.
+            </summary>
+            <value>
+            The <see cref="T:Microsoft.Practices.ObjectBuilder2.ILifetimeContainer"/> associated with the build.
+            </value>
+        </member>
+        <member name="P:Microsoft.Practices.ObjectBuilder2.BuilderContext.OriginalBuildKey">
+            <summary>
+            Gets the original build key for the build operation.
+            </summary>
+            <value>
+            The original build key for the build operation.
+            </value>
+        </member>
+        <member name="P:Microsoft.Practices.ObjectBuilder2.BuilderContext.PersistentPolicies">
+            <summary>
+            The set of policies that were passed into this context.
+            </summary>
+            <remarks>This returns the policies passed into the context.
+            Policies added here will remain after buildup completes.</remarks>
+            <value>The persistent policies for the current context.</value>
+        </member>
+        <member name="P:Microsoft.Practices.ObjectBuilder2.BuilderContext.Policies">
+            <summary>
+            Gets the policies for the current context. 
+            </summary>
+            <remarks>
+            Any modifications will be transient (meaning, they will be forgotten when 
+            the outer BuildUp for this context is finished executing).
+            </remarks>
+            <value>
+            The policies for the current context.
+            </value>
+        </member>
+        <member name="P:Microsoft.Practices.ObjectBuilder2.BuilderContext.RecoveryStack">
+            <summary>
+            Gets the collection of <see cref="T:Microsoft.Practices.ObjectBuilder2.IRequiresRecovery"/> objects
+            that need to execute in event of an exception.
+            </summary>
+        </member>
+        <member name="P:Microsoft.Practices.ObjectBuilder2.BuilderContext.BuildComplete">
+            <summary>
+            Flag indicating if the build operation should continue.
+            </summary>
+            <value>true means that building should not call any more
+            strategies, false means continue to the next strategy.</value>
+        </member>
+        <member name="P:Microsoft.Practices.ObjectBuilder2.BuilderContext.CurrentOperation">
+            <summary>
+            An object representing what is currently being done in the
+            build chain. Used to report back errors if there's a failure.
+            </summary>
+        </member>
+        <member name="P:Microsoft.Practices.ObjectBuilder2.BuilderContext.ChildContext">
+            <summary>
+            The build context used to resolve a dependency during the build operation represented by this context.
+            </summary>
+        </member>
+        <member name="T:Microsoft.Practices.ObjectBuilder2.DependencyMissingException">
+            <summary>
+            Represents that a dependency could not be resolved.
+            </summary>
+        </member>
+        <member name="M:Microsoft.Practices.ObjectBuilder2.DependencyMissingException.#ctor">
+            <summary>
+            Initializes a new instance of the <see cref="T:Microsoft.Practices.ObjectBuilder2.DependencyMissingException"/> class with no extra information.
+            </summary>
+        </member>
+        <member name="M:Microsoft.Practices.ObjectBuilder2.DependencyMissingException.#ctor(System.String)">
+            <summary>
+            Initializes a new instance of the <see cref="T:Microsoft.Practices.ObjectBuilder2.DependencyMissingException"/> class with the given message.
+            </summary>
+            <param name="message">Some random message.</param>
+        </member>
+        <member name="M:Microsoft.Practices.ObjectBuilder2.DependencyMissingException.#ctor(System.String,System.Exception)">
+            <summary>
+            Initialize a new instance of the <see cref="T:Microsoft.Practices.ObjectBuilder2.DependencyMissingException"/> class with the given
+            message and inner exception.
+            </summary>
+            <param name="message">Some random message</param>
+            <param name="innerException">Inner exception.</param>
+        </member>
+        <member name="M:Microsoft.Practices.ObjectBuilder2.DependencyMissingException.#ctor(System.Object)">
+            <summary>
+            Initializes a new instance of the <see cref="T:Microsoft.Practices.ObjectBuilder2.DependencyMissingException"/> class with the build key of the object begin built.
+            </summary>
+            <param name="buildKey">The build key of the object begin built.</param>
+        </member>
+        <member name="T:Microsoft.Practices.ObjectBuilder2.IllegalInjectionMethodException">
+            <summary>
+            The exception thrown when injection is attempted on a method
+            that is an open generic or has out or ref params.
+            </summary>
+        </member>
+        <member name="M:Microsoft.Practices.ObjectBuilder2.IllegalInjectionMethodException.#ctor">
+            <summary>
+            Construct a new <see cref="T:Microsoft.Practices.ObjectBuilder2.IllegalInjectionMethodException"/> with no
+            message.
+            </summary>
+        </member>
+        <member name="M:Microsoft.Practices.ObjectBuilder2.IllegalInjectionMethodException.#ctor(System.String)">
+            <summary>
+            Construct a <see cref="T:Microsoft.Practices.ObjectBuilder2.IllegalInjectionMethodException"/> with the given message
+            </summary>
+            <param name="message">Message to return.</param>
+        </member>
+        <member name="M:Microsoft.Practices.ObjectBuilder2.IllegalInjectionMethodException.#ctor(System.String,System.Exception)">
+            <summary>
+            Construct a <see cref="T:Microsoft.Practices.ObjectBuilder2.IllegalInjectionMethodException"/> with the given message
+            and inner exception.
+            </summary>
+            <param name="message">Message to return.</param>
+            <param name="innerException">Inner exception</param>
+        </member>
+        <member name="T:Microsoft.Practices.ObjectBuilder2.BuilderContextExtensions">
+            <summary>
+            Extension methods to provide convenience overloads over the
+            <see cref="T:Microsoft.Practices.ObjectBuilder2.IBuilderContext"/> interface.
+            </summary>
+        </member>
+        <member name="M:Microsoft.Practices.ObjectBuilder2.BuilderContextExtensions.NewBuildUp``1(Microsoft.Practices.ObjectBuilder2.IBuilderContext)">
+            <summary>
+            Start a recursive build up operation to retrieve the default
+            value for the given <typeparamref name="TResult"/> type.
+            </summary>
+            <typeparam name="TResult">Type of object to build.</typeparam>
+            <param name="context">Parent context.</param>
+            <returns>Resulting object.</returns>
+        </member>
+        <member name="M:Microsoft.Practices.ObjectBuilder2.BuilderContextExtensions.NewBuildUp``1(Microsoft.Practices.ObjectBuilder2.IBuilderContext,System.String)">
+            <summary>
+            Start a recursive build up operation to retrieve the named
+            implementation for the given <typeparamref name="TResult"/> type.
+            </summary>
+            <typeparam name="TResult">Type to resolve.</typeparam>
+            <param name="context">Parent context.</param>
+            <param name="name">Name to resolve with.</param>
+            <returns>The resulting object.</returns>
+        </member>
+        <member name="M:Microsoft.Practices.ObjectBuilder2.BuilderContextExtensions.AddResolverOverrides(Microsoft.Practices.ObjectBuilder2.IBuilderContext,Microsoft.Practices.Unity.ResolverOverride[])">
+            <summary>
+            Add a set of <see cref="T:Microsoft.Practices.Unity.ResolverOverride"/>s to the context, specified as a 
+            variable argument list.
+            </summary>
+            <param name="context">Context to add overrides to.</param>
+            <param name="overrides">The overrides.</param>
+        </member>
+        <member name="T:Microsoft.Practices.ObjectBuilder2.IRecoveryStack">
+            <summary>
+            Data structure that stores the set of <see cref="T:Microsoft.Practices.ObjectBuilder2.IRequiresRecovery"/>
+            objects and executes them when requested.
+            </summary>
+        </member>
+        <member name="M:Microsoft.Practices.ObjectBuilder2.IRecoveryStack.Add(Microsoft.Practices.ObjectBuilder2.IRequiresRecovery)">
+            <summary>
+            Add a new <see cref="T:Microsoft.Practices.ObjectBuilder2.IRequiresRecovery"/> object to this
+            list.
+            </summary>
+            <param name="recovery">Object to add.</param>
+        </member>
+        <member name="M:Microsoft.Practices.ObjectBuilder2.IRecoveryStack.ExecuteRecovery">
+            <summary>
+            Execute the <see cref="M:Microsoft.Practices.ObjectBuilder2.IRequiresRecovery.Recover"/> method
+            of everything in the recovery list. Recoveries will execute
+            in the opposite order of add - it's a stack.
+            </summary>
+        </member>
+        <member name="P:Microsoft.Practices.ObjectBuilder2.IRecoveryStack.Count">
+            <summary>
+            Return the number of recovery objects currently in the stack.
+            </summary>
+        </member>
+        <member name="T:Microsoft.Practices.ObjectBuilder2.ILifetimeContainer">
+            <summary>
+            Represents a lifetime container.
+            </summary>
+            <remarks>
+            A lifetime container tracks the lifetime of an object, and implements
+            IDisposable. When the container is disposed, any objects in the
+            container which implement IDisposable are also disposed.
+            </remarks>
+        </member>
+        <member name="M:Microsoft.Practices.ObjectBuilder2.ILifetimeContainer.Add(System.Object)">
+            <summary>
+            Adds an object to the lifetime container.
+            </summary>
+            <param name="item">The item to be added to the lifetime container.</param>
+        </member>
+        <member name="M:Microsoft.Practices.ObjectBuilder2.ILifetimeContainer.Contains(System.Object)">
+            <summary>
+            Determine if a given object is in the lifetime container.
+            </summary>
+            <param name="item">
+            The item to locate in the lifetime container.
+            </param>
+            <returns>
+            Returns true if the object is contained in the lifetime
+            container; returns false otherwise.
+            </returns>
+        </member>
+        <member name="M:Microsoft.Practices.ObjectBuilder2.ILifetimeContainer.Remove(System.Object)">
+            <summary>
+            Removes an item from the lifetime container. The item is
+            not disposed.
+            </summary>
+            <param name="item">The item to be removed.</param>
+        </member>
+        <member name="P:Microsoft.Practices.ObjectBuilder2.ILifetimeContainer.Count">
+            <summary>
+            Gets the number of references in the lifetime container
+            </summary>
+            <value>
+            The number of references in the lifetime container
+            </value>
+        </member>
+        <member name="T:Microsoft.Practices.ObjectBuilder2.LifetimeContainer">
+            <summary>
+            Represents a lifetime container.
+            </summary>
+            <remarks>
+            A lifetime container tracks the lifetime of an object, and implements
+            IDisposable. When the container is disposed, any objects in the
+            container which implement IDisposable are also disposed.
+            </remarks>
+        </member>
+        <member name="M:Microsoft.Practices.ObjectBuilder2.LifetimeContainer.Add(System.Object)">
+            <summary>
+            Adds an object to the lifetime container.
+            </summary>
+            <param name="item">The item to be added to the lifetime container.</param>
+        </member>
+        <member name="M:Microsoft.Practices.ObjectBuilder2.LifetimeContainer.Contains(System.Object)">
+            <summary>
+            Determine if a given object is in the lifetime container.
+            </summary>
+            <param name="item">
+            The item to locate in the lifetime container.
+            </param>
+            <returns>
+            Returns true if the object is contained in the lifetime
+            container; returns false otherwise.
+            </returns>
+        </member>
+        <member name="M:Microsoft.Practices.ObjectBuilder2.LifetimeContainer.Dispose">
+            <summary>
+            Releases the resources used by the <see cref="T:Microsoft.Practices.ObjectBuilder2.LifetimeContainer"/>. 
+            </summary>
+        </member>
+        <member name="M:Microsoft.Practices.ObjectBuilder2.LifetimeContainer.Dispose(System.Boolean)">
+            <summary>
+            Releases the resources used by the <see cref="T:Microsoft.Practices.ObjectBuilder2.LifetimeContainer"/>. 
+            </summary>
+            <param name="disposing">
+            true to release managed and unmanaged resources; false to release only unmanaged resources.
+            </param>
+        </member>
+        <member name="M:Microsoft.Practices.ObjectBuilder2.LifetimeContainer.GetEnumerator">
+            <summary>
+            Returns an enumerator that iterates through the lifetime container.
+            </summary>
+            <returns>
+            An <see cref="T:System.Collections.IEnumerator"/> object that can be used to iterate through the life time container. 
+            </returns>
+        </member>
+        <member name="M:Microsoft.Practices.ObjectBuilder2.LifetimeContainer.System#Collections#IEnumerable#GetEnumerator">
+            <summary>
+            Returns an enumerator that iterates through the lifetime container.
+            </summary>
+            <returns>
+            An <see cref="T:System.Collections.IEnumerator"/> object that can be used to iterate through the life time container. 
+            </returns>
+        </member>
+        <member name="M:Microsoft.Practices.ObjectBuilder2.LifetimeContainer.Remove(System.Object)">
+            <summary>
+            Removes an item from the lifetime container. The item is
+            not disposed.
+            </summary>
+            <param name="item">The item to be removed.</param>
+        </member>
+        <member name="P:Microsoft.Practices.ObjectBuilder2.LifetimeContainer.Count">
+            <summary>
+            Gets the number of references in the lifetime container
+            </summary>
+            <value>
+            The number of references in the lifetime container
+            </value>
+        </member>
+        <member name="T:Microsoft.Practices.ObjectBuilder2.IPolicyList">
+            <summary>
+            A custom collection over <see cref="T:Microsoft.Practices.ObjectBuilder2.IBuilderPolicy"/> objects.
+            </summary>
+        </member>
+        <member name="M:Microsoft.Practices.ObjectBuilder2.IPolicyList.Clear(System.Type,System.Object)">
+            <summary>
+            Removes an individual policy type for a build key.
+            </summary>
+            <param name="policyInterface">The type of policy to remove.</param>
+            <param name="buildKey">The key the policy applies.</param>
+        </member>
+        <member name="M:Microsoft.Practices.ObjectBuilder2.IPolicyList.ClearAll">
+            <summary>
+            Removes all policies from the list.
+            </summary>
+        </member>
+        <member name="M:Microsoft.Practices.ObjectBuilder2.IPolicyList.ClearDefault(System.Type)">
+            <summary>
+            Removes a default policy.
+            </summary>
+            <param name="policyInterface">The type the policy was registered as.</param>
+        </member>
+        <member name="M:Microsoft.Practices.ObjectBuilder2.IPolicyList.Get(System.Type,System.Object,System.Boolean,Microsoft.Practices.ObjectBuilder2.IPolicyList@)">
+            <summary>
+            Gets an individual policy.
+            </summary>
+            <param name="policyInterface">The interface the policy is registered under.</param>
+            <param name="buildKey">The key the policy applies.</param>
+            <param name="localOnly">true if the policy searches local only; otherwise false to search up the parent chain.</param>
+            <param name="containingPolicyList">The policy list in the chain that the searched for policy was found in, null if the policy was
+            not found.</param>
+            <returns>The policy in the list, if present; returns null otherwise.</returns>
+        </member>
+        <member name="M:Microsoft.Practices.ObjectBuilder2.IPolicyList.GetNoDefault(System.Type,System.Object,System.Boolean,Microsoft.Practices.ObjectBuilder2.IPolicyList@)">
+            <summary>
+            Get the non default policy.
+            </summary>
+            <param name="policyInterface">The interface the policy is registered under.</param>
+            <param name="buildKey">The key the policy applies to.</param>
+            <param name="localOnly">True if the search should be in the local policy list only; otherwise false to search up the parent chain.</param>
+            <param name="containingPolicyList">The policy list in the chain that the searched for policy was found in, null if the policy was
+            not found.</param>
+            <returns>The policy in the list if present; returns null otherwise.</returns>
+        </member>
+        <member name="M:Microsoft.Practices.ObjectBuilder2.IPolicyList.Set(System.Type,Microsoft.Practices.ObjectBuilder2.IBuilderPolicy,System.Object)">
+            <summary>
+            Sets an individual policy.
+            </summary>
+            <param name="policyInterface">The <see cref="T:System.Type"/> of the policy.</param>
+            <param name="policy">The policy to be registered.</param>
+            <param name="buildKey">The key the policy applies.</param>
+        </member>
+        <member name="M:Microsoft.Practices.ObjectBuilder2.IPolicyList.SetDefault(System.Type,Microsoft.Practices.ObjectBuilder2.IBuilderPolicy)">
+            <summary>
+            Sets a default policy. When checking for a policy, if no specific individual policy
+            is available, the default will be used.
+            </summary>
+            <param name="policyInterface">The interface to register the policy under.</param>
+            <param name="policy">The default policy to be registered.</param>
+        </member>
+        <member name="T:Microsoft.Practices.ObjectBuilder2.PolicyList">
+            <summary>
+            A custom collection wrapper over <see cref="T:Microsoft.Practices.ObjectBuilder2.IBuilderPolicy"/> objects.
+            </summary>
+        </member>
+        <member name="M:Microsoft.Practices.ObjectBuilder2.PolicyList.#ctor">
+            <summary>
+            Initialize a new instance of a <see cref="T:Microsoft.Practices.ObjectBuilder2.PolicyList"/> class.
+            </summary>
+        </member>
+        <member name="M:Microsoft.Practices.ObjectBuilder2.PolicyList.#ctor(Microsoft.Practices.ObjectBuilder2.IPolicyList)">
+            <summary>
+            Initialize a new instance of a <see cref="T:Microsoft.Practices.ObjectBuilder2.PolicyList"/> class with another policy list.
+            </summary>
+            <param name="innerPolicyList">An inner policy list to search.</param>
+        </member>
+        <member name="M:Microsoft.Practices.ObjectBuilder2.PolicyList.Clear(System.Type,System.Object)">
+            <summary>
+            Removes an individual policy type for a build key.
+            </summary>
+            <param name="policyInterface">The type of policy to remove.</param>
+            <param name="buildKey">The key the policy applies.</param>
+        </member>
+        <member name="M:Microsoft.Practices.ObjectBuilder2.PolicyList.ClearAll">
+            <summary>
+            Removes all policies from the list.
+            </summary>
+        </member>
+        <member name="M:Microsoft.Practices.ObjectBuilder2.PolicyList.ClearDefault(System.Type)">
+            <summary>
+            Removes a default policy.
+            </summary>
+            <param name="policyInterface">The type the policy was registered as.</param>
+        </member>
+        <member name="M:Microsoft.Practices.ObjectBuilder2.PolicyList.Get(System.Type,System.Object,System.Boolean,Microsoft.Practices.ObjectBuilder2.IPolicyList@)">
+            <summary>
+            Gets an individual policy.
+            </summary>
+            <param name="policyInterface">The interface the policy is registered under.</param>
+            <param name="buildKey">The key the policy applies.</param>
+            <param name="localOnly">true if the policy searches local only; otherwise false to search up the parent chain.</param>
+            <param name="containingPolicyList">The policy list in the chain that the searched for policy was found in, null if the policy was
+            not found.</param>
+            <returns>The policy in the list, if present; returns null otherwise.</returns>
+        </member>
+        <member name="M:Microsoft.Practices.ObjectBuilder2.PolicyList.GetNoDefault(System.Type,System.Object,System.Boolean,Microsoft.Practices.ObjectBuilder2.IPolicyList@)">
+            <summary>
+            Get the non default policy.
+            </summary>
+            <param name="policyInterface">The interface the policy is registered under.</param>
+            <param name="buildKey">The key the policy applies to.</param>
+            <param name="localOnly">True if the search should be in the local policy list only; otherwise false to search up the parent chain.</param>
+            <param name="containingPolicyList">The policy list in the chain that the searched for policy was found in, null if the policy was
+            not found.</param>
+            <returns>The policy in the list if present; returns null otherwise.</returns>
+        </member>
+        <member name="M:Microsoft.Practices.ObjectBuilder2.PolicyList.Set(System.Type,Microsoft.Practices.ObjectBuilder2.IBuilderPolicy,System.Object)">
+            <summary>
+            Sets an individual policy.
+            </summary>
+            <param name="policyInterface">The <see cref="T:System.Type"/> of the policy.</param>
+            <param name="policy">The policy to be registered.</param>
+            <param name="buildKey">The key the policy applies.</param>
+        </member>
+        <member name="M:Microsoft.Practices.ObjectBuilder2.PolicyList.SetDefault(System.Type,Microsoft.Practices.ObjectBuilder2.IBuilderPolicy)">
+            <summary>
+            Sets a default policy. When checking for a policy, if no specific individual policy
+            is available, the default will be used.
+            </summary>
+            <param name="policyInterface">The interface to register the policy under.</param>
+            <param name="policy">The default policy to be registered.</param>
+        </member>
+        <member name="P:Microsoft.Practices.ObjectBuilder2.PolicyList.Count">
+            <summary>
+            Gets the number of items in the locator.
+            </summary>
+            <value>
+            The number of items in the locator.
+            </value>
+        </member>
+        <member name="T:Microsoft.Practices.ObjectBuilder2.RecoveryStack">
+            <summary>
+            An implementation of <see cref="T:Microsoft.Practices.ObjectBuilder2.IRecoveryStack"/>.
+            </summary>
+        </member>
+        <member name="M:Microsoft.Practices.ObjectBuilder2.RecoveryStack.Add(Microsoft.Practices.ObjectBuilder2.IRequiresRecovery)">
+            <summary>
+            Add a new <see cref="T:Microsoft.Practices.ObjectBuilder2.IRequiresRecovery"/> object to this
+            list.
+            </summary>
+            <param name="recovery">Object to add.</param>
+        </member>
+        <member name="M:Microsoft.Practices.ObjectBuilder2.RecoveryStack.ExecuteRecovery">
+            <summary>
+            Execute the <see cref="M:Microsoft.Practices.ObjectBuilder2.IRequiresRecovery.Recover"/> method
+            of everything in the recovery list. Recoveries will execute
+            in the opposite order of add - it's a stack.
+            </summary>
+        </member>
+        <member name="P:Microsoft.Practices.ObjectBuilder2.RecoveryStack.Count">
+            <summary>
+            Return the number of recovery objects currently in the stack.
+            </summary>
+        </member>
+        <member name="T:Microsoft.Practices.ObjectBuilder2.BuilderAwareStrategy">
+            <summary>
+            Implementation of <see cref="T:Microsoft.Practices.ObjectBuilder2.IBuilderStrategy"/> which will notify an object about
+            the completion of a BuildUp operation, or start of a TearDown operation.
+            </summary>
+            <remarks>
+            This strategy checks the object that is passing through the builder chain to see if it
+            implements IBuilderAware and if it does, it will call <see cref="M:Microsoft.Practices.ObjectBuilder2.IBuilderAware.OnBuiltUp(Microsoft.Practices.ObjectBuilder2.NamedTypeBuildKey)"/>
+            and <see cref="M:Microsoft.Practices.ObjectBuilder2.IBuilderAware.OnTearingDown"/>. This strategy is meant to be used from the
+            <see cref="F:Microsoft.Practices.ObjectBuilder2.BuilderStage.PostInitialization"/> stage.
+            </remarks>
+        </member>
+        <member name="M:Microsoft.Practices.ObjectBuilder2.BuilderAwareStrategy.PreBuildUp(Microsoft.Practices.ObjectBuilder2.IBuilderContext)">
+            <summary>
+            Called during the chain of responsibility for a build operation. The
+            PreBuildUp method is called when the chain is being executed in the
+            forward direction.
+            </summary>
+            <param name="context">Context of the build operation.</param>
+        </member>
+        <member name="M:Microsoft.Practices.ObjectBuilder2.BuilderAwareStrategy.PreTearDown(Microsoft.Practices.ObjectBuilder2.IBuilderContext)">
+            <summary>
+            Called during the chain of responsibility for a teardown operation. The
+            PreTearDown method is called when the chain is being executed in the
+            forward direction.
+            </summary>
+            <param name="context">Context of the teardown operation.</param>
+        </member>
+        <member name="T:Microsoft.Practices.ObjectBuilder2.IBuilderAware">
+            <summary>
+            Implemented on a class when it wants to receive notifications
+            about the build process.
+            </summary>
+        </member>
+        <member name="M:Microsoft.Practices.ObjectBuilder2.IBuilderAware.OnBuiltUp(Microsoft.Practices.ObjectBuilder2.NamedTypeBuildKey)">
+            <summary>
+            Called by the <see cref="T:Microsoft.Practices.ObjectBuilder2.BuilderAwareStrategy"/> when the object is being built up.
+            </summary>
+            <param name="buildKey">The key of the object that was just built up.</param>
+        </member>
+        <member name="M:Microsoft.Practices.ObjectBuilder2.IBuilderAware.OnTearingDown">
+            <summary>
+            Called by the <see cref="T:Microsoft.Practices.ObjectBuilder2.BuilderAwareStrategy"/> when the object is being torn down.
+            </summary>
+        </member>
+        <member name="T:Microsoft.Practices.ObjectBuilder2.BuilderStage">
+            <summary>
+            Enumeration to represent the object builder stages.
+            </summary>
+            <remarks>
+            The order of the values in the enumeration is the order in which the stages are run.
+            </remarks>
+        </member>
+        <member name="F:Microsoft.Practices.ObjectBuilder2.BuilderStage.PreCreation">
+            <summary>
+            Strategies in this stage run before creation. Typical work done in this stage might
+            include strategies that use reflection to set policies into the context that other
+            strategies would later use.
+            </summary>
+        </member>
+        <member name="F:Microsoft.Practices.ObjectBuilder2.BuilderStage.Creation">
+            <summary>
+            Strategies in this stage create objects. Typically you will only have a single policy-driven
+            creation strategy in this stage.
+            </summary>
+        </member>
+        <member name="F:Microsoft.Practices.ObjectBuilder2.BuilderStage.Initialization">
+            <summary>
+            Strategies in this stage work on created objects. Typical work done in this stage might
+            include setter injection and method calls.
+            </summary>
+        </member>
+        <member name="F:Microsoft.Practices.ObjectBuilder2.BuilderStage.PostInitialization">
+            <summary>
+            Strategies in this stage work on objects that are already initialized. Typical work done in
+            this stage might include looking to see if the object implements some notification interface
+            to discover when its initialization stage has been completed.
+            </summary>
+        </member>
+        <member name="T:Microsoft.Practices.ObjectBuilder2.BuildKeyMappingPolicy">
+            <summary>
+            Represents a builder policy for mapping build keys.
+            </summary>
+        </member>
+        <member name="T:Microsoft.Practices.ObjectBuilder2.IBuildKeyMappingPolicy">
+            <summary>
+            Represents a builder policy for mapping build keys.
+            </summary>
+        </member>
+        <member name="M:Microsoft.Practices.ObjectBuilder2.IBuildKeyMappingPolicy.Map(Microsoft.Practices.ObjectBuilder2.NamedTypeBuildKey,Microsoft.Practices.ObjectBuilder2.IBuilderContext)">
+            <summary>
+            Maps the build key.
+            </summary>
+            <param name="buildKey">The build key to map.</param>
+            <param name="context">Current build context. Used for contextual information
+            if writing a more sophisticated mapping. This parameter can be null
+            (called when getting container registrations).</param>
+            <returns>The new build key.</returns>
+        </member>
+        <member name="M:Microsoft.Practices.ObjectBuilder2.BuildKeyMappingPolicy.#ctor(Microsoft.Practices.ObjectBuilder2.NamedTypeBuildKey)">
+            <summary>
+            Initialize a new instance of the <see cref="T:Microsoft.Practices.ObjectBuilder2.BuildKeyMappingPolicy"/> with the new build key.
+            </summary>
+            <param name="newBuildKey">The new build key.</param>
+        </member>
+        <member name="M:Microsoft.Practices.ObjectBuilder2.BuildKeyMappingPolicy.Map(Microsoft.Practices.ObjectBuilder2.NamedTypeBuildKey,Microsoft.Practices.ObjectBuilder2.IBuilderContext)">
+            <summary>
+            Maps the build key.
+            </summary>
+            <param name="buildKey">The build key to map.</param>
+            <param name="context">Current build context. Used for contextual information
+            if writing a more sophisticated mapping, unused in this implementation.</param>
+            <returns>The new build key.</returns>
+        </member>
+        <member name="T:Microsoft.Practices.ObjectBuilder2.BuildKeyMappingStrategy">
+            <summary>
+            Represents a strategy for mapping build keys in the build up operation.
+            </summary>
+        </member>
+        <member name="M:Microsoft.Practices.ObjectBuilder2.BuildKeyMappingStrategy.PreBuildUp(Microsoft.Practices.ObjectBuilder2.IBuilderContext)">
+            <summary>
+            Called during the chain of responsibility for a build operation.  Looks for the <see cref="T:Microsoft.Practices.ObjectBuilder2.IBuildKeyMappingPolicy"/>
+            and if found maps the build key for the current operation.
+            </summary>
+            <param name="context">The context for the operation.</param>
+        </member>
+        <member name="T:Microsoft.Practices.ObjectBuilder2.GenericTypeBuildKeyMappingPolicy">
+            <summary>
+            An implementation of <see cref="T:Microsoft.Practices.ObjectBuilder2.IBuildKeyMappingPolicy"/> that can map
+            generic types.
+            </summary>
+        </member>
+        <member name="M:Microsoft.Practices.ObjectBuilder2.GenericTypeBuildKeyMappingPolicy.#ctor(Microsoft.Practices.ObjectBuilder2.NamedTypeBuildKey)">
+            <summary>
+            Create a new <see cref="T:Microsoft.Practices.ObjectBuilder2.GenericTypeBuildKeyMappingPolicy"/> instance
+            that will map generic types.
+            </summary>
+            <param name="destinationKey">Build key to map to. This must be or contain an open generic type.</param>
+        </member>
+        <member name="M:Microsoft.Practices.ObjectBuilder2.GenericTypeBuildKeyMappingPolicy.Map(Microsoft.Practices.ObjectBuilder2.NamedTypeBuildKey,Microsoft.Practices.ObjectBuilder2.IBuilderContext)">
+            <summary>
+            Maps the build key.
+            </summary>
+            <param name="buildKey">The build key to map.</param>
+            <param name="context">Current build context. Used for contextual information
+            if writing a more sophisticated mapping.</param>
+            <returns>The new build key.</returns>
+        </member>
+        <member name="T:Microsoft.Practices.ObjectBuilder2.BuildPlanStrategy">
+            <summary>
+            A <see cref="T:Microsoft.Practices.ObjectBuilder2.BuilderStrategy"/> that will look for a build plan
+            in the current context. If it exists, it invokes it, otherwise
+            it creates one and stores it for later, and invokes it.
+            </summary>
+        </member>
+        <member name="M:Microsoft.Practices.ObjectBuilder2.BuildPlanStrategy.PreBuildUp(Microsoft.Practices.ObjectBuilder2.IBuilderContext)">
+            <summary>
+            Called during the chain of responsibility for a build operation.
+            </summary>
+            <param name="context">The context for the operation.</param>
+        </member>
+        <member name="T:Microsoft.Practices.ObjectBuilder2.ConstructorSelectorPolicy`1">
+            <summary>
+            An implementation of <see cref="T:Microsoft.Practices.ObjectBuilder2.IConstructorSelectorPolicy"/> that chooses
+            constructors based on these criteria: first, pick a constructor marked with the
+            <typeparamref name="TInjectionConstructorMarkerAttribute"/> attribute. If there
+            isn't one, then choose the constructor with the longest parameter list. If that is ambiguous,
+            then throw.
+            </summary>
+            <exception cref="T:System.InvalidOperationException">Thrown when the constructor to choose is ambiguous.</exception>
+            <typeparam name="TInjectionConstructorMarkerAttribute">Attribute used to mark the constructor to call.</typeparam>
+        </member>
+        <member name="T:Microsoft.Practices.ObjectBuilder2.ConstructorSelectorPolicyBase`1">
+            <summary>
+            Base class that provides an implementation of <see cref="T:Microsoft.Practices.ObjectBuilder2.IConstructorSelectorPolicy"/>
+            which lets you override how the parameter resolvers are created.
+            </summary>
+        </member>
+        <member name="T:Microsoft.Practices.ObjectBuilder2.IConstructorSelectorPolicy">
+            <summary>
+            A <see cref="T:Microsoft.Practices.ObjectBuilder2.IBuilderPolicy"/> that, when implemented,
+            will determine which constructor to call from the build plan.
+            </summary>
+        </member>
+        <member name="M:Microsoft.Practices.ObjectBuilder2.IConstructorSelectorPolicy.SelectConstructor(Microsoft.Practices.ObjectBuilder2.IBuilderContext,Microsoft.Practices.ObjectBuilder2.IPolicyList)">
+            <summary>
+            Choose the constructor to call for the given type.
+            </summary>
+            <param name="context">Current build context</param>
+            <param name="resolverPolicyDestination">The <see cref="T:Microsoft.Practices.ObjectBuilder2.IPolicyList"/> to add any
+            generated resolver objects into.</param>
+            <returns>The chosen constructor.</returns>
+        </member>
+        <member name="M:Microsoft.Practices.ObjectBuilder2.ConstructorSelectorPolicyBase`1.SelectConstructor(Microsoft.Practices.ObjectBuilder2.IBuilderContext,Microsoft.Practices.ObjectBuilder2.IPolicyList)">
+            <summary>
+            Choose the constructor to call for the given type.
+            </summary>
+            <param name="context">Current build context</param>
+            <param name="resolverPolicyDestination">The <see cref="T:Microsoft.Practices.ObjectBuilder2.IPolicyList"/> to add any
+            generated resolver objects into.</param>
+            <returns>The chosen constructor.</returns>
+        </member>
+        <member name="M:Microsoft.Practices.ObjectBuilder2.ConstructorSelectorPolicyBase`1.CreateResolver(System.Reflection.ParameterInfo)">
+            <summary>
+            Create a <see cref="T:Microsoft.Practices.ObjectBuilder2.IDependencyResolverPolicy"/> instance for the given
+            <see cref="T:System.Reflection.ParameterInfo"/>.
+            </summary>
+            <param name="parameter">Parameter to create the resolver for.</param>
+            <returns>The resolver object.</returns>
+        </member>
+        <member name="M:Microsoft.Practices.ObjectBuilder2.ConstructorSelectorPolicyBase`1.ConstructorLengthComparer.Compare(System.Reflection.ConstructorInfo,System.Reflection.ConstructorInfo)">
+            <summary>
+            Compares two objects and returns a value indicating whether one is less than, equal to, or greater than the other.
+            </summary>
+            <param name="y">The second object to compare.</param>
+            <param name="x">The first object to compare.</param>
+            <returns>
+            Value Condition Less than zero is less than y. Zero equals y. Greater than zero is greater than y.
+            </returns>
+        </member>
+        <member name="M:Microsoft.Practices.ObjectBuilder2.ConstructorSelectorPolicy`1.CreateResolver(System.Reflection.ParameterInfo)">
+            <summary>
+            Create a <see cref="T:Microsoft.Practices.ObjectBuilder2.IDependencyResolverPolicy"/> instance for the given
+            <see cref="T:System.Reflection.ParameterInfo"/>.
+            </summary>
+            <param name="parameter">Parameter to create the resolver for.</param>
+            <returns>The resolver object.</returns>
+        </member>
+        <member name="T:Microsoft.Practices.ObjectBuilder2.SelectedConstructor">
+            <summary>
+            Objects of this type are the return value from <see cref="M:Microsoft.Practices.ObjectBuilder2.IConstructorSelectorPolicy.SelectConstructor(Microsoft.Practices.ObjectBuilder2.IBuilderContext,Microsoft.Practices.ObjectBuilder2.IPolicyList)"/>.
+            It encapsulates the desired <see cref="T:System.Reflection.ConstructorInfo"/> with the string keys
+            needed to look up the <see cref="T:Microsoft.Practices.ObjectBuilder2.IDependencyResolverPolicy"/> for each
+            parameter.
+            </summary>
+        </member>
+        <member name="T:Microsoft.Practices.ObjectBuilder2.SelectedMemberWithParameters`1">
+            <summary>
+            Base class for return values from selector policies that
+            return a MemberInfo of some sort plus a list of parameter
+            keys to look up the parameter resolvers.
+            </summary>
+        </member>
+        <member name="T:Microsoft.Practices.ObjectBuilder2.SelectedMemberWithParameters">
+            <summary>
+            Base class for return of selector policies that need
+            to keep track of a set of parameter resolvers.
+            </summary>
+        </member>
+        <member name="M:Microsoft.Practices.ObjectBuilder2.SelectedMemberWithParameters.AddParameterResolver(Microsoft.Practices.ObjectBuilder2.IDependencyResolverPolicy)">
+            <summary>
+            Adds the parameter resolver. Resolvers are assumed
+            to be in the order of the parameters to the member.
+            </summary>
+            <param name="newResolver">The new resolver.</param>
+        </member>
+        <member name="M:Microsoft.Practices.ObjectBuilder2.SelectedMemberWithParameters.GetParameterResolvers">
+            <summary>
+            Gets the parameter resolvers.
+            </summary>
+            <returns>An array with the parameter resolvers.</returns>
+        </member>
+        <member name="M:Microsoft.Practices.ObjectBuilder2.SelectedMemberWithParameters`1.#ctor(`0)">
+            <summary>
+            Construct a new <see cref="T:Microsoft.Practices.ObjectBuilder2.SelectedMemberWithParameters`1"/>, storing
+            the given member info.
+            </summary>
+            <param name="memberInfo">Member info to store.</param>
+        </member>
+        <member name="P:Microsoft.Practices.ObjectBuilder2.SelectedMemberWithParameters`1.MemberInfo">
+            <summary>
+            The member info stored.
+            </summary>
+        </member>
+        <member name="M:Microsoft.Practices.ObjectBuilder2.SelectedConstructor.#ctor(System.Reflection.ConstructorInfo)">
+            <summary>
+            Create a new <see cref="T:Microsoft.Practices.ObjectBuilder2.SelectedConstructor"/> instance which
+            contains the given constructor.
+            </summary>
+            <param name="constructor">The constructor to wrap.</param>
+        </member>
+        <member name="P:Microsoft.Practices.ObjectBuilder2.SelectedConstructor.Constructor">
+            <summary>
+            The constructor this object wraps.
+            </summary>
+        </member>
+        <member name="T:Microsoft.Practices.ObjectBuilder2.ConstructorArgumentResolveOperation">
+            <summary>
+            This class records the information about which constructor argument is currently
+            being resolved, and is responsible for generating the error string required when
+            an error has occurred.
+            </summary>
+        </member>
+        <member name="M:Microsoft.Practices.ObjectBuilder2.ConstructorArgumentResolveOperation.#ctor(System.Type,System.String,System.String)">
+            <summary>
+            Initializes a new instance of the <see cref="T:Microsoft.Practices.ObjectBuilder2.ConstructorArgumentResolveOperation"/> class.
+            </summary>
+            <param name="typeBeingConstructed">The type that is being constructed.</param>
+            <param name="constructorSignature">A string representing the constructor being called.</param>
+            <param name="parameterName">Parameter being resolved.</param>
+        </member>
+        <member name="M:Microsoft.Practices.ObjectBuilder2.ConstructorArgumentResolveOperation.ToString">
+            <summary>
+            Generate the string describing what parameter was being resolved.
+            </summary>
+            <returns>The description string.</returns>
+        </member>
+        <member name="P:Microsoft.Practices.ObjectBuilder2.ConstructorArgumentResolveOperation.ConstructorSignature">
+            <summary>
+            String describing the constructor being set up.
+            </summary>
+        </member>
+        <member name="P:Microsoft.Practices.ObjectBuilder2.ConstructorArgumentResolveOperation.ParameterName">
+            <summary>
+            Parameter that's being resolved.
+            </summary>
+        </member>
+        <member name="T:Microsoft.Practices.ObjectBuilder2.DynamicMethodConstructorStrategy">
+            <summary>
+            A <see cref="T:Microsoft.Practices.ObjectBuilder2.BuilderStrategy"/> that emits IL to call constructors
+            as part of creating a build plan.
+            </summary>
+        </member>
+        <member name="M:Microsoft.Practices.ObjectBuilder2.DynamicMethodConstructorStrategy.PreBuildUp(Microsoft.Practices.ObjectBuilder2.IBuilderContext)">
+            <summary>
+            Called during the chain of responsibility for a build operation.
+            </summary>
+            <remarks>Existing object is an instance of <see cref="T:Microsoft.Practices.ObjectBuilder2.DynamicBuildPlanGenerationContext"/>.</remarks>
+            <param name="context">The context for the operation.</param>
+        </member>
+        <member name="M:Microsoft.Practices.ObjectBuilder2.DynamicMethodConstructorStrategy.SetPerBuildSingleton(Microsoft.Practices.ObjectBuilder2.IBuilderContext)">
+            <summary>
+            A helper method used by the generated IL to set up a PerResolveLifetimeManager lifetime manager
+            if the current object is such.
+            </summary>
+            <param name="context">Current build context.</param>
+        </member>
+        <member name="M:Microsoft.Practices.ObjectBuilder2.DynamicMethodConstructorStrategy.CreateSignatureString(System.Reflection.ConstructorInfo)">
+            <summary>
+            Build up the string that will represent the constructor signature
+            in any exception message.
+            </summary>
+            <param name="constructor"></param>
+            <returns></returns>
+        </member>
+        <member name="M:Microsoft.Practices.ObjectBuilder2.DynamicMethodConstructorStrategy.SetCurrentOperationToResolvingParameter(System.String,System.String,Microsoft.Practices.ObjectBuilder2.IBuilderContext)">
+            <summary>
+            A helper method used by the generated IL to store the current operation in the build context.
+            </summary>
+        </member>
+        <member name="M:Microsoft.Practices.ObjectBuilder2.DynamicMethodConstructorStrategy.SetCurrentOperationToInvokingConstructor(System.String,Microsoft.Practices.ObjectBuilder2.IBuilderContext)">
+            <summary>
+            A helper method used by the generated IL to store the current operation in the build context.
+            </summary>
+        </member>
+        <member name="M:Microsoft.Practices.ObjectBuilder2.DynamicMethodConstructorStrategy.ThrowForAttemptingToConstructInterface(Microsoft.Practices.ObjectBuilder2.IBuilderContext)">
+            <summary>
+            A helper method used by the generated IL to throw an exception if
+            no existing object is present, but the user is attempting to build
+            an interface (usually due to the lack of a type mapping).
+            </summary>
+            <param name="context">The <see cref="T:Microsoft.Practices.ObjectBuilder2.IBuilderContext"/> currently being
+            used for the build of this object.</param>
+        </member>
+        <member name="M:Microsoft.Practices.ObjectBuilder2.DynamicMethodConstructorStrategy.ThrowForAttemptingToConstructAbstractClass(Microsoft.Practices.ObjectBuilder2.IBuilderContext)">
+            <summary>
+            A helper method used by the generated IL to throw an exception if
+            no existing object is present, but the user is attempting to build
+            an abstract class (usually due to the lack of a type mapping).
+            </summary>
+            <param name="context">The <see cref="T:Microsoft.Practices.ObjectBuilder2.IBuilderContext"/> currently being
+            used for the build of this object.</param>
+        </member>
+        <member name="M:Microsoft.Practices.ObjectBuilder2.DynamicMethodConstructorStrategy.ThrowForAttemptingToConstructDelegate(Microsoft.Practices.ObjectBuilder2.IBuilderContext)">
+            <summary>
+            A helper method used by the generated IL to throw an exception if
+            no existing object is present, but the user is attempting to build
+            an delegate other than Func{T} or Func{IEnumerable{T}}.
+            </summary>
+            <param name="context">The <see cref="T:Microsoft.Practices.ObjectBuilder2.IBuilderContext"/> currently being
+            used for the build of this object.</param>
+        </member>
+        <member name="M:Microsoft.Practices.ObjectBuilder2.DynamicMethodConstructorStrategy.ThrowForNullExistingObject(Microsoft.Practices.ObjectBuilder2.IBuilderContext)">
+            <summary>
+            A helper method used by the generated IL to throw an exception if
+            a dependency cannot be resolved.
+            </summary>
+            <param name="context">The <see cref="T:Microsoft.Practices.ObjectBuilder2.IBuilderContext"/> currently being
+            used for the build of this object.</param>
+        </member>
+        <member name="M:Microsoft.Practices.ObjectBuilder2.DynamicMethodConstructorStrategy.ThrowForNullExistingObjectWithInvalidConstructor(Microsoft.Practices.ObjectBuilder2.IBuilderContext,System.String)">
+            <summary>
+            A helper method used by the generated IL to throw an exception if
+            a dependency cannot be resolved because of an invalid constructor.
+            </summary>
+            <param name="context">The <see cref="T:Microsoft.Practices.ObjectBuilder2.IBuilderContext"/> currently being
+            used for the build of this object.</param>
+            <param name="signature">The signature of the invalid constructor.</param>
+        </member>
+        <member name="T:Microsoft.Practices.ObjectBuilder2.InvokingConstructorOperation">
+            <summary>
+            A class that records that a constructor is about to be call, and is 
+            responsible for generating the error string required when
+            an error has occurred.
+            </summary>
+        </member>
+        <member name="M:Microsoft.Practices.ObjectBuilder2.InvokingConstructorOperation.#ctor(System.Type,System.String)">
+            <summary>
+            Initializes a new instance of the <see cref="T:Microsoft.Practices.ObjectBuilder2.InvokingConstructorOperation"/> class.
+            </summary>
+        </member>
+        <member name="M:Microsoft.Practices.ObjectBuilder2.InvokingConstructorOperation.ToString">
+            <summary>
+            Generate the description string.
+            </summary>
+            <returns>The string.</returns>
+        </member>
+        <member name="P:Microsoft.Practices.ObjectBuilder2.InvokingConstructorOperation.ConstructorSignature">
+            <summary>
+            Constructor we're trying to call.
+            </summary>
+        </member>
+        <member name="T:Microsoft.Practices.ObjectBuilder2.DynamicBuildPlanGenerationContext">
+            <summary>
+            
+            </summary>
+        </member>
+        <member name="M:Microsoft.Practices.ObjectBuilder2.DynamicBuildPlanGenerationContext.#ctor(System.Type)">
+            <summary>
+            
+            </summary>
+            <param name="typeToBuild"></param>
+        </member>
+        <member name="M:Microsoft.Practices.ObjectBuilder2.DynamicBuildPlanGenerationContext.AddToBuildPlan(System.Linq.Expressions.Expression)">
+            <summary>
+            
+            </summary>
+            <param name="expression"></param>
+        </member>
+        <member name="M:Microsoft.Practices.ObjectBuilder2.DynamicBuildPlanGenerationContext.CreateParameterExpression(Microsoft.Practices.ObjectBuilder2.IDependencyResolverPolicy,System.Type,System.Linq.Expressions.Expression)">
+            <summary>
+            
+            </summary>
+            <param name="resolver"></param>
+            <param name="parameterType"></param>
+            <param name="setOperationExpression"></param>
+            <returns></returns>
+        </member>
+        <member name="M:Microsoft.Practices.ObjectBuilder2.DynamicBuildPlanGenerationContext.GetResolver(Microsoft.Practices.ObjectBuilder2.IBuilderContext,System.Type,System.String)">
+            <summary>
+            Helper method used by generated IL to look up a dependency resolver based on the given key.
+            </summary>
+            <param name="context">Current build context.</param>
+            <param name="dependencyType">Type of the dependency being resolved.</param>
+            <param name="resolverKey">Key the resolver was stored under.</param>
+            <returns>The found dependency resolver.</returns>
+        </member>
+        <member name="M:Microsoft.Practices.ObjectBuilder2.DynamicBuildPlanGenerationContext.GetResolver(Microsoft.Practices.ObjectBuilder2.IBuilderContext,System.Type,Microsoft.Practices.ObjectBuilder2.IDependencyResolverPolicy)">
+            <summary>
+            Helper method used by generated IL to look up a dependency resolver based on the given key.
+            </summary>
+            <param name="context">Current build context.</param>
+            <param name="dependencyType">Type of the dependency being resolved.</param>
+            <param name="resolver">The configured resolver.</param>
+            <returns>The found dependency resolver.</returns>
+        </member>
+        <member name="P:Microsoft.Practices.ObjectBuilder2.DynamicBuildPlanGenerationContext.TypeToBuild">
+            <summary>
+            The type that is to be built with the dynamic build plan.
+            </summary>
+        </member>
+        <member name="P:Microsoft.Practices.ObjectBuilder2.DynamicBuildPlanGenerationContext.ContextParameter">
+            <summary>
+            The context parameter representing the <see cref="T:Microsoft.Practices.ObjectBuilder2.IBuilderContext"/> used when the build plan is executed.
+            </summary>
+        </member>
+        <member name="T:Microsoft.Practices.ObjectBuilder2.DynamicBuildPlanMethod">
+            <summary>
+            
+            </summary>
+            <param name="context"></param>
+        </member>
+        <member name="T:Microsoft.Practices.ObjectBuilder2.DynamicMethodBuildPlan">
+            <summary>
+            
+            </summary>
+        </member>
+        <member name="M:Microsoft.Practices.ObjectBuilder2.DynamicMethodBuildPlan.#ctor(Microsoft.Practices.ObjectBuilder2.DynamicBuildPlanMethod)">
+            <summary>
+            
+            </summary>
+            <param name="buildMethod"></param>
+        </member>
+        <member name="M:Microsoft.Practices.ObjectBuilder2.DynamicMethodBuildPlan.BuildUp(Microsoft.Practices.ObjectBuilder2.IBuilderContext)">
+            <summary>
+            
+            </summary>
+            <param name="context"></param>
+        </member>
+        <member name="T:Microsoft.Practices.ObjectBuilder2.DynamicMethodBuildPlanCreatorPolicy">
+            <summary>
+            An <see cref="T:Microsoft.Practices.ObjectBuilder2.IBuildPlanCreatorPolicy"/> implementation
+            that constructs a build plan via dynamic IL emission.
+            </summary>
+        </member>
+        <member name="M:Microsoft.Practices.ObjectBuilder2.DynamicMethodBuildPlanCreatorPolicy.#ctor(Microsoft.Practices.ObjectBuilder2.IStagedStrategyChain)">
+            <summary>
+            Construct a <see cref="T:Microsoft.Practices.ObjectBuilder2.DynamicMethodBuildPlanCreatorPolicy"/> that
+            uses the given strategy chain to construct the build plan.
+            </summary>
+            <param name="strategies">The strategy chain.</param>
+        </member>
+        <member name="M:Microsoft.Practices.ObjectBuilder2.DynamicMethodBuildPlanCreatorPolicy.CreatePlan(Microsoft.Practices.ObjectBuilder2.IBuilderContext,Microsoft.Practices.ObjectBuilder2.NamedTypeBuildKey)">
+            <summary>
+            Construct a build plan.
+            </summary>
+            <param name="context">The current build context.</param>
+            <param name="buildKey">The current build key.</param>
+            <returns>The created build plan.</returns>
+        </member>
+        <member name="T:Microsoft.Practices.ObjectBuilder2.InvokingMethodOperation">
+            <summary>
+            A class that records that a constructor is about to be call, and is 
+            responsible for generating the error string required when
+            an error has occurred.
+            </summary>
+        </member>
+        <member name="M:Microsoft.Practices.ObjectBuilder2.InvokingMethodOperation.#ctor(System.Type,System.String)">
+            <summary>
+            Initializes a new instance of the <see cref="T:Microsoft.Practices.ObjectBuilder2.InvokingMethodOperation"/> class.
+            </summary>
+        </member>
+        <member name="M:Microsoft.Practices.ObjectBuilder2.InvokingMethodOperation.ToString">
+            <summary>
+            Generate the description string.
+            </summary>
+            <returns>The string.</returns>
+        </member>
+        <member name="P:Microsoft.Practices.ObjectBuilder2.InvokingMethodOperation.MethodSignature">
+            <summary>
+            Method we're trying to call.
+            </summary>
+        </member>
+        <member name="T:Microsoft.Practices.ObjectBuilder2.MethodArgumentResolveOperation">
+            <summary>
+            This class records the information about which constructor argument is currently
+            being resolved, and is responsible for generating the error string required when
+            an error has occurred.
+            </summary>
+        </member>
+        <member name="M:Microsoft.Practices.ObjectBuilder2.MethodArgumentResolveOperation.#ctor(System.Type,System.String,System.String)">
+            <summary>
+            Initializes a new instance of the <see cref="T:Microsoft.Practices.ObjectBuilder2.ConstructorArgumentResolveOperation"/> class.
+            </summary>
+            <param name="typeBeingConstructed">The type that is being constructed.</param>
+            <param name="methodSignature">A string representing the method being called.</param>
+            <param name="parameterName">Parameter being resolved.</param>
+        </member>
+        <member name="M:Microsoft.Practices.ObjectBuilder2.MethodArgumentResolveOperation.ToString">
+            <summary>
+            Generate the string describing what parameter was being resolved.
+            </summary>
+            <returns>The description string.</returns>
+        </member>
+        <member name="P:Microsoft.Practices.ObjectBuilder2.MethodArgumentResolveOperation.MethodSignature">
+            <summary>
+            String describing the method being set up.
+            </summary>
+        </member>
+        <member name="P:Microsoft.Practices.ObjectBuilder2.MethodArgumentResolveOperation.ParameterName">
+            <summary>
+            Parameter that's being resolved.
+            </summary>
+        </member>
+        <member name="T:Microsoft.Practices.ObjectBuilder2.DynamicMethodCallStrategy">
+            <summary>
+            A <see cref="T:Microsoft.Practices.ObjectBuilder2.BuilderStrategy"/> that generates IL to call
+            chosen methods (as specified by the current <see cref="T:Microsoft.Practices.ObjectBuilder2.IMethodSelectorPolicy"/>)
+            as part of object build up.
+            </summary>
+        </member>
+        <member name="M:Microsoft.Practices.ObjectBuilder2.DynamicMethodCallStrategy.PreBuildUp(Microsoft.Practices.ObjectBuilder2.IBuilderContext)">
+            <summary>
+            Called during the chain of responsibility for a build operation. The
+            PreBuildUp method is called when the chain is being executed in the
+            forward direction.
+            </summary>
+            <param name="context">Context of the build operation.</param>
+        </member>
+        <member name="M:Microsoft.Practices.ObjectBuilder2.DynamicMethodCallStrategy.SetCurrentOperationToResolvingParameter(System.String,System.String,Microsoft.Practices.ObjectBuilder2.IBuilderContext)">
+            <summary>
+            A helper method used by the generated IL to store the current operation in the build context.
+            </summary>
+        </member>
+        <member name="M:Microsoft.Practices.ObjectBuilder2.DynamicMethodCallStrategy.SetCurrentOperationToInvokingMethod(System.String,Microsoft.Practices.ObjectBuilder2.IBuilderContext)">
+            <summary>
+            A helper method used by the generated IL to store the current operation in the build context.
+            </summary>
+        </member>
+        <member name="T:Microsoft.Practices.ObjectBuilder2.PropertyOperation">
+            <summary>
+            A base class that holds the information shared by all operations
+            performed by the container while setting properties.
+            </summary>
+        </member>
+        <member name="M:Microsoft.Practices.ObjectBuilder2.PropertyOperation.#ctor(System.Type,System.String)">
+            <summary>
+            Initializes a new instance of the <see cref="T:System.Object"/> class.
+            </summary>
+        </member>
+        <member name="M:Microsoft.Practices.ObjectBuilder2.PropertyOperation.ToString">
+            <summary>
+            Generate the description of this operation.
+            </summary>
+            <returns>The string.</returns>
+        </member>
+        <member name="M:Microsoft.Practices.ObjectBuilder2.PropertyOperation.GetDescriptionFormat">
+            <summary>
+            Get a format string used to create the description. Called by
+            the base <see cref="M:Microsoft.Practices.ObjectBuilder2.PropertyOperation.ToString"/> method.
+            </summary>
+            <returns>The format string.</returns>
+        </member>
+        <member name="P:Microsoft.Practices.ObjectBuilder2.PropertyOperation.PropertyName">
+            <summary>
+            The property value currently being resolved.
+            </summary>
+        </member>
+        <member name="T:Microsoft.Practices.ObjectBuilder2.ResolvingPropertyValueOperation">
+            <summary>
+            This class records the information about which property value is currently
+            being resolved, and is responsible for generating the error string required when
+            an error has occurred.
+            </summary>
+        </member>
+        <member name="M:Microsoft.Practices.ObjectBuilder2.ResolvingPropertyValueOperation.#ctor(System.Type,System.String)">
+            <summary>
+            Initializes a new instance of the <see cref="T:System.Object"/> class.
+            </summary>
+        </member>
+        <member name="M:Microsoft.Practices.ObjectBuilder2.ResolvingPropertyValueOperation.GetDescriptionFormat">
+            <summary>
+            Get a format string used to create the description. Called by
+            the base <see cref="M:Microsoft.Practices.ObjectBuilder2.PropertyOperation.ToString"/> method.
+            </summary>
+            <returns>The format string.</returns>
+        </member>
+        <member name="T:Microsoft.Practices.ObjectBuilder2.DynamicMethodPropertySetterStrategy">
+            <summary>
+            A <see cref="T:Microsoft.Practices.ObjectBuilder2.BuilderStrategy"/> that generates IL to resolve properties
+            on an object being built.
+            </summary>
+        </member>
+        <member name="M:Microsoft.Practices.ObjectBuilder2.DynamicMethodPropertySetterStrategy.PreBuildUp(Microsoft.Practices.ObjectBuilder2.IBuilderContext)">
+            <summary>
+            Called during the chain of responsibility for a build operation.
+            </summary>
+            <param name="context">The context for the operation.</param>
+        </member>
+        <member name="M:Microsoft.Practices.ObjectBuilder2.DynamicMethodPropertySetterStrategy.SetCurrentOperationToResolvingPropertyValue(System.String,Microsoft.Practices.ObjectBuilder2.IBuilderContext)">
+            <summary>
+            A helper method used by the generated IL to store the current operation in the build context.
+            </summary>
+        </member>
+        <member name="M:Microsoft.Practices.ObjectBuilder2.DynamicMethodPropertySetterStrategy.SetCurrentOperationToSettingProperty(System.String,Microsoft.Practices.ObjectBuilder2.IBuilderContext)">
+            <summary>
+            A helper method used by the generated IL to store the current operation in the build context.
+            </summary>
+        </member>
+        <member name="T:Microsoft.Practices.ObjectBuilder2.SettingPropertyOperation">
+            <summary>
+            This class records the information about which property value is currently
+            being set, and is responsible for generating the error string required when
+            an error has occurred.
+            </summary>
+        </member>
+        <member name="M:Microsoft.Practices.ObjectBuilder2.SettingPropertyOperation.#ctor(System.Type,System.String)">
+            <summary>
+            Initializes a new instance of the <see cref="T:Microsoft.Practices.ObjectBuilder2.SettingPropertyOperation"/> class.
+            </summary>
+            <param name="typeBeingConstructed">Type property is on.</param>
+            <param name="propertyName">Name of property being set.</param>
+        </member>
+        <member name="M:Microsoft.Practices.ObjectBuilder2.SettingPropertyOperation.GetDescriptionFormat">
+            <summary>
+            Get a format string used to create the description. Called by
+            the base <see cref="M:Microsoft.Practices.ObjectBuilder2.PropertyOperation.ToString"/> method.
+            </summary>
+            <returns>The format string.</returns>
+        </member>
+        <member name="M:Microsoft.Practices.ObjectBuilder2.FactoryDelegateBuildPlanPolicy.BuildUp(Microsoft.Practices.ObjectBuilder2.IBuilderContext)">
+            <summary>
+            Creates an instance of this build plan's type, or fills
+            in the existing type if passed in.
+            </summary>
+            <param name="context">Context used to build up the object.</param>
+        </member>
+        <member name="T:Microsoft.Practices.ObjectBuilder2.IMethodSelectorPolicy">
+            <summary>
+            An <see cref="T:Microsoft.Practices.ObjectBuilder2.IBuilderPolicy"/> that will examine the given
+            types and return a sequence of <see cref="T:System.Reflection.MethodInfo"/> objects
+            that should be called as part of building the object.
+            </summary>
+        </member>
+        <member name="M:Microsoft.Practices.ObjectBuilder2.IMethodSelectorPolicy.SelectMethods(Microsoft.Practices.ObjectBuilder2.IBuilderContext,Microsoft.Practices.ObjectBuilder2.IPolicyList)">
+            <summary>
+            Return the sequence of methods to call while building the target object.
+            </summary>
+            <param name="context">Current build context.</param>
+            <param name="resolverPolicyDestination">The <see cref="T:Microsoft.Practices.ObjectBuilder2.IPolicyList"/> to add any
+            generated resolver objects into.</param>
+            <returns>Sequence of methods to call.</returns>
+        </member>
+        <member name="T:Microsoft.Practices.ObjectBuilder2.MethodSelectorPolicy`1">
+            <summary>
+            An implementation of <see cref="T:Microsoft.Practices.ObjectBuilder2.IMethodSelectorPolicy"/> that selects
+            methods by looking for the given <typeparamref name="TMarkerAttribute"/>
+            attribute on those methods.
+            </summary>
+            <typeparam name="TMarkerAttribute">Type of attribute used to mark methods
+            to inject.</typeparam>
+        </member>
+        <member name="T:Microsoft.Practices.ObjectBuilder2.MethodSelectorPolicyBase`1">
+            <summary>
+            Base class that provides an implementation of <see cref="T:Microsoft.Practices.ObjectBuilder2.IMethodSelectorPolicy"/>
+            which lets you override how the parameter resolvers are created.
+            </summary>
+            <typeparam name="TMarkerAttribute">Attribute that marks methods that should
+            be called.</typeparam>
+        </member>
+        <member name="M:Microsoft.Practices.ObjectBuilder2.MethodSelectorPolicyBase`1.SelectMethods(Microsoft.Practices.ObjectBuilder2.IBuilderContext,Microsoft.Practices.ObjectBuilder2.IPolicyList)">
+            <summary>
+            Return the sequence of methods to call while building the target object.
+            </summary>
+            <param name="context">Current build context.</param>
+            <param name="resolverPolicyDestination">The <see cref="T:Microsoft.Practices.ObjectBuilder2.IPolicyList"/> to add any
+            generated resolver objects into.</param>
+            <returns>Sequence of methods to call.</returns>
+        </member>
+        <member name="M:Microsoft.Practices.ObjectBuilder2.MethodSelectorPolicyBase`1.CreateResolver(System.Reflection.ParameterInfo)">
+            <summary>
+            Create a <see cref="T:Microsoft.Practices.ObjectBuilder2.IDependencyResolverPolicy"/> instance for the given
+            <see cref="T:System.Reflection.ParameterInfo"/>.
+            </summary>
+            <param name="parameter">Parameter to create the resolver for.</param>
+            <returns>The resolver object.</returns>
+        </member>
+        <member name="M:Microsoft.Practices.ObjectBuilder2.MethodSelectorPolicy`1.CreateResolver(System.Reflection.ParameterInfo)">
+            <summary>
+            Create a <see cref="T:Microsoft.Practices.ObjectBuilder2.IDependencyResolverPolicy"/> instance for the given
+            <see cref="T:System.Reflection.ParameterInfo"/>.
+            </summary>
+            <param name="parameter">Parameter to create the resolver for.</param>
+            <returns>The resolver object.</returns>
+        </member>
+        <member name="T:Microsoft.Practices.ObjectBuilder2.SelectedMethod">
+            <summary>
+            Objects of this type are the return value from <see cref="M:Microsoft.Practices.ObjectBuilder2.IMethodSelectorPolicy.SelectMethods(Microsoft.Practices.ObjectBuilder2.IBuilderContext,Microsoft.Practices.ObjectBuilder2.IPolicyList)"/>.
+            It encapsulates the desired <see cref="T:System.Reflection.MethodInfo"/> with the string keys
+            needed to look up the <see cref="T:Microsoft.Practices.ObjectBuilder2.IDependencyResolverPolicy"/> for each
+            parameter.
+            </summary>
+        </member>
+        <member name="M:Microsoft.Practices.ObjectBuilder2.SelectedMethod.#ctor(System.Reflection.MethodInfo)">
+            <summary>
+            Create a new <see cref="T:Microsoft.Practices.ObjectBuilder2.SelectedMethod"/> instance which
+            contains the given method.
+            </summary>
+            <param name="method">The method</param>
+        </member>
+        <member name="P:Microsoft.Practices.ObjectBuilder2.SelectedMethod.Method">
+            <summary>
+            The constructor this object wraps.
+            </summary>
+        </member>
+        <member name="T:Microsoft.Practices.ObjectBuilder2.IPropertySelectorPolicy">
+            <summary>
+            An <see cref="T:Microsoft.Practices.ObjectBuilder2.IBuilderPolicy"/> that returns a sequence
+            of properties that should be injected for the given type.
+            </summary>
+        </member>
+        <member name="M:Microsoft.Practices.ObjectBuilder2.IPropertySelectorPolicy.SelectProperties(Microsoft.Practices.ObjectBuilder2.IBuilderContext,Microsoft.Practices.ObjectBuilder2.IPolicyList)">
+            <summary>
+            Returns sequence of properties on the given type that
+            should be set as part of building that object.
+            </summary>
+            <param name="context">Current build context.</param>
+            <param name="resolverPolicyDestination">The <see cref="T:Microsoft.Practices.ObjectBuilder2.IPolicyList"/> to add any
+            generated resolver objects into.</param>
+            <returns>Sequence of <see cref="T:System.Reflection.PropertyInfo"/> objects
+            that contain the properties to set.</returns>
+        </member>
+        <member name="T:Microsoft.Practices.ObjectBuilder2.PropertySelectorBase`1">
+            <summary>
+            Base class that provides an implementation of <see cref="T:Microsoft.Practices.ObjectBuilder2.IPropertySelectorPolicy"/>
+            which lets you override how the parameter resolvers are created.
+            </summary>
+        </member>
+        <member name="M:Microsoft.Practices.ObjectBuilder2.PropertySelectorBase`1.SelectProperties(Microsoft.Practices.ObjectBuilder2.IBuilderContext,Microsoft.Practices.ObjectBuilder2.IPolicyList)">
+            <summary>
+            Returns sequence of properties on the given type that
+            should be set as part of building that object.
+            </summary>
+            <param name="context">Current build context.</param>
+            <param name="resolverPolicyDestination">The <see cref="T:Microsoft.Practices.ObjectBuilder2.IPolicyList"/> to add any
+            generated resolver objects into.</param>
+            <returns>Sequence of <see cref="T:System.Reflection.PropertyInfo"/> objects
+            that contain the properties to set.</returns>
+        </member>
+        <member name="M:Microsoft.Practices.ObjectBuilder2.PropertySelectorBase`1.CreateResolver(System.Reflection.PropertyInfo)">
+            <summary>
+            Create a <see cref="T:Microsoft.Practices.ObjectBuilder2.IDependencyResolverPolicy"/> for the given
+            property.
+            </summary>
+            <param name="property">Property to create resolver for.</param>
+            <returns>The resolver object.</returns>
+        </member>
+        <member name="T:Microsoft.Practices.ObjectBuilder2.PropertySelectorPolicy`1">
+            <summary>
+            An implementation of <see cref="T:Microsoft.Practices.ObjectBuilder2.IPropertySelectorPolicy"/> that looks
+            for properties marked with the <typeparamref name="TResolutionAttribute"/>
+            attribute that are also settable and not indexers.
+            </summary>
+            <typeparam name="TResolutionAttribute"></typeparam>
+        </member>
+        <member name="M:Microsoft.Practices.ObjectBuilder2.PropertySelectorPolicy`1.CreateResolver(System.Reflection.PropertyInfo)">
+            <summary>
+            Create a <see cref="T:Microsoft.Practices.ObjectBuilder2.IDependencyResolverPolicy"/> for the given
+            property.
+            </summary>
+            <param name="property">Property to create resolver for.</param>
+            <returns>The resolver object.</returns>
+        </member>
+        <member name="T:Microsoft.Practices.ObjectBuilder2.SelectedProperty">
+            <summary>
+            Objects of this type are returned from
+            <see cref="M:Microsoft.Practices.ObjectBuilder2.IPropertySelectorPolicy.SelectProperties(Microsoft.Practices.ObjectBuilder2.IBuilderContext,Microsoft.Practices.ObjectBuilder2.IPolicyList)"/>.
+            This class combines the <see cref="T:System.Reflection.PropertyInfo"/> about
+            the property with the string key used to look up the resolver
+            for this property's value.
+            </summary>
+        </member>
+        <member name="M:Microsoft.Practices.ObjectBuilder2.SelectedProperty.#ctor(System.Reflection.PropertyInfo,Microsoft.Practices.ObjectBuilder2.IDependencyResolverPolicy)">
+            <summary>
+            Create an instance of <see cref="T:Microsoft.Practices.ObjectBuilder2.SelectedProperty"/>
+            with the given <see cref="T:System.Reflection.PropertyInfo"/> and key.
+            </summary>
+            <param name="property">The property.</param>
+            <param name="resolver"></param>
+        </member>
+        <member name="P:Microsoft.Practices.ObjectBuilder2.SelectedProperty.Property">
+            <summary>
+            PropertyInfo for this property.
+            </summary>
+        </member>
+        <member name="P:Microsoft.Practices.ObjectBuilder2.SelectedProperty.Resolver">
+            <summary>
+            IDependencyResolverPolicy for this property
+            </summary>
+        </member>
+        <member name="T:Microsoft.Practices.ObjectBuilder2.DependencyResolverTrackerPolicy">
+            <summary>
+            Implementation of <see cref="T:Microsoft.Practices.ObjectBuilder2.IDependencyResolverTrackerPolicy"/>.
+            </summary>
+        </member>
+        <member name="T:Microsoft.Practices.ObjectBuilder2.IDependencyResolverTrackerPolicy">
+            <summary>
+            A builder policy that lets you keep track of the current
+            resolvers and will remove them from the given policy set.
+            </summary>
+        </member>
+        <member name="M:Microsoft.Practices.ObjectBuilder2.IDependencyResolverTrackerPolicy.AddResolverKey(System.Object)">
+            <summary>
+            Add a new resolver to track by key.
+            </summary>
+            <param name="key">Key that was used to add the resolver to the policy set.</param>
+        </member>
+        <member name="M:Microsoft.Practices.ObjectBuilder2.IDependencyResolverTrackerPolicy.RemoveResolvers(Microsoft.Practices.ObjectBuilder2.IPolicyList)">
+            <summary>
+            Remove the currently tracked resolvers from the given policy list.
+            </summary>
+            <param name="policies">Policy list to remove the resolvers from.</param>
+        </member>
+        <member name="M:Microsoft.Practices.ObjectBuilder2.DependencyResolverTrackerPolicy.AddResolverKey(System.Object)">
+            <summary>
+            Add a new resolver to track by key.
+            </summary>
+            <param name="key">Key that was used to add the resolver to the policy set.</param>
+        </member>
+        <member name="M:Microsoft.Practices.ObjectBuilder2.DependencyResolverTrackerPolicy.RemoveResolvers(Microsoft.Practices.ObjectBuilder2.IPolicyList)">
+            <summary>
+            Remove the currently tracked resolvers from the given policy list.
+            </summary>
+            <param name="policies">Policy list to remove the resolvers from.</param>
+        </member>
+        <member name="M:Microsoft.Practices.ObjectBuilder2.DependencyResolverTrackerPolicy.GetTracker(Microsoft.Practices.ObjectBuilder2.IPolicyList,System.Object)">
+            <summary>
+            Get an instance that implements <see cref="T:Microsoft.Practices.ObjectBuilder2.IDependencyResolverTrackerPolicy"/>,
+            either the current one in the policy set or creating a new one if it doesn't
+            exist.
+            </summary>
+            <param name="policies">Policy list to look up from.</param>
+            <param name="buildKey">Build key to track.</param>
+            <returns>The resolver tracker.</returns>
+        </member>
+        <member name="M:Microsoft.Practices.ObjectBuilder2.DependencyResolverTrackerPolicy.TrackKey(Microsoft.Practices.ObjectBuilder2.IPolicyList,System.Object,System.Object)">
+            <summary>
+            Add a key to be tracked to the current tracker.
+            </summary>
+            <param name="policies">Policy list containing the resolvers and trackers.</param>
+            <param name="buildKey">Build key for the resolvers being tracked.</param>
+            <param name="resolverKey">Key for the resolver.</param>
+        </member>
+        <member name="M:Microsoft.Practices.ObjectBuilder2.DependencyResolverTrackerPolicy.RemoveResolvers(Microsoft.Practices.ObjectBuilder2.IPolicyList,System.Object)">
+            <summary>
+            Remove the resolvers for the given build key.
+            </summary>
+            <param name="policies">Policy list containing the build key.</param>
+            <param name="buildKey">Build key.</param>
+        </member>
+        <member name="T:Microsoft.Practices.ObjectBuilder2.FixedTypeResolverPolicy">
+            <summary>
+            An implementation of <see cref="T:Microsoft.Practices.ObjectBuilder2.IDependencyResolverPolicy"/> that
+            calls back into the build chain to build up the dependency, passing
+            a type given at compile time as its build key.
+            </summary>
+        </member>
+        <member name="M:Microsoft.Practices.ObjectBuilder2.FixedTypeResolverPolicy.#ctor(System.Type)">
+            <summary>
+            Create a new instance storing the given type.
+            </summary>
+            <param name="typeToBuild">Type to resolve.</param>
+        </member>
+        <member name="M:Microsoft.Practices.ObjectBuilder2.FixedTypeResolverPolicy.Resolve(Microsoft.Practices.ObjectBuilder2.IBuilderContext)">
+            <summary>
+            Get the value for a dependency.
+            </summary>
+            <param name="context">Current build context.</param>
+            <returns>The value for the dependency.</returns>
+        </member>
+        <member name="T:Microsoft.Practices.ObjectBuilder2.IStagedStrategyChain">
+            <summary>
+            This interface defines a standard method to convert any <see cref="T:Microsoft.Practices.ObjectBuilder2.StagedStrategyChain`1"/> regardless
+            of the stage enum into a regular, flat strategy chain.
+            </summary>
+        </member>
+        <member name="M:Microsoft.Practices.ObjectBuilder2.IStagedStrategyChain.MakeStrategyChain">
+            <summary>
+            Convert this <see cref="T:Microsoft.Practices.ObjectBuilder2.StagedStrategyChain`1"/> into
+            a flat <see cref="T:Microsoft.Practices.ObjectBuilder2.IStrategyChain"/>.
+            </summary>
+            <returns>The flattened <see cref="T:Microsoft.Practices.ObjectBuilder2.IStrategyChain"/>.</returns>
+        </member>
+        <member name="T:Microsoft.Practices.ObjectBuilder2.IStrategyChain">
+            <summary>
+            Represents a chain of responsibility for builder strategies.
+            </summary>
+        </member>
+        <member name="M:Microsoft.Practices.ObjectBuilder2.IStrategyChain.Reverse">
+            <summary>
+            Reverse the order of the strategy chain.
+            </summary>
+            <returns>The reversed strategy chain.</returns>
+        </member>
+        <member name="M:Microsoft.Practices.ObjectBuilder2.IStrategyChain.ExecuteBuildUp(Microsoft.Practices.ObjectBuilder2.IBuilderContext)">
+            <summary>
+            Execute this strategy chain against the given context,
+            calling the Buildup methods on the strategies.
+            </summary>
+            <param name="context">Context for the build process.</param>
+            <returns>The build up object</returns>
+        </member>
+        <member name="M:Microsoft.Practices.ObjectBuilder2.IStrategyChain.ExecuteTearDown(Microsoft.Practices.ObjectBuilder2.IBuilderContext)">
+            <summary>
+            Execute this strategy chain against the given context,
+            calling the TearDown methods on the strategies.
+            </summary>
+            <param name="context">Context for the teardown process.</param>
+        </member>
+        <member name="T:Microsoft.Practices.ObjectBuilder2.ILifetimeFactoryPolicy">
+            <summary>
+            A builder policy used to create lifetime policy instances.
+            Used by the LifetimeStrategy when instantiating open
+            generic types.
+            </summary>
+        </member>
+        <member name="M:Microsoft.Practices.ObjectBuilder2.ILifetimeFactoryPolicy.CreateLifetimePolicy">
+            <summary>
+            Create a new instance of <see cref="T:Microsoft.Practices.ObjectBuilder2.ILifetimePolicy"/>.
+            </summary>
+            <returns>The new instance.</returns>
+        </member>
+        <member name="P:Microsoft.Practices.ObjectBuilder2.ILifetimeFactoryPolicy.LifetimeType">
+            <summary>
+            The type of Lifetime manager that will be created by this factory.
+            </summary>
+        </member>
+        <member name="T:Microsoft.Practices.ObjectBuilder2.LifetimeStrategy">
+            <summary>
+            An <see cref="T:Microsoft.Practices.ObjectBuilder2.IBuilderStrategy"/> implementation that uses
+            a <see cref="T:Microsoft.Practices.ObjectBuilder2.ILifetimePolicy"/> to figure out if an object
+            has already been created and to update or remove that
+            object from some backing store.
+            </summary>
+        </member>
+        <member name="M:Microsoft.Practices.ObjectBuilder2.LifetimeStrategy.PreBuildUp(Microsoft.Practices.ObjectBuilder2.IBuilderContext)">
+            <summary>
+            Called during the chain of responsibility for a build operation. The
+            PreBuildUp method is called when the chain is being executed in the
+            forward direction.
+            </summary>
+            <param name="context">Context of the build operation.</param>
+        </member>
+        <member name="M:Microsoft.Practices.ObjectBuilder2.LifetimeStrategy.PostBuildUp(Microsoft.Practices.ObjectBuilder2.IBuilderContext)">
+            <summary>
+            Called during the chain of responsibility for a build operation. The
+            PostBuildUp method is called when the chain has finished the PreBuildUp
+            phase and executes in reverse order from the PreBuildUp calls.
+            </summary>
+            <param name="context">Context of the build operation.</param>
+        </member>
+        <member name="T:Microsoft.Practices.ObjectBuilder2.StagedStrategyChain`1">
+            <summary>
+            Represents a chain of responsibility for builder strategies partitioned by stages.
+            </summary>
+            <typeparam name="TStageEnum">The stage enumeration to partition the strategies.</typeparam>
+        </member>
+        <member name="M:Microsoft.Practices.ObjectBuilder2.StagedStrategyChain`1.#ctor">
+            <summary>
+            Initialize a new instance of the <see cref="T:Microsoft.Practices.ObjectBuilder2.StagedStrategyChain`1"/> class.
+            </summary>
+        </member>
+        <member name="M:Microsoft.Practices.ObjectBuilder2.StagedStrategyChain`1.#ctor(Microsoft.Practices.ObjectBuilder2.StagedStrategyChain{`0})">
+            <summary>
+            Initialize a new instance of the <see cref="T:Microsoft.Practices.ObjectBuilder2.StagedStrategyChain`1"/> class with an inner strategy chain to use when building.
+            </summary>
+            <param name="innerChain">The inner strategy chain to use first when finding strategies in the build operation.</param>
+        </member>
+        <member name="M:Microsoft.Practices.ObjectBuilder2.StagedStrategyChain`1.Add(Microsoft.Practices.ObjectBuilder2.IBuilderStrategy,`0)">
+            <summary>
+            Adds a strategy to the chain at a particular stage.
+            </summary>
+            <param name="strategy">The strategy to add to the chain.</param>
+            <param name="stage">The stage to add the strategy.</param>
+        </member>
+        <member name="M:Microsoft.Practices.ObjectBuilder2.StagedStrategyChain`1.AddNew``1(`0)">
+            <summary>
+            Add a new strategy for the <paramref name="stage"/>.
+            </summary>
+            <typeparam name="TStrategy">The <see cref="T:System.Type"/> of <see cref="T:Microsoft.Practices.ObjectBuilder2.IBuilderStrategy"/></typeparam>
+            <param name="stage">The stage to add the strategy.</param>
+        </member>
+        <member name="M:Microsoft.Practices.ObjectBuilder2.StagedStrategyChain`1.Clear">
+            <summary>
+            Clear the current strategy chain list.
+            </summary>
+            <remarks>
+            This will not clear the inner strategy chain if this instance was created with one.
+            </remarks>
+        </member>
+        <member name="M:Microsoft.Practices.ObjectBuilder2.StagedStrategyChain`1.MakeStrategyChain">
+            <summary>
+            Makes a strategy chain based on this instance.
+            </summary>
+            <returns>A new <see cref="T:Microsoft.Practices.ObjectBuilder2.StrategyChain"/>.</returns>
+        </member>
+        <member name="T:Microsoft.Practices.ObjectBuilder2.StrategyChain">
+            <summary>
+            Represents a chain of responsibility for builder strategies.
+            </summary>
+        </member>
+        <member name="M:Microsoft.Practices.ObjectBuilder2.StrategyChain.#ctor">
+            <summary>
+            Initialize a new instance of the <see cref="T:Microsoft.Practices.ObjectBuilder2.StrategyChain"/> class.
+            </summary>
+        </member>
+        <member name="M:Microsoft.Practices.ObjectBuilder2.StrategyChain.#ctor(System.Collections.IEnumerable)">
+            <summary>
+            Initialize a new instance of the <see cref="T:Microsoft.Practices.ObjectBuilder2.StrategyChain"/> class with a collection of strategies.
+            </summary>
+            <param name="strategies">A collection of strategies to initialize the chain.</param>
+        </member>
+        <member name="M:Microsoft.Practices.ObjectBuilder2.StrategyChain.Add(Microsoft.Practices.ObjectBuilder2.IBuilderStrategy)">
+            <summary>
+            Adds a strategy to the chain.
+            </summary>
+            <param name="strategy">The strategy to add to the chain.</param>
+        </member>
+        <member name="M:Microsoft.Practices.ObjectBuilder2.StrategyChain.AddRange(System.Collections.IEnumerable)">
+            <summary>
+            Adds strategies to the chain.
+            </summary>
+            <param name="strategyEnumerable">The strategies to add to the chain.</param>
+        </member>
+        <member name="M:Microsoft.Practices.ObjectBuilder2.StrategyChain.Reverse">
+            <summary>
+            Reverse the order of the strategy chain.
+            </summary>
+            <returns>The reversed strategy chain.</returns>
+        </member>
+        <member name="M:Microsoft.Practices.ObjectBuilder2.StrategyChain.ExecuteBuildUp(Microsoft.Practices.ObjectBuilder2.IBuilderContext)">
+            <summary>
+            Execute this strategy chain against the given context to build up.
+            </summary>
+            <param name="context">Context for the build processes.</param>
+            <returns>The build up object</returns>
+        </member>
+        <member name="M:Microsoft.Practices.ObjectBuilder2.StrategyChain.ExecuteTearDown(Microsoft.Practices.ObjectBuilder2.IBuilderContext)">
+            <summary>
+            Execute this strategy chain against the given context,
+            calling the TearDown methods on the strategies.
+            </summary>
+            <param name="context">Context for the teardown process.</param>
+        </member>
+        <member name="M:Microsoft.Practices.ObjectBuilder2.StrategyChain.System#Collections#Generic#IEnumerable{Microsoft#Practices#ObjectBuilder2#IBuilderStrategy}#GetEnumerator">
+             <summary>
+             Returns an enumerator that iterates through the collection.
+             </summary>
+            
+             <returns>
+             A <see cref="T:System.Collections.Generic.IEnumerator`1"></see> that can be used to iterate through the collection.
+             </returns>
+        </member>
+        <member name="M:Microsoft.Practices.ObjectBuilder2.StrategyChain.GetEnumerator">
+             <summary>
+             Returns an enumerator that iterates through a collection.
+             </summary>
+            
+             <returns>
+             An <see cref="T:System.Collections.IEnumerator"></see> object that can be used to iterate through the collection.
+             </returns>
+        </member>
+        <member name="T:Microsoft.Practices.ObjectBuilder2.NamedTypeBuildKey">
+            <summary>
+            Build key used to combine a type object with a string name. Used by
+            ObjectBuilder to indicate exactly what is being built.
+            </summary>
+        </member>
+        <member name="M:Microsoft.Practices.ObjectBuilder2.NamedTypeBuildKey.#ctor(System.Type,System.String)">
+            <summary>
+            Create a new <see cref="T:Microsoft.Practices.ObjectBuilder2.NamedTypeBuildKey"/> instance with the given
+            type and name.
+            </summary>
+            <param name="type"><see cref="P:Microsoft.Practices.ObjectBuilder2.NamedTypeBuildKey.Type"/> to build.</param>
+            <param name="name">Key to use to look up type mappings and singletons.</param>
+        </member>
+        <member name="M:Microsoft.Practices.ObjectBuilder2.NamedTypeBuildKey.#ctor(System.Type)">
+            <summary>
+            Create a new <see cref="T:Microsoft.Practices.ObjectBuilder2.NamedTypeBuildKey"/> instance for the default
+            buildup of the given type.
+            </summary>
+            <param name="type"><see cref="P:Microsoft.Practices.ObjectBuilder2.NamedTypeBuildKey.Type"/> to build.</param>
+        </member>
+        <member name="M:Microsoft.Practices.ObjectBuilder2.NamedTypeBuildKey.Make``1">
+            <summary>
+            This helper method creates a new <see cref="T:Microsoft.Practices.ObjectBuilder2.NamedTypeBuildKey"/> instance. It is
+            initialized for the default key for the given type.
+            </summary>
+            <typeparam name="T">Type to build.</typeparam>
+            <returns>A new <see cref="T:Microsoft.Practices.ObjectBuilder2.NamedTypeBuildKey"/> instance.</returns>
+        </member>
+        <member name="M:Microsoft.Practices.ObjectBuilder2.NamedTypeBuildKey.Make``1(System.String)">
+            <summary>
+            This helper method creates a new <see cref="T:Microsoft.Practices.ObjectBuilder2.NamedTypeBuildKey"/> instance for
+            the given type and key.
+            </summary>
+            <typeparam name="T">Type to build</typeparam>
+            <param name="name">Key to use to look up type mappings and singletons.</param>
+            <returns>A new <see cref="T:Microsoft.Practices.ObjectBuilder2.NamedTypeBuildKey"/> instance initialized with the given type and name.</returns>
+        </member>
+        <member name="M:Microsoft.Practices.ObjectBuilder2.NamedTypeBuildKey.Equals(System.Object)">
+            <summary>
+            Compare two <see cref="T:Microsoft.Practices.ObjectBuilder2.NamedTypeBuildKey"/> instances.
+            </summary>
+            <remarks>Two <see cref="T:Microsoft.Practices.ObjectBuilder2.NamedTypeBuildKey"/> instances compare equal
+            if they contain the same name and the same type. Also, comparing
+            against a different type will also return false.</remarks>
+            <param name="obj">Object to compare to.</param>
+            <returns>True if the two keys are equal, false if not.</returns>
+        </member>
+        <member name="M:Microsoft.Practices.ObjectBuilder2.NamedTypeBuildKey.GetHashCode">
+            <summary>
+            Calculate a hash code for this instance.
+            </summary>
+            <returns>A hash code.</returns>
+        </member>
+        <member name="M:Microsoft.Practices.ObjectBuilder2.NamedTypeBuildKey.op_Equality(Microsoft.Practices.ObjectBuilder2.NamedTypeBuildKey,Microsoft.Practices.ObjectBuilder2.NamedTypeBuildKey)">
+            <summary>
+            Compare two <see cref="T:Microsoft.Practices.ObjectBuilder2.NamedTypeBuildKey"/> instances for equality.
+            </summary>
+            <remarks>Two <see cref="T:Microsoft.Practices.ObjectBuilder2.NamedTypeBuildKey"/> instances compare equal
+            if they contain the same name and the same type.</remarks>
+            <param name="left">First of the two keys to compare.</param>
+            <param name="right">Second of the two keys to compare.</param>
+            <returns>True if the values of the keys are the same, else false.</returns>
+        </member>
+        <member name="M:Microsoft.Practices.ObjectBuilder2.NamedTypeBuildKey.op_Inequality(Microsoft.Practices.ObjectBuilder2.NamedTypeBuildKey,Microsoft.Practices.ObjectBuilder2.NamedTypeBuildKey)">
+            <summary>
+            Compare two <see cref="T:Microsoft.Practices.ObjectBuilder2.NamedTypeBuildKey"/> instances for inequality.
+            </summary>
+            <remarks>Two <see cref="T:Microsoft.Practices.ObjectBuilder2.NamedTypeBuildKey"/> instances compare equal
+            if they contain the same name and the same type. If either field differs
+            the keys are not equal.</remarks>
+            <param name="left">First of the two keys to compare.</param>
+            <param name="right">Second of the two keys to compare.</param>
+            <returns>false if the values of the keys are the same, else true.</returns>
+        </member>
+        <member name="M:Microsoft.Practices.ObjectBuilder2.NamedTypeBuildKey.ToString">
+            <summary>
+            Formats the build key as a string (primarily for debugging).
+            </summary>
+            <returns>A readable string representation of the build key.</returns>
+        </member>
+        <member name="P:Microsoft.Practices.ObjectBuilder2.NamedTypeBuildKey.Type">
+            <summary>
+            Return the <see cref="P:Microsoft.Practices.ObjectBuilder2.NamedTypeBuildKey.Type"/> stored in this build key.
+            </summary>
+            <value>The type to build.</value>
+        </member>
+        <member name="P:Microsoft.Practices.ObjectBuilder2.NamedTypeBuildKey.Name">
+            <summary>
+            Returns the name stored in this build key.
+            </summary>
+            <remarks>The name to use when building.</remarks>
+        </member>
+        <member name="T:Microsoft.Practices.ObjectBuilder2.NamedTypeBuildKey`1">
+            <summary>
+            A generic version of <see cref="T:Microsoft.Practices.ObjectBuilder2.NamedTypeBuildKey"/> so that
+            you can new up a key using generic syntax.
+            </summary>
+            <typeparam name="T">Type for the key.</typeparam>
+        </member>
+        <member name="M:Microsoft.Practices.ObjectBuilder2.NamedTypeBuildKey`1.#ctor">
+            <summary>
+            Construct a new <see cref="T:Microsoft.Practices.ObjectBuilder2.NamedTypeBuildKey`1"/> that
+            specifies the given type.
+            </summary>
+        </member>
+        <member name="M:Microsoft.Practices.ObjectBuilder2.NamedTypeBuildKey`1.#ctor(System.String)">
+            <summary>
+            Construct a new <see cref="T:Microsoft.Practices.ObjectBuilder2.NamedTypeBuildKey`1"/> that
+            specifies the given type and name.
+            </summary>
+            <param name="name">Name for the key.</param>
+        </member>
+        <member name="T:Microsoft.Practices.ObjectBuilder2.Sequence">
+            <summary>
+            A series of helper methods to deal with sequences -
+            objects that implement <see cref="T:System.Collections.Generic.IEnumerable`1"/>.
+            </summary>
+        </member>
+        <member name="M:Microsoft.Practices.ObjectBuilder2.Sequence.Collect``1(``0[])">
+            <summary>
+            A function that turns an arbitrary parameter list into an
+            <see cref="T:System.Collections.Generic.IEnumerable`1"/>.
+            </summary>
+            <typeparam name="T">Type of arguments.</typeparam>
+            <param name="arguments">The items to put into the collection.</param>
+            <returns>An array that contains the values of the <paramref name="arguments"/>.</returns>
+        </member>
+        <member name="M:Microsoft.Practices.ObjectBuilder2.Sequence.Zip``2(System.Collections.Generic.IEnumerable{``0},System.Collections.Generic.IEnumerable{``1})">
+            <summary>
+            Given two sequences, return a new sequence containing the corresponding values
+            from each one.
+            </summary>
+            <typeparam name="TFirstSequenceElement">Type of first sequence.</typeparam>
+            <typeparam name="TSecondSequenceElement">Type of second sequence.</typeparam>
+            <param name="sequence1">First sequence of items.</param>
+            <param name="sequence2">Second sequence of items.</param>
+            <returns>New sequence of pairs. This sequence ends when the shorter of sequence1 and sequence2 does.</returns>
+        </member>
+        <member name="T:Microsoft.Practices.Unity.StaticFactory.StaticFactoryExtension">
+            <summary>
+            A <see cref="T:Microsoft.Practices.Unity.UnityContainerExtension"/> that lets you register a
+            delegate with the container to create an object, rather than calling
+            the object's constructor.
+            </summary>
+        </member>
+        <member name="T:Microsoft.Practices.Unity.UnityContainerExtension">
+            <summary>
+            Base class for all <see cref="T:Microsoft.Practices.Unity.UnityContainer"/> extension objects.
+            </summary>
+        </member>
+        <member name="M:Microsoft.Practices.Unity.UnityContainerExtension.InitializeExtension(Microsoft.Practices.Unity.ExtensionContext)">
+            <summary>
+            The container calls this method when the extension is added.
+            </summary>
+            <param name="context">A <see cref="T:Microsoft.Practices.Unity.ExtensionContext"/> instance that gives the
+            extension access to the internals of the container.</param>
+        </member>
+        <member name="M:Microsoft.Practices.Unity.UnityContainerExtension.Initialize">
+            <summary>
+            Initial the container with this extension's functionality.
+            </summary>
+            <remarks>
+            When overridden in a derived class, this method will modify the given
+            <see cref="T:Microsoft.Practices.Unity.ExtensionContext"/> by adding strategies, policies, etc. to
+            install it's functions into the container.</remarks>
+        </member>
+        <member name="M:Microsoft.Practices.Unity.UnityContainerExtension.Remove">
+            <summary>
+            Removes the extension's functions from the container.
+            </summary>
+            <remarks>
+            <para>
+            This method is called when extensions are being removed from the container. It can be
+            used to do things like disconnect event handlers or clean up member state. You do not
+            need to remove strategies or policies here; the container will do that automatically.
+            </para>
+            <para>
+            The default implementation of this method does nothing.</para>
+            </remarks>
+        </member>
+        <member name="P:Microsoft.Practices.Unity.UnityContainerExtension.Container">
+            <summary>
+            The container this extension has been added to.
+            </summary>
+            <value>The <see cref="T:Microsoft.Practices.Unity.UnityContainer"/> that this extension has been added to.</value>
+        </member>
+        <member name="P:Microsoft.Practices.Unity.UnityContainerExtension.Context">
+            <summary>
+            The <see cref="T:Microsoft.Practices.Unity.ExtensionContext"/> object used to manipulate
+            the inner state of the container.
+            </summary>
+        </member>
+        <member name="M:Microsoft.Practices.Unity.StaticFactory.StaticFactoryExtension.Initialize">
+            <summary>
+            Initialize this extension. This particular extension requires no
+            initialization work.
+            </summary>
+        </member>
+        <member name="M:Microsoft.Practices.Unity.StaticFactory.StaticFactoryExtension.RegisterFactory``1(System.String,System.Func{Microsoft.Practices.Unity.IUnityContainer,System.Object})">
+            <summary>
+            Register the given factory delegate to be called when the container is
+            asked to resolve <typeparamref name="TTypeToBuild"/> and <paramref name="name"/>.
+            </summary>
+            <typeparam name="TTypeToBuild">Type that will be requested from the container.</typeparam>
+            <param name="name">The name that will be used when requesting to resolve this type.</param>
+            <param name="factoryMethod">Delegate to invoke to create the instance.</param>
+            <returns>The container extension object this method was invoked on.</returns>
+        </member>
+        <member name="M:Microsoft.Practices.Unity.StaticFactory.StaticFactoryExtension.RegisterFactory``1(System.Func{Microsoft.Practices.Unity.IUnityContainer,System.Object})">
+            <summary>
+            Register the given factory delegate to be called when the container is
+            asked to resolve <typeparamref name="TTypeToBuild"/>.
+            </summary>
+            <typeparam name="TTypeToBuild">Type that will be requested from the container.</typeparam>
+            <param name="factoryMethod">Delegate to invoke to create the instance.</param>
+            <returns>The container extension object this method was invoked on.</returns>
+        </member>
+        <member name="T:Microsoft.Practices.Unity.TypeBasedOverride">
+            <summary>
+            An implementation of <see cref="T:Microsoft.Practices.Unity.ResolverOverride"/> that
+            acts as a decorator over another <see cref="T:Microsoft.Practices.Unity.ResolverOverride"/>.
+            This checks to see if the current type being built is the
+            right one before checking the inner <see cref="T:Microsoft.Practices.Unity.ResolverOverride"/>.
+            </summary>
+        </member>
+        <member name="M:Microsoft.Practices.Unity.TypeBasedOverride.#ctor(System.Type,Microsoft.Practices.Unity.ResolverOverride)">
+            <summary>
+            Create an instance of <see cref="T:Microsoft.Practices.Unity.TypeBasedOverride"/>
+            </summary>
+            <param name="targetType">Type to check for.</param>
+            <param name="innerOverride">Inner override to check after type matches.</param>
+        </member>
+        <member name="M:Microsoft.Practices.Unity.TypeBasedOverride.GetResolver(Microsoft.Practices.ObjectBuilder2.IBuilderContext,System.Type)">
+            <summary>
+            Return a <see cref="T:Microsoft.Practices.ObjectBuilder2.IDependencyResolverPolicy"/> that can be used to give a value
+            for the given desired dependency.
+            </summary>
+            <param name="context">Current build context.</param>
+            <param name="dependencyType">Type of dependency desired.</param>
+            <returns>a <see cref="T:Microsoft.Practices.ObjectBuilder2.IDependencyResolverPolicy"/> object if this override applies, null if not.</returns>
+        </member>
+        <member name="T:Microsoft.Practices.Unity.TypeBasedOverride`1">
+            <summary>
+            A convenience version of <see cref="T:Microsoft.Practices.Unity.TypeBasedOverride"/> that lets you
+            specify the type to construct via generics syntax.
+            </summary>
+            <typeparam name="T">Type to check for.</typeparam>
+        </member>
+        <member name="M:Microsoft.Practices.Unity.TypeBasedOverride`1.#ctor(Microsoft.Practices.Unity.ResolverOverride)">
+            <summary>
+            Create an instance of <see cref="T:Microsoft.Practices.Unity.TypeBasedOverride`1"/>.
+            </summary>
+            <param name="innerOverride">Inner override to check after type matches.</param>
+        </member>
+        <member name="T:Microsoft.Practices.Unity.UnityContainerExtensions">
+            <summary>
+            Extension class that adds a set of convenience overloads to the
+            <see cref="T:Microsoft.Practices.Unity.IUnityContainer"/> interface.
+            </summary>
+        </member>
+        <member name="M:Microsoft.Practices.Unity.UnityContainerExtensions.RegisterType``1(Microsoft.Practices.Unity.IUnityContainer,Microsoft.Practices.Unity.InjectionMember[])">
+            <summary>
+            Register a type with specific members to be injected.
+            </summary>
+            <typeparam name="T">Type this registration is for.</typeparam>
+            <param name="container">Container to configure.</param>
+            <param name="injectionMembers">Injection configuration objects.</param>
+            <returns>The <see cref="T:Microsoft.Practices.Unity.UnityContainer"/> object that this method was called on (this in C#, Me in Visual Basic).</returns>
+        </member>
+        <member name="M:Microsoft.Practices.Unity.UnityContainerExtensions.RegisterType``2(Microsoft.Practices.Unity.IUnityContainer,Microsoft.Practices.Unity.InjectionMember[])">
+            <summary>
+            Register a type mapping with the container.
+            </summary>
+            <remarks>
+            <para>
+            This method is used to tell the container that when asked for type <typeparamref name="TFrom"/>,
+            actually return an instance of type <typeparamref name="TTo"/>. This is very useful for
+            getting instances of interfaces.
+            </para>
+            <para>
+            This overload registers a default mapping and transient lifetime.
+            </para>
+            </remarks>
+            <typeparam name="TFrom"><see cref="T:System.Type"/> that will be requested.</typeparam>
+            <typeparam name="TTo"><see cref="T:System.Type"/> that will actually be returned.</typeparam>
+            <param name="container">Container to configure.</param>
+            <param name="injectionMembers">Injection configuration objects.</param>
+            <returns>The <see cref="T:Microsoft.Practices.Unity.UnityContainer"/> object that this method was called on (this in C#, Me in Visual Basic).</returns>
+        </member>
+        <member name="M:Microsoft.Practices.Unity.UnityContainerExtensions.RegisterType``2(Microsoft.Practices.Unity.IUnityContainer,Microsoft.Practices.Unity.LifetimeManager,Microsoft.Practices.Unity.InjectionMember[])">
+            <summary>
+            Register a type mapping with the container, where the created instances will use
+            the given <see cref="T:Microsoft.Practices.Unity.LifetimeManager"/>.
+            </summary>
+            <typeparam name="TFrom"><see cref="T:System.Type"/> that will be requested.</typeparam>
+            <typeparam name="TTo"><see cref="T:System.Type"/> that will actually be returned.</typeparam>
+            <param name="container">Container to configure.</param>
+            <param name="lifetimeManager">The <see cref="T:Microsoft.Practices.Unity.LifetimeManager"/> that controls the lifetime
+            of the returned instance.</param>
+            <param name="injectionMembers">Injection configuration objects.</param>
+            <returns>The <see cref="T:Microsoft.Practices.Unity.UnityContainer"/> object that this method was called on (this in C#, Me in Visual Basic).</returns>
+        </member>
+        <member name="M:Microsoft.Practices.Unity.UnityContainerExtensions.RegisterType``2(Microsoft.Practices.Unity.IUnityContainer,System.String,Microsoft.Practices.Unity.InjectionMember[])">
+            <summary>
+            Register a type mapping with the container.
+            </summary>
+            <remarks>
+            This method is used to tell the container that when asked for type <typeparamref name="TFrom"/>,
+            actually return an instance of type <typeparamref name="TTo"/>. This is very useful for
+            getting instances of interfaces.
+            </remarks>
+            <typeparam name="TFrom"><see cref="T:System.Type"/> that will be requested.</typeparam>
+            <typeparam name="TTo"><see cref="T:System.Type"/> that will actually be returned.</typeparam>
+            <param name="container">Container to configure.</param>
+            <param name="name">Name of this mapping.</param>
+            <param name="injectionMembers">Injection configuration objects.</param>
+            <returns>The <see cref="T:Microsoft.Practices.Unity.UnityContainer"/> object that this method was called on (this in C#, Me in Visual Basic).</returns>
+        </member>
+        <member name="M:Microsoft.Practices.Unity.UnityContainerExtensions.RegisterType``2(Microsoft.Practices.Unity.IUnityContainer,System.String,Microsoft.Practices.Unity.LifetimeManager,Microsoft.Practices.Unity.InjectionMember[])">
+            <summary>
+            Register a type mapping with the container, where the created instances will use
+            the given <see cref="T:Microsoft.Practices.Unity.LifetimeManager"/>.
+            </summary>
+            <typeparam name="TFrom"><see cref="T:System.Type"/> that will be requested.</typeparam>
+            <typeparam name="TTo"><see cref="T:System.Type"/> that will actually be returned.</typeparam>
+            <param name="container">Container to configure.</param>
+            <param name="name">Name to use for registration, null if a default registration.</param>
+            <param name="lifetimeManager">The <see cref="T:Microsoft.Practices.Unity.LifetimeManager"/> that controls the lifetime
+            of the returned instance.</param>
+            <param name="injectionMembers">Injection configuration objects.</param>
+            <returns>The <see cref="T:Microsoft.Practices.Unity.UnityContainer"/> object that this method was called on (this in C#, Me in Visual Basic).</returns>
+        </member>
+        <member name="M:Microsoft.Practices.Unity.UnityContainerExtensions.RegisterType``1(Microsoft.Practices.Unity.IUnityContainer,Microsoft.Practices.Unity.LifetimeManager,Microsoft.Practices.Unity.InjectionMember[])">
+            <summary>
+            Register a <see cref="T:Microsoft.Practices.Unity.LifetimeManager"/> for the given type with the container.
+            No type mapping is performed for this type.
+            </summary>
+            <typeparam name="T">The type to apply the <paramref name="lifetimeManager"/> to.</typeparam>
+            <param name="container">Container to configure.</param>
+            <param name="lifetimeManager">The <see cref="T:Microsoft.Practices.Unity.LifetimeManager"/> that controls the lifetime
+            of the returned instance.</param>
+            <param name="injectionMembers">Injection configuration objects.</param>
+            <returns>The <see cref="T:Microsoft.Practices.Unity.UnityContainer"/> object that this method was called on (this in C#, Me in Visual Basic).</returns>
+        </member>
+        <member name="M:Microsoft.Practices.Unity.UnityContainerExtensions.RegisterType``1(Microsoft.Practices.Unity.IUnityContainer,System.String,Microsoft.Practices.Unity.InjectionMember[])">
+            <summary>
+            Register a <see cref="T:Microsoft.Practices.Unity.LifetimeManager"/> for the given type with the container.
+            No type mapping is performed for this type.
+            </summary>
+            <typeparam name="T">The type to configure injection on.</typeparam>
+            <param name="container">Container to configure.</param>
+            <param name="name">Name that will be used to request the type.</param>
+            <param name="injectionMembers">Injection configuration objects.</param>
+            <returns>The <see cref="T:Microsoft.Practices.Unity.UnityContainer"/> object that this method was called on (this in C#, Me in Visual Basic).</returns>
+        </member>
+        <member name="M:Microsoft.Practices.Unity.UnityContainerExtensions.RegisterType``1(Microsoft.Practices.Unity.IUnityContainer,System.String,Microsoft.Practices.Unity.LifetimeManager,Microsoft.Practices.Unity.InjectionMember[])">
+            <summary>
+            Register a <see cref="T:Microsoft.Practices.Unity.LifetimeManager"/> for the given type and name with the container.
+            No type mapping is performed for this type.
+            </summary>
+            <typeparam name="T">The type to apply the <paramref name="lifetimeManager"/> to.</typeparam>
+            <param name="container">Container to configure.</param>
+            <param name="name">Name that will be used to request the type.</param>
+            <param name="lifetimeManager">The <see cref="T:Microsoft.Practices.Unity.LifetimeManager"/> that controls the lifetime
+            of the returned instance.</param>
+            <param name="injectionMembers">Injection configuration objects.</param>
+            <returns>The <see cref="T:Microsoft.Practices.Unity.UnityContainer"/> object that this method was called on (this in C#, Me in Visual Basic).</returns>
+        </member>
+        <member name="M:Microsoft.Practices.Unity.UnityContainerExtensions.RegisterType(Microsoft.Practices.Unity.IUnityContainer,System.Type,Microsoft.Practices.Unity.InjectionMember[])">
+            <summary>
+            Register a type with specific members to be injected.
+            </summary>
+            <param name="container">Container to configure.</param>
+            <param name="t">Type this registration is for.</param>
+            <param name="injectionMembers">Injection configuration objects.</param>
+            <returns>The <see cref="T:Microsoft.Practices.Unity.UnityContainer"/> object that this method was called on (this in C#, Me in Visual Basic).</returns>
+        </member>
+        <member name="M:Microsoft.Practices.Unity.UnityContainerExtensions.RegisterType(Microsoft.Practices.Unity.IUnityContainer,System.Type,System.Type,Microsoft.Practices.Unity.InjectionMember[])">
+            <summary>
+            Register a type mapping with the container.
+            </summary>
+            <remarks>
+            <para>
+            This method is used to tell the container that when asked for type <paramref name="from"/>,
+            actually return an instance of type <paramref name="to"/>. This is very useful for
+            getting instances of interfaces.
+            </para>
+            <para>
+            This overload registers a default mapping.
+            </para>
+            </remarks>
+            <param name="container">Container to configure.</param>
+            <param name="from"><see cref="T:System.Type"/> that will be requested.</param>
+            <param name="to"><see cref="T:System.Type"/> that will actually be returned.</param>
+            <param name="injectionMembers">Injection configuration objects.</param>
+            <returns>The <see cref="T:Microsoft.Practices.Unity.UnityContainer"/> object that this method was called on (this in C#, Me in Visual Basic).</returns>
+        </member>
+        <member name="M:Microsoft.Practices.Unity.UnityContainerExtensions.RegisterType(Microsoft.Practices.Unity.IUnityContainer,System.Type,System.Type,System.String,Microsoft.Practices.Unity.InjectionMember[])">
+            <summary>
+            Register a type mapping with the container.
+            </summary>
+            <remarks>
+            This method is used to tell the container that when asked for type <paramref name="from"/>,
+            actually return an instance of type <paramref name="to"/>. This is very useful for
+            getting instances of interfaces.
+            </remarks>
+            <param name="container">Container to configure.</param>
+            <param name="from"><see cref="T:System.Type"/> that will be requested.</param>
+            <param name="to"><see cref="T:System.Type"/> that will actually be returned.</param>
+            <param name="name">Name to use for registration, null if a default registration.</param>
+            <param name="injectionMembers">Injection configuration objects.</param>
+            <returns>The <see cref="T:Microsoft.Practices.Unity.UnityContainer"/> object that this method was called on (this in C#, Me in Visual Basic).</returns>
+        </member>
+        <member name="M:Microsoft.Practices.Unity.UnityContainerExtensions.RegisterType(Microsoft.Practices.Unity.IUnityContainer,System.Type,System.Type,Microsoft.Practices.Unity.LifetimeManager,Microsoft.Practices.Unity.InjectionMember[])">
+            <summary>
+            Register a type mapping with the container, where the created instances will use
+            the given <see cref="T:Microsoft.Practices.Unity.LifetimeManager"/>.
+            </summary>
+            <param name="container">Container to configure.</param>
+            <param name="from"><see cref="T:System.Type"/> that will be requested.</param>
+            <param name="to"><see cref="T:System.Type"/> that will actually be returned.</param>
+            <param name="lifetimeManager">The <see cref="T:Microsoft.Practices.Unity.LifetimeManager"/> that controls the lifetime
+            of the returned instance.</param>
+            <param name="injectionMembers">Injection configuration objects.</param>
+            <returns>The <see cref="T:Microsoft.Practices.Unity.UnityContainer"/> object that this method was called on (this in C#, Me in Visual Basic).</returns>
+        </member>
+        <member name="M:Microsoft.Practices.Unity.UnityContainerExtensions.RegisterType(Microsoft.Practices.Unity.IUnityContainer,System.Type,Microsoft.Practices.Unity.LifetimeManager,Microsoft.Practices.Unity.InjectionMember[])">
+            <summary>
+            Register a <see cref="T:Microsoft.Practices.Unity.LifetimeManager"/> for the given type and name with the container.
+            No type mapping is performed for this type.
+            </summary>
+            <param name="container">Container to configure.</param>
+            <param name="t">The <see cref="T:System.Type"/> to apply the <paramref name="lifetimeManager"/> to.</param>
+            <param name="lifetimeManager">The <see cref="T:Microsoft.Practices.Unity.LifetimeManager"/> that controls the lifetime
+            of the returned instance.</param>
+            <param name="injectionMembers">Injection configuration objects.</param>
+            <returns>The <see cref="T:Microsoft.Practices.Unity.UnityContainer"/> object that this method was called on (this in C#, Me in Visual Basic).</returns>
+        </member>
+        <member name="M:Microsoft.Practices.Unity.UnityContainerExtensions.RegisterType(Microsoft.Practices.Unity.IUnityContainer,System.Type,System.String,Microsoft.Practices.Unity.InjectionMember[])">
+            <summary>
+            Register a <see cref="T:Microsoft.Practices.Unity.LifetimeManager"/> for the given type and name with the container.
+            No type mapping is performed for this type.
+            </summary>
+            <param name="container">Container to configure.</param>
+            <param name="t">The <see cref="T:System.Type"/> to configure in the container.</param>
+            <param name="name">Name to use for registration, null if a default registration.</param>
+            <param name="injectionMembers">Injection configuration objects.</param>
+            <returns>The <see cref="T:Microsoft.Practices.Unity.UnityContainer"/> object that this method was called on (this in C#, Me in Visual Basic).</returns>
+        </member>
+        <member name="M:Microsoft.Practices.Unity.UnityContainerExtensions.RegisterType(Microsoft.Practices.Unity.IUnityContainer,System.Type,System.String,Microsoft.Practices.Unity.LifetimeManager,Microsoft.Practices.Unity.InjectionMember[])">
+            <summary>
+            Register a <see cref="T:Microsoft.Practices.Unity.LifetimeManager"/> for the given type and name with the container.
+            No type mapping is performed for this type.
+            </summary>
+            <param name="container">Container to configure.</param>
+            <param name="t">The <see cref="T:System.Type"/> to apply the <paramref name="lifetimeManager"/> to.</param>
+            <param name="name">Name to use for registration, null if a default registration.</param>
+            <param name="lifetimeManager">The <see cref="T:Microsoft.Practices.Unity.LifetimeManager"/> that controls the lifetime
+            of the returned instance.</param>
+            <param name="injectionMembers">Injection configuration objects.</param>
+            <returns>The <see cref="T:Microsoft.Practices.Unity.UnityContainer"/> object that this method was called on (this in C#, Me in Visual Basic).</returns>
+        </member>
+        <member name="M:Microsoft.Practices.Unity.UnityContainerExtensions.RegisterInstance``1(Microsoft.Practices.Unity.IUnityContainer,``0)">
+            <summary>
+            Register an instance with the container.
+            </summary>
+            <remarks>
+            <para>
+            Instance registration is much like setting a type as a singleton, except that instead
+            of the container creating the instance the first time it is requested, the user
+            creates the instance ahead of type and adds that instance to the container.
+            </para>
+            <para>
+            This overload does a default registration and has the container take over the lifetime of the instance.</para>
+            </remarks>
+            <typeparam name="TInterface">Type of instance to register (may be an implemented interface instead of the full type).</typeparam>
+            <param name="container">Container to configure.</param>
+            <param name="instance">Object to returned.</param>
+            <returns>The <see cref="T:Microsoft.Practices.Unity.UnityContainer"/> object that this method was called on (this in C#, Me in Visual Basic).</returns>
+        </member>
+        <member name="M:Microsoft.Practices.Unity.UnityContainerExtensions.RegisterInstance``1(Microsoft.Practices.Unity.IUnityContainer,``0,Microsoft.Practices.Unity.LifetimeManager)">
+            <summary>
+            Register an instance with the container.
+            </summary>
+            <remarks>
+            <para>
+            Instance registration is much like setting a type as a singleton, except that instead
+            of the container creating the instance the first time it is requested, the user
+            creates the instance ahead of type and adds that instance to the container.
+            </para>
+            <para>
+            This overload does a default registration (name = null).
+            </para>
+            </remarks>
+            <typeparam name="TInterface">Type of instance to register (may be an implemented interface instead of the full type).</typeparam>
+            <param name="container">Container to configure.</param>
+            <param name="instance">Object to returned.</param>
+            <param name="lifetimeManager">
+            <see cref="T:Microsoft.Practices.Unity.LifetimeManager"/> object that controls how this instance will be managed by the container.</param>
+            <returns>The <see cref="T:Microsoft.Practices.Unity.UnityContainer"/> object that this method was called on (this in C#, Me in Visual Basic).</returns>
+        </member>
+        <member name="M:Microsoft.Practices.Unity.UnityContainerExtensions.RegisterInstance``1(Microsoft.Practices.Unity.IUnityContainer,System.String,``0)">
+            <summary>
+            Register an instance with the container.
+            </summary>
+            <remarks>
+            <para>
+            Instance registration is much like setting a type as a singleton, except that instead
+            of the container creating the instance the first time it is requested, the user
+            creates the instance ahead of type and adds that instance to the container.
+            </para>
+            <para>
+            This overload automatically has the container take ownership of the <paramref name="instance"/>.</para>
+            </remarks>
+            <typeparam name="TInterface">Type of instance to register (may be an implemented interface instead of the full type).</typeparam>
+            <param name="instance">Object to returned.</param>
+            <param name="container">Container to configure.</param>
+            <param name="name">Name for registration.</param>
+            <returns>The <see cref="T:Microsoft.Practices.Unity.UnityContainer"/> object that this method was called on (this in C#, Me in Visual Basic).</returns>
+        </member>
+        <member name="M:Microsoft.Practices.Unity.UnityContainerExtensions.RegisterInstance``1(Microsoft.Practices.Unity.IUnityContainer,System.String,``0,Microsoft.Practices.Unity.LifetimeManager)">
+            <summary>
+            Register an instance with the container.
+            </summary>
+            <remarks>
+            <para>
+            Instance registration is much like setting a type as a singleton, except that instead
+            of the container creating the instance the first time it is requested, the user
+            creates the instance ahead of type and adds that instance to the container.
+            </para>
+            </remarks>
+            <typeparam name="TInterface">Type of instance to register (may be an implemented interface instead of the full type).</typeparam>
+            <param name="instance">Object to returned.</param>
+            <param name="container">Container to configure.</param>
+            <param name="name">Name for registration.</param>
+            <param name="lifetimeManager">
+            <see cref="T:Microsoft.Practices.Unity.LifetimeManager"/> object that controls how this instance will be managed by the container.</param>
+            <returns>The <see cref="T:Microsoft.Practices.Unity.UnityContainer"/> object that this method was called on (this in C#, Me in Visual Basic).</returns>
+        </member>
+        <member name="M:Microsoft.Practices.Unity.UnityContainerExtensions.RegisterInstance(Microsoft.Practices.Unity.IUnityContainer,System.Type,System.Object)">
+            <summary>
+            Register an instance with the container.
+            </summary>
+            <remarks>
+            <para>
+            Instance registration is much like setting a type as a singleton, except that instead
+            of the container creating the instance the first time it is requested, the user
+            creates the instance ahead of type and adds that instance to the container.
+            </para>
+            <para>
+            This overload does a default registration and has the container take over the lifetime of the instance.</para>
+            </remarks>
+            <param name="container">Container to configure.</param>
+            <param name="t">Type of instance to register (may be an implemented interface instead of the full type).</param>
+            <param name="instance">Object to returned.</param>
+            <returns>The <see cref="T:Microsoft.Practices.Unity.UnityContainer"/> object that this method was called on (this in C#, Me in Visual Basic).</returns>
+        </member>
+        <member name="M:Microsoft.Practices.Unity.UnityContainerExtensions.RegisterInstance(Microsoft.Practices.Unity.IUnityContainer,System.Type,System.Object,Microsoft.Practices.Unity.LifetimeManager)">
+            <summary>
+            Register an instance with the container.
+            </summary>
+            <remarks>
+            <para>
+            Instance registration is much like setting a type as a singleton, except that instead
+            of the container creating the instance the first time it is requested, the user
+            creates the instance ahead of type and adds that instance to the container.
+            </para>
+            <para>
+            This overload does a default registration (name = null).
+            </para>
+            </remarks>
+            <param name="container">Container to configure.</param>
+            <param name="t">Type of instance to register (may be an implemented interface instead of the full type).</param>
+            <param name="instance">Object to returned.</param>
+            <param name="lifetimeManager">
+            <see cref="T:Microsoft.Practices.Unity.LifetimeManager"/> object that controls how this instance will be managed by the container.</param>
+            <returns>The <see cref="T:Microsoft.Practices.Unity.UnityContainer"/> object that this method was called on (this in C#, Me in Visual Basic).</returns>
+        </member>
+        <member name="M:Microsoft.Practices.Unity.UnityContainerExtensions.RegisterInstance(Microsoft.Practices.Unity.IUnityContainer,System.Type,System.String,System.Object)">
+            <summary>
+            Register an instance with the container.
+            </summary>
+            <remarks>
+            <para>
+            Instance registration is much like setting a type as a singleton, except that instead
+            of the container creating the instance the first time it is requested, the user
+            creates the instance ahead of type and adds that instance to the container.
+            </para>
+            <para>
+            This overload automatically has the container take ownership of the <paramref name="instance"/>.</para>
+            </remarks>
+            <param name="container">Container to configure.</param>
+            <param name="t">Type of instance to register (may be an implemented interface instead of the full type).</param>
+            <param name="instance">Object to returned.</param>
+            <param name="name">Name for registration.</param>
+            <returns>The <see cref="T:Microsoft.Practices.Unity.UnityContainer"/> object that this method was called on (this in C#, Me in Visual Basic).</returns>
+        </member>
+        <member name="M:Microsoft.Practices.Unity.UnityContainerExtensions.Resolve``1(Microsoft.Practices.Unity.IUnityContainer,Microsoft.Practices.Unity.ResolverOverride[])">
+            <summary>
+            Resolve an instance of the default requested type from the container.
+            </summary>
+            <typeparam name="T"><see cref="T:System.Type"/> of object to get from the container.</typeparam>
+            <param name="container">Container to resolve from.</param>
+            <param name="overrides">Any overrides for the resolve call.</param>
+            <returns>The retrieved object.</returns>
+        </member>
+        <member name="M:Microsoft.Practices.Unity.UnityContainerExtensions.Resolve``1(Microsoft.Practices.Unity.IUnityContainer,System.String,Microsoft.Practices.Unity.ResolverOverride[])">
+            <summary>
+            Resolve an instance of the requested type with the given name from the container.
+            </summary>
+            <typeparam name="T"><see cref="T:System.Type"/> of object to get from the container.</typeparam>
+            <param name="container">Container to resolve from.</param>
+            <param name="name">Name of the object to retrieve.</param>
+            <param name="overrides">Any overrides for the resolve call.</param>
+            <returns>The retrieved object.</returns>
+        </member>
+        <member name="M:Microsoft.Practices.Unity.UnityContainerExtensions.Resolve(Microsoft.Practices.Unity.IUnityContainer,System.Type,Microsoft.Practices.Unity.ResolverOverride[])">
+            <summary>
+            Resolve an instance of the default requested type from the container.
+            </summary>
+            <param name="container">Container to resolve from.</param>
+            <param name="t"><see cref="T:System.Type"/> of object to get from the container.</param>
+            <param name="overrides">Any overrides for the resolve call.</param>
+            <returns>The retrieved object.</returns>
+        </member>
+        <member name="M:Microsoft.Practices.Unity.UnityContainerExtensions.ResolveAll``1(Microsoft.Practices.Unity.IUnityContainer,Microsoft.Practices.Unity.ResolverOverride[])">
+            <summary>
+            Return instances of all registered types requested.
+            </summary>
+            <remarks>
+            <para>
+            This method is useful if you've registered multiple types with the same
+            <see cref="T:System.Type"/> but different names.
+            </para>
+            <para>
+            Be aware that this method does NOT return an instance for the default (unnamed) registration.
+            </para>
+            </remarks>
+            <typeparam name="T">The type requested.</typeparam>
+            <param name="container">Container to resolve from.</param>
+            <param name="resolverOverrides">Any overrides for the resolve calls.</param>
+            <returns>Set of objects of type <typeparamref name="T"/>.</returns>
+        </member>
+        <member name="M:Microsoft.Practices.Unity.UnityContainerExtensions.BuildUp``1(Microsoft.Practices.Unity.IUnityContainer,``0,Microsoft.Practices.Unity.ResolverOverride[])">
+            <summary>
+            Run an existing object through the container and perform injection on it.
+            </summary>
+            <remarks>
+            <para>
+            This method is useful when you don't control the construction of an
+            instance (ASP.NET pages or objects created via XAML, for instance)
+            but you still want properties and other injection performed.
+            </para>
+            <para>
+            This overload uses the default registrations.
+            </para>
+            </remarks>
+            <typeparam name="T"><see cref="T:System.Type"/> of object to perform injection on.</typeparam>
+            <param name="container">Container to resolve through.</param>
+            <param name="existing">Instance to build up.</param>
+            <param name="resolverOverrides">Any overrides for the buildup.</param>
+            <returns>The resulting object. By default, this will be <paramref name="existing"/>, but
+            container extensions may add things like automatic proxy creation which would
+            cause this to return a different object (but still type compatible with <typeparamref name="T"/>).</returns>
+        </member>
+        <member name="M:Microsoft.Practices.Unity.UnityContainerExtensions.BuildUp``1(Microsoft.Practices.Unity.IUnityContainer,``0,System.String,Microsoft.Practices.Unity.ResolverOverride[])">
+            <summary>
+            Run an existing object through the container and perform injection on it.
+            </summary>
+            <remarks>
+            <para>
+            This method is useful when you don't control the construction of an
+            instance (ASP.NET pages or objects created via XAML, for instance)
+            but you still want properties and other injection performed.
+            </para></remarks>
+            <typeparam name="T"><see cref="T:System.Type"/> of object to perform injection on.</typeparam>
+            <param name="container">Container to resolve through.</param>
+            <param name="existing">Instance to build up.</param>
+            <param name="name">name to use when looking up the typemappings and other configurations.</param>
+            <param name="resolverOverrides">Any overrides for the Buildup.</param>
+            <returns>The resulting object. By default, this will be <paramref name="existing"/>, but
+            container extensions may add things like automatic proxy creation which would
+            cause this to return a different object (but still type compatible with <typeparamref name="T"/>).</returns>
+        </member>
+        <member name="M:Microsoft.Practices.Unity.UnityContainerExtensions.BuildUp(Microsoft.Practices.Unity.IUnityContainer,System.Type,System.Object,Microsoft.Practices.Unity.ResolverOverride[])">
+            <summary>
+            Run an existing object through the container and perform injection on it.
+            </summary>
+            <remarks>
+            <para>
+            This method is useful when you don't control the construction of an
+            instance (ASP.NET pages or objects created via XAML, for instance)
+            but you still want properties and other injection performed.
+            </para>
+            <para>
+            This overload uses the default registrations.
+            </para>
+            </remarks>
+            <param name="container">Container to resolve through.</param>
+            <param name="t"><see cref="T:System.Type"/> of object to perform injection on.</param>
+            <param name="existing">Instance to build up.</param>
+            <param name="resolverOverrides">Any overrides for the Buildup.</param>
+            <returns>The resulting object. By default, this will be <paramref name="existing"/>, but
+            container extensions may add things like automatic proxy creation which would
+            cause this to return a different object (but still type compatible with <paramref name="t"/>).</returns>
+        </member>
+        <member name="M:Microsoft.Practices.Unity.UnityContainerExtensions.AddNewExtension``1(Microsoft.Practices.Unity.IUnityContainer)">
+            <summary>
+            Creates a new extension object and adds it to the container.
+            </summary>
+            <typeparam name="TExtension">Type of <see cref="T:Microsoft.Practices.Unity.UnityContainerExtension"/> to add. The extension type
+            will be resolved from within the supplied <paramref name="container"/>.</typeparam>
+            <param name="container">Container to add the extension to.</param>
+            <returns>The <see cref="T:Microsoft.Practices.Unity.UnityContainer"/> object that this method was called on (this in C#, Me in Visual Basic).</returns>
+        </member>
+        <member name="M:Microsoft.Practices.Unity.UnityContainerExtensions.Configure``1(Microsoft.Practices.Unity.IUnityContainer)">
+            <summary>
+            Resolve access to a configuration interface exposed by an extension.
+            </summary>
+            <remarks>Extensions can expose configuration interfaces as well as adding
+            strategies and policies to the container. This method walks the list of
+            added extensions and returns the first one that implements the requested type.
+            </remarks>
+            <typeparam name="TConfigurator">The configuration interface required.</typeparam>
+            <param name="container">Container to configure.</param>
+            <returns>The requested extension's configuration interface, or null if not found.</returns>
+        </member>
+        <member name="M:Microsoft.Practices.Unity.UnityContainerExtensions.IsRegistered(Microsoft.Practices.Unity.IUnityContainer,System.Type)">
+            <summary>
+            Check if a particular type has been registered with the container with
+            the default name.
+            </summary>
+            <param name="container">Container to inspect.</param>
+            <param name="typeToCheck">Type to check registration for.</param>
+            <returns>True if this type has been registered, false if not.</returns>
+        </member>
+        <member name="M:Microsoft.Practices.Unity.UnityContainerExtensions.IsRegistered(Microsoft.Practices.Unity.IUnityContainer,System.Type,System.String)">
+            <summary>
+            Check if a particular type/name pair has been registered with the container.
+            </summary>
+            <param name="container">Container to inspect.</param>
+            <param name="typeToCheck">Type to check registration for.</param>
+            <param name="nameToCheck">Name to check registration for.</param>
+            <returns>True if this type/name pair has been registered, false if not.</returns>
+        </member>
+        <member name="M:Microsoft.Practices.Unity.UnityContainerExtensions.IsRegistered``1(Microsoft.Practices.Unity.IUnityContainer)">
+            <summary>
+            Check if a particular type has been registered with the container with the default name.
+            </summary>
+            <typeparam name="T">Type to check registration for.</typeparam>
+            <param name="container">Container to inspect.</param>
+            <returns>True if this type has been registered, false if not.</returns>
+        </member>
+        <member name="M:Microsoft.Practices.Unity.UnityContainerExtensions.IsRegistered``1(Microsoft.Practices.Unity.IUnityContainer,System.String)">
+            <summary>
+            Check if a particular type/name pair has been registered with the container.
+            </summary>
+            <typeparam name="T">Type to check registration for.</typeparam>
+            <param name="container">Container to inspect.</param>
+            <param name="nameToCheck">Name to check registration for.</param>
+            <returns>True if this type/name pair has been registered, false if not.</returns>
+        </member>
+        <member name="T:Microsoft.Practices.Unity.ExtensionContext">
+            <summary>
+            The <see cref="T:Microsoft.Practices.Unity.ExtensionContext"/> class provides the means for extension objects
+            to manipulate the internal state of the <see cref="T:Microsoft.Practices.Unity.UnityContainer"/>.
+            </summary>
+        </member>
+        <member name="M:Microsoft.Practices.Unity.ExtensionContext.RegisterNamedType(System.Type,System.String)">
+            <summary>
+            Store a type/name pair for later resolution.
+            </summary>
+            <remarks>
+            <para>
+            When users register type mappings (or other things) with a named key, this method
+            allows you to register that name with the container so that when the <see cref="M:Microsoft.Practices.Unity.IUnityContainer.ResolveAll(System.Type,Microsoft.Practices.Unity.ResolverOverride[])"/>
+            method is called, that name is included in the list that is returned.
+            </para></remarks>
+            <param name="t"><see cref="T:System.Type"/> to register.</param>
+            <param name="name">Name associated with that type.</param>
+        </member>
+        <member name="P:Microsoft.Practices.Unity.ExtensionContext.Container">
+            <summary>
+            The container that this context is associated with.
+            </summary>
+            <value>The <see cref="T:Microsoft.Practices.Unity.IUnityContainer"/> object.</value>
+        </member>
+        <member name="P:Microsoft.Practices.Unity.ExtensionContext.Strategies">
+            <summary>
+            The strategies this container uses.
+            </summary>
+            <value>The <see cref="T:Microsoft.Practices.ObjectBuilder2.StagedStrategyChain`1"/> that the container uses to build objects.</value>
+        </member>
+        <member name="P:Microsoft.Practices.Unity.ExtensionContext.BuildPlanStrategies">
+            <summary>
+            The strategies this container uses to construct build plans.
+            </summary>
+            <value>The <see cref="T:Microsoft.Practices.ObjectBuilder2.StagedStrategyChain`1"/> that this container uses when creating
+            build plans.</value>
+        </member>
+        <member name="P:Microsoft.Practices.Unity.ExtensionContext.Policies">
+            <summary>
+            The policies this container uses.
+            </summary>
+            <remarks>The <see cref="T:Microsoft.Practices.ObjectBuilder2.IPolicyList"/> the that container uses to build objects.</remarks>
+        </member>
+        <member name="P:Microsoft.Practices.Unity.ExtensionContext.Lifetime">
+            <summary>
+            The <see cref="T:Microsoft.Practices.ObjectBuilder2.ILifetimeContainer"/> that this container uses.
+            </summary>
+            <value>The <see cref="T:Microsoft.Practices.ObjectBuilder2.ILifetimeContainer"/> is used to manage <see cref="T:System.IDisposable"/> objects that the container is managing.</value>
+        </member>
+        <member name="E:Microsoft.Practices.Unity.ExtensionContext.Registering">
+            <summary>
+            This event is raised when the <see cref="M:Microsoft.Practices.Unity.IUnityContainer.RegisterType(System.Type,System.Type,System.String,Microsoft.Practices.Unity.LifetimeManager,Microsoft.Practices.Unity.InjectionMember[])"/> method,
+            or one of its overloads, is called.
+            </summary>
+        </member>
+        <member name="E:Microsoft.Practices.Unity.ExtensionContext.RegisteringInstance">
+            <summary>
+            This event is raised when the <see cref="M:Microsoft.Practices.Unity.IUnityContainer.RegisterInstance(System.Type,System.String,System.Object,Microsoft.Practices.Unity.LifetimeManager)"/> method,
+            or one of its overloads, is called.
+            </summary>
+        </member>
+        <member name="E:Microsoft.Practices.Unity.ExtensionContext.ChildContainerCreated">
+            <summary>
+            This event is raised when the <see cref="M:Microsoft.Practices.Unity.IUnityContainer.CreateChildContainer"/> method is called, providing 
+            the newly created child container to extensions to act on as they see fit.
+            </summary>
+        </member>
+        <member name="T:Microsoft.Practices.Unity.NamedEventArgs">
+            <summary>
+            An EventArgs class that holds a string Name.
+            </summary>
+        </member>
+        <member name="M:Microsoft.Practices.Unity.NamedEventArgs.#ctor">
+            <summary>
+            Create a new <see cref="T:Microsoft.Practices.Unity.NamedEventArgs"/> with a null name.
+            </summary>
+        </member>
+        <member name="M:Microsoft.Practices.Unity.NamedEventArgs.#ctor(System.String)">
+            <summary>
+            Create a new <see cref="T:Microsoft.Practices.Unity.NamedEventArgs"/> with the given name.
+            </summary>
+            <param name="name">Name to store.</param>
+        </member>
+        <member name="P:Microsoft.Practices.Unity.NamedEventArgs.Name">
+            <summary>
+            The name.
+            </summary>
+            <value>Name used for this EventArg object.</value>
+        </member>
+        <member name="T:Microsoft.Practices.Unity.RegisterEventArgs">
+            <summary>
+            Event argument class for the <see cref="E:Microsoft.Practices.Unity.ExtensionContext.Registering"/> event.
+            </summary>
+        </member>
+        <member name="M:Microsoft.Practices.Unity.RegisterEventArgs.#ctor(System.Type,System.Type,System.String,Microsoft.Practices.Unity.LifetimeManager)">
+            <summary>
+            Create a new instance of <see cref="T:Microsoft.Practices.Unity.RegisterEventArgs"/>.
+            </summary>
+            <param name="typeFrom">Type to map from.</param>
+            <param name="typeTo">Type to map to.</param>
+            <param name="name">Name for the registration.</param>
+            <param name="lifetimeManager"><see cref="P:Microsoft.Practices.Unity.RegisterEventArgs.LifetimeManager"/> to manage instances.</param>
+        </member>
+        <member name="P:Microsoft.Practices.Unity.RegisterEventArgs.TypeFrom">
+            <summary>
+            Type to map from.
+            </summary>
+        </member>
+        <member name="P:Microsoft.Practices.Unity.RegisterEventArgs.TypeTo">
+            <summary>
+            Type to map to.
+            </summary>
+        </member>
+        <member name="P:Microsoft.Practices.Unity.RegisterEventArgs.LifetimeManager">
+            <summary>
+            <see cref="P:Microsoft.Practices.Unity.RegisterEventArgs.LifetimeManager"/> to manage instances.
+            </summary>
+        </member>
+        <member name="T:Microsoft.Practices.Unity.RegisterInstanceEventArgs">
+            <summary>
+            Event argument class for the <see cref="E:Microsoft.Practices.Unity.ExtensionContext.RegisteringInstance"/> event.
+            </summary>
+        </member>
+        <member name="M:Microsoft.Practices.Unity.RegisterInstanceEventArgs.#ctor">
+            <summary>
+            Create a default <see cref="T:Microsoft.Practices.Unity.RegisterInstanceEventArgs"/> instance.
+            </summary>
+        </member>
+        <member name="M:Microsoft.Practices.Unity.RegisterInstanceEventArgs.#ctor(System.Type,System.Object,System.String,Microsoft.Practices.Unity.LifetimeManager)">
+            <summary>
+            Create a <see cref="T:Microsoft.Practices.Unity.RegisterInstanceEventArgs"/> instance initialized with the given arguments.
+            </summary>
+            <param name="registeredType">Type of instance being registered.</param>
+            <param name="instance">The instance object itself.</param>
+            <param name="name">Name to register under, null if default registration.</param>
+            <param name="lifetimeManager"><see cref="P:Microsoft.Practices.Unity.RegisterInstanceEventArgs.LifetimeManager"/> object that handles how
+            the instance will be owned.</param>
+        </member>
+        <member name="P:Microsoft.Practices.Unity.RegisterInstanceEventArgs.RegisteredType">
+            <summary>
+            Type of instance being registered.
+            </summary>
+            <value>
+            Type of instance being registered.
+            </value>
+        </member>
+        <member name="P:Microsoft.Practices.Unity.RegisterInstanceEventArgs.Instance">
+            <summary>
+            Instance object being registered.
+            </summary>
+            <value>Instance object being registered</value>
+        </member>
+        <member name="P:Microsoft.Practices.Unity.RegisterInstanceEventArgs.LifetimeManager">
+            <summary>
+            <see cref="T:Microsoft.Practices.Unity.LifetimeManager"/> that controls ownership of
+            this instance.
+            </summary>
+        </member>
+        <member name="T:Microsoft.Practices.Unity.GenericParameter">
+            <summary>
+            A <see cref="T:Microsoft.Practices.Unity.InjectionParameterValue"/> that lets you specify that
+            an instance of a generic type parameter should be resolved.
+            </summary>
+        </member>
+        <member name="M:Microsoft.Practices.Unity.GenericParameter.#ctor(System.String)">
+            <summary>
+            Create a new <see cref="T:Microsoft.Practices.Unity.GenericParameter"/> instance that specifies
+            that the given named generic parameter should be resolved.
+            </summary>
+            <param name="genericParameterName">The generic parameter name to resolve.</param>
+        </member>
+        <member name="M:Microsoft.Practices.Unity.GenericParameter.#ctor(System.String,System.String)">
+            <summary>
+            Create a new <see cref="T:Microsoft.Practices.Unity.GenericParameter"/> instance that specifies
+            that the given named generic parameter should be resolved.
+            </summary>
+            <param name="genericParameterName">The generic parameter name to resolve.</param>
+            <param name="resolutionKey">name to use when looking up in the container.</param>
+        </member>
+        <member name="M:Microsoft.Practices.Unity.GenericParameter.DoGetResolverPolicy(System.Type,System.String)">
+            <summary>
+            Return a <see cref="T:Microsoft.Practices.ObjectBuilder2.IDependencyResolverPolicy"/> instance that will
+            return this types value for the parameter.
+            </summary>
+            <param name="typeToResolve">The actual type to resolve.</param>
+            <param name="resolutionKey">The resolution key.</param>
+            <returns>The <see cref="T:Microsoft.Practices.ObjectBuilder2.IDependencyResolverPolicy"/>.</returns>
+        </member>
+        <member name="T:Microsoft.Practices.Unity.GenericResolvedArrayParameter">
+            <summary>
+            A <see cref="T:Microsoft.Practices.Unity.InjectionParameterValue"/> that lets you specify that
+            an array containing the registered instances of a generic type parameter 
+            should be resolved.
+            </summary>
+        </member>
+        <member name="M:Microsoft.Practices.Unity.GenericResolvedArrayParameter.#ctor(System.String,System.Object[])">
+            <summary>
+            Create a new <see cref="T:Microsoft.Practices.Unity.GenericResolvedArrayParameter"/> instance that specifies
+            that the given named generic parameter should be resolved.
+            </summary>
+            <param name="genericParameterName">The generic parameter name to resolve.</param>
+            <param name="elementValues">The values for the elements, that will
+            be converted to <see cref="T:Microsoft.Practices.Unity.InjectionParameterValue"/> objects.</param>
+        </member>
+        <member name="M:Microsoft.Practices.Unity.GenericResolvedArrayParameter.MatchesType(System.Type)">
+            <summary>
+            Test to see if this parameter value has a matching type for the given type.
+            </summary>
+            <param name="t">Type to check.</param>
+            <returns>True if this parameter value is compatible with type <paramref name="t"/>,
+            false if not.</returns>
+            <remarks>A type is considered compatible if it is an array type of rank one
+            and its element type is a generic type parameter with a name matching this generic
+            parameter name configured for the receiver.</remarks>
+        </member>
+        <member name="M:Microsoft.Practices.Unity.GenericResolvedArrayParameter.GetResolverPolicy(System.Type)">
+            <summary>
+            Return a <see cref="T:Microsoft.Practices.ObjectBuilder2.IDependencyResolverPolicy"/> instance that will
+            return this types value for the parameter.
+            </summary>
+            <param name="typeToBuild">Type that contains the member that needs this parameter. Used
+            to resolve open generic parameters.</param>
+            <returns>The <see cref="T:Microsoft.Practices.ObjectBuilder2.IDependencyResolverPolicy"/>.</returns>
+        </member>
+        <member name="P:Microsoft.Practices.Unity.GenericResolvedArrayParameter.ParameterTypeName">
+            <summary>
+            Name for the type represented by this <see cref="T:Microsoft.Practices.Unity.InjectionParameterValue"/>.
+            This may be an actual type name or a generic argument name.
+            </summary>
+        </member>
+        <member name="T:Microsoft.Practices.Unity.InjectedMembers">
+            <summary>
+            A Unity container extension that allows you to configure
+            which constructors, properties, and methods get injected
+            via an API rather than through attributes.
+            </summary>
+        </member>
+        <member name="M:Microsoft.Practices.Unity.InjectedMembers.Initialize">
+            <summary>
+            Initial the container with this extension's functionality.
+            </summary>
+            <remarks>
+            When overridden in a derived class, this method will modify the given
+            <see cref="T:Microsoft.Practices.Unity.ExtensionContext"/> by adding strategies, policies, etc. to
+            install it's functions into the container.</remarks>
+        </member>
+        <member name="M:Microsoft.Practices.Unity.InjectedMembers.ConfigureInjectionFor``1(Microsoft.Practices.Unity.InjectionMember[])">
+            <summary>
+            API to configure the injection settings for a particular type.
+            </summary>
+            <typeparam name="TTypeToInject">Type the injection is being configured for.</typeparam>
+            <param name="injectionMembers">Objects containing the details on which members to inject and how.</param>
+            <returns>This extension object.</returns>
+        </member>
+        <member name="M:Microsoft.Practices.Unity.InjectedMembers.ConfigureInjectionFor``1(System.String,Microsoft.Practices.Unity.InjectionMember[])">
+            <summary>
+            API to configure the injection settings for a particular type/name pair.
+            </summary>
+            <typeparam name="TTypeToInject">Type the injection is being configured for.</typeparam>
+            <param name="name">Name of registration</param>
+            <param name="injectionMembers">Objects containing the details on which members to inject and how.</param>
+            <returns>This extension object.</returns>
+        </member>
+        <member name="M:Microsoft.Practices.Unity.InjectedMembers.ConfigureInjectionFor(System.Type,Microsoft.Practices.Unity.InjectionMember[])">
+            <summary>
+            API to configure the injection settings for a particular type.
+            </summary>
+            <param name="typeToInject">Type to configure.</param>
+            <param name="injectionMembers">Objects containing the details on which members to inject and how.</param>
+            <returns>This extension object.</returns>
+        </member>
+        <member name="M:Microsoft.Practices.Unity.InjectedMembers.ConfigureInjectionFor(System.Type,System.String,Microsoft.Practices.Unity.InjectionMember[])">
+            <summary>
+            API to configure the injection settings for a particular type/name pair.
+            </summary>
+            <param name="typeToInject">Type to configure.</param>
+            <param name="name">Name of registration.</param>
+            <param name="injectionMembers">Objects containing the details on which members to inject and how.</param>
+            <returns>This extension object.</returns>
+        </member>
+        <member name="M:Microsoft.Practices.Unity.InjectedMembers.ConfigureInjectionFor(System.Type,System.Type,System.String,Microsoft.Practices.Unity.InjectionMember[])">
+            <summary>
+            API to configure the injection settings for a particular type/name pair.
+            </summary>
+            <param name="serviceType">Type of interface/base class being registered (may be null).</param>
+            <param name="implementationType">Type of actual implementation class being registered.</param>
+            <param name="name">Name of registration.</param>
+            <param name="injectionMembers">Objects containing the details on which members to inject and how.</param>
+            <returns>This extension object.</returns>
+        </member>
+        <member name="T:Microsoft.Practices.Unity.InjectionConstructor">
+            <summary>
+            A class that holds the collection of information
+            for a constructor, so that the container can
+            be configured to call this constructor.
+            </summary>
+        </member>
+        <member name="M:Microsoft.Practices.Unity.InjectionConstructor.#ctor(System.Object[])">
+            <summary>
+            Create a new instance of <see cref="T:Microsoft.Practices.Unity.InjectionConstructor"/> that looks
+            for a constructor with the given set of parameters.
+            </summary>
+            <param name="parameterValues">The values for the parameters, that will
+            be converted to <see cref="T:Microsoft.Practices.Unity.InjectionParameterValue"/> objects.</param>
+        </member>
+        <member name="M:Microsoft.Practices.Unity.InjectionConstructor.AddPolicies(System.Type,System.Type,System.String,Microsoft.Practices.ObjectBuilder2.IPolicyList)">
+            <summary>
+            Add policies to the <paramref name="policies"/> to configure the
+            container to call this constructor with the appropriate parameter values.
+            </summary>
+            <param name="serviceType">Interface registered, ignored in this implementation.</param>
+            <param name="implementationType">Type to register.</param>
+            <param name="name">Name used to resolve the type object.</param>
+            <param name="policies">Policy list to add policies to.</param>
+        </member>
+        <member name="T:Microsoft.Practices.Unity.InjectionMethod">
+            <summary>
+            An <see cref="T:Microsoft.Practices.Unity.InjectionMember"/> that configures the
+            container to call a method as part of buildup.
+            </summary>
+        </member>
+        <member name="M:Microsoft.Practices.Unity.InjectionMethod.#ctor(System.String,System.Object[])">
+            <summary>
+            Create a new <see cref="T:Microsoft.Practices.Unity.InjectionMethod"/> instance which will configure
+            the container to call the given methods with the given parameters.
+            </summary>
+            <param name="methodName">Name of the method to call.</param>
+            <param name="methodParameters">Parameter values for the method.</param>
+        </member>
+        <member name="M:Microsoft.Practices.Unity.InjectionMethod.AddPolicies(System.Type,System.Type,System.String,Microsoft.Practices.ObjectBuilder2.IPolicyList)">
+            <summary>
+            Add policies to the <paramref name="policies"/> to configure the
+            container to call this constructor with the appropriate parameter values.
+            </summary>
+            <param name="serviceType">Type of interface registered, ignored in this implementation.</param>
+            <param name="implementationType">Type to register.</param>
+            <param name="name">Name used to resolve the type object.</param>
+            <param name="policies">Policy list to add policies to.</param>
+        </member>
+        <member name="M:Microsoft.Practices.Unity.InjectionMethod.MethodNameMatches(System.Reflection.MemberInfo,System.String)">
+            <summary>
+            A small function to handle name matching. You can override this
+            to do things like case insensitive comparisons.
+            </summary>
+            <param name="targetMethod">MethodInfo for the method you're checking.</param>
+            <param name="nameToMatch">Name of the method you're looking for.</param>
+            <returns>True if a match, false if not.</returns>
+        </member>
+        <member name="T:Microsoft.Practices.Unity.InjectionParameter">
+            <summary>
+            A class that holds on to the given value and provides
+            the required <see cref="T:Microsoft.Practices.ObjectBuilder2.IDependencyResolverPolicy"/>
+            when the container is configured.
+            </summary>
+        </member>
+        <member name="M:Microsoft.Practices.Unity.InjectionParameter.#ctor(System.Object)">
+            <summary>
+            Create an instance of <see cref="T:Microsoft.Practices.Unity.InjectionParameter"/> that stores
+            the given value, using the runtime type of that value as the
+            type of the parameter.
+            </summary>
+            <param name="parameterValue">Value to be injected for this parameter.</param>
+        </member>
+        <member name="M:Microsoft.Practices.Unity.InjectionParameter.#ctor(System.Type,System.Object)">
+            <summary>
+            Create an instance of <see cref="T:Microsoft.Practices.Unity.InjectionParameter"/> that stores
+            the given value, associated with the given type.
+            </summary>
+            <param name="parameterType">Type of the parameter.</param>
+            <param name="parameterValue">Value of the parameter</param>
+        </member>
+        <member name="M:Microsoft.Practices.Unity.InjectionParameter.GetResolverPolicy(System.Type)">
+            <summary>
+            Return a <see cref="T:Microsoft.Practices.ObjectBuilder2.IDependencyResolverPolicy"/> instance that will
+            return this types value for the parameter.
+            </summary>
+            <param name="typeToBuild">Type that contains the member that needs this parameter. Used
+            to resolve open generic parameters.</param>
+            <returns>The <see cref="T:Microsoft.Practices.ObjectBuilder2.IDependencyResolverPolicy"/>.</returns>
+        </member>
+        <member name="T:Microsoft.Practices.Unity.InjectionParameter`1">
+            <summary>
+            A generic version of <see cref="T:Microsoft.Practices.Unity.InjectionParameter"/> that makes it a
+            little easier to specify the type of the parameter.
+            </summary>
+            <typeparam name="TParameter">Type of parameter.</typeparam>
+        </member>
+        <member name="M:Microsoft.Practices.Unity.InjectionParameter`1.#ctor(`0)">
+            <summary>
+            Create a new <see cref="T:Microsoft.Practices.Unity.InjectionParameter`1"/>.
+            </summary>
+            <param name="parameterValue">Value for the parameter.</param>
+        </member>
+        <member name="T:Microsoft.Practices.Unity.InjectionProperty">
+            <summary>
+            This class stores information about which properties to inject,
+            and will configure the container accordingly.
+            </summary>
+        </member>
+        <member name="M:Microsoft.Practices.Unity.InjectionProperty.#ctor(System.String)">
+            <summary>
+            Configure the container to inject the given property name,
+            resolving the value via the container.
+            </summary>
+            <param name="propertyName">Name of the property to inject.</param>
+        </member>
+        <member name="M:Microsoft.Practices.Unity.InjectionProperty.#ctor(System.String,System.Object)">
+            <summary>
+            Configure the container to inject the given property name,
+            using the value supplied. This value is converted to an
+            <see cref="T:Microsoft.Practices.Unity.InjectionParameterValue"/> object using the
+            rules defined by the <see cref="M:Microsoft.Practices.Unity.InjectionParameterValue.ToParameters(System.Object[])"/>
+            method.
+            </summary>
+            <param name="propertyName">Name of property to inject.</param>
+            <param name="propertyValue">Value for property.</param>
+        </member>
+        <member name="M:Microsoft.Practices.Unity.InjectionProperty.AddPolicies(System.Type,System.Type,System.String,Microsoft.Practices.ObjectBuilder2.IPolicyList)">
+            <summary>
+            Add policies to the <paramref name="policies"/> to configure the
+            container to call this constructor with the appropriate parameter values.
+            </summary>
+            <param name="serviceType">Interface being registered, ignored in this implementation.</param>
+            <param name="implementationType">Type to register.</param>
+            <param name="name">Name used to resolve the type object.</param>
+            <param name="policies">Policy list to add policies to.</param>
+        </member>
+        <member name="T:Microsoft.Practices.Unity.ResolvedArrayParameter">
+            <summary>
+            A class that stores a type, and generates a 
+            resolver object that resolves all the named instances or the
+            type registered in a container.
+            </summary>
+        </member>
+        <member name="M:Microsoft.Practices.Unity.ResolvedArrayParameter.#ctor(System.Type,System.Object[])">
+            <summary>
+            Construct a new <see cref="T:Microsoft.Practices.Unity.ResolvedArrayParameter"/> that
+            resolves to the given element type and collection of element values.
+            </summary>
+            <param name="elementType">The type of elements to resolve.</param>
+            <param name="elementValues">The values for the elements, that will
+            be converted to <see cref="T:Microsoft.Practices.Unity.InjectionParameterValue"/> objects.</param>
+        </member>
+        <member name="M:Microsoft.Practices.Unity.ResolvedArrayParameter.#ctor(System.Type,System.Type,System.Object[])">
+            <summary>
+            Construct a new <see cref="T:Microsoft.Practices.Unity.ResolvedArrayParameter"/> that
+            resolves to the given array and element types and collection of element values.
+            </summary>
+            <param name="arrayParameterType">The type for the array of elements to resolve.</param>
+            <param name="elementType">The type of elements to resolve.</param>
+            <param name="elementValues">The values for the elements, that will
+            be converted to <see cref="T:Microsoft.Practices.Unity.InjectionParameterValue"/> objects.</param>
+        </member>
+        <member name="M:Microsoft.Practices.Unity.ResolvedArrayParameter.GetResolverPolicy(System.Type)">
+            <summary>
+            Return a <see cref="T:Microsoft.Practices.ObjectBuilder2.IDependencyResolverPolicy"/> instance that will
+            return this types value for the parameter.
+            </summary>
+            <param name="typeToBuild">Type that contains the member that needs this parameter. Used
+            to resolve open generic parameters.</param>
+            <returns>The <see cref="T:Microsoft.Practices.ObjectBuilder2.IDependencyResolverPolicy"/>.</returns>
+        </member>
+        <member name="T:Microsoft.Practices.Unity.ResolvedArrayParameter`1">
+            <summary>
+            A generic version of <see cref="T:Microsoft.Practices.Unity.ResolvedArrayParameter"/> for convenience
+            when creating them by hand.
+            </summary>
+            <typeparam name="TElement">Type of the elements for the array of the parameter.</typeparam>
+        </member>
+        <member name="M:Microsoft.Practices.Unity.ResolvedArrayParameter`1.#ctor(System.Object[])">
+            <summary>
+            Construct a new <see cref="T:Microsoft.Practices.Unity.ResolvedArrayParameter`1"/> that
+            resolves to the given element generic type with the given element values.
+            </summary>
+            <param name="elementValues">The values for the elements, that will
+            be converted to <see cref="T:Microsoft.Practices.Unity.InjectionParameterValue"/> objects.</param>
+        </member>
+        <member name="T:Microsoft.Practices.Unity.IUnityContainer">
+            <summary>
+            Interface defining the behavior of the Unity dependency injection container.
+            </summary>
+        </member>
+        <member name="M:Microsoft.Practices.Unity.IUnityContainer.RegisterType(System.Type,System.Type,System.String,Microsoft.Practices.Unity.LifetimeManager,Microsoft.Practices.Unity.InjectionMember[])">
+            <summary>
+            Register a type mapping with the container, where the created instances will use
+            the given <see cref="T:Microsoft.Practices.Unity.LifetimeManager"/>.
+            </summary>
+            <param name="from"><see cref="T:System.Type"/> that will be requested.</param>
+            <param name="to"><see cref="T:System.Type"/> that will actually be returned.</param>
+            <param name="name">Name to use for registration, null if a default registration.</param>
+            <param name="lifetimeManager">The <see cref="T:Microsoft.Practices.Unity.LifetimeManager"/> that controls the lifetime
+            of the returned instance.</param>
+            <param name="injectionMembers">Injection configuration objects.</param>
+            <returns>The <see cref="T:Microsoft.Practices.Unity.UnityContainer"/> object that this method was called on (this in C#, Me in Visual Basic).</returns>
+        </member>
+        <member name="M:Microsoft.Practices.Unity.IUnityContainer.RegisterInstance(System.Type,System.String,System.Object,Microsoft.Practices.Unity.LifetimeManager)">
+            <summary>
+            Register an instance with the container.
+            </summary>
+            <remarks>
+            <para>
+            Instance registration is much like setting a type as a singleton, except that instead
+            of the container creating the instance the first time it is requested, the user
+            creates the instance ahead of type and adds that instance to the container.
+            </para>
+            </remarks>
+            <param name="t">Type of instance to register (may be an implemented interface instead of the full type).</param>
+            <param name="instance">Object to returned.</param>
+            <param name="name">Name for registration.</param>
+            <param name="lifetime">
+            <see cref="T:Microsoft.Practices.Unity.LifetimeManager"/> object that controls how this instance will be managed by the container.</param>
+            <returns>The <see cref="T:Microsoft.Practices.Unity.UnityContainer"/> object that this method was called on (this in C#, Me in Visual Basic).</returns>
+        </member>
+        <member name="M:Microsoft.Practices.Unity.IUnityContainer.Resolve(System.Type,System.String,Microsoft.Practices.Unity.ResolverOverride[])">
+            <summary>
+            Resolve an instance of the requested type with the given name from the container.
+            </summary>
+            <param name="t"><see cref="T:System.Type"/> of object to get from the container.</param>
+            <param name="name">Name of the object to retrieve.</param>
+            <param name="resolverOverrides">Any overrides for the resolve call.</param>
+            <returns>The retrieved object.</returns>
+        </member>
+        <member name="M:Microsoft.Practices.Unity.IUnityContainer.ResolveAll(System.Type,Microsoft.Practices.Unity.ResolverOverride[])">
+            <summary>
+            Return instances of all registered types requested.
+            </summary>
+            <remarks>
+            <para>
+            This method is useful if you've registered multiple types with the same
+            <see cref="T:System.Type"/> but different names.
+            </para>
+            <para>
+            Be aware that this method does NOT return an instance for the default (unnamed) registration.
+            </para>
+            </remarks>
+            <param name="t">The type requested.</param>
+            <param name="resolverOverrides">Any overrides for the resolve calls.</param>
+            <returns>Set of objects of type <paramref name="t"/>.</returns>
+        </member>
+        <member name="M:Microsoft.Practices.Unity.IUnityContainer.BuildUp(System.Type,System.Object,System.String,Microsoft.Practices.Unity.ResolverOverride[])">
+            <summary>
+            Run an existing object through the container and perform injection on it.
+            </summary>
+            <remarks>
+            <para>
+            This method is useful when you don't control the construction of an
+            instance (ASP.NET pages or objects created via XAML, for instance)
+            but you still want properties and other injection performed.
+            </para></remarks>
+            <param name="t"><see cref="T:System.Type"/> of object to perform injection on.</param>
+            <param name="existing">Instance to build up.</param>
+            <param name="name">name to use when looking up the TypeMappings and other configurations.</param>
+            <param name="resolverOverrides">Any overrides for the resolve calls.</param>
+            <returns>The resulting object. By default, this will be <paramref name="existing"/>, but
+            container extensions may add things like automatic proxy creation which would
+            cause this to return a different object (but still type compatible with <paramref name="t"/>).</returns>
+        </member>
+        <member name="M:Microsoft.Practices.Unity.IUnityContainer.Teardown(System.Object)">
+            <summary>
+            Run an existing object through the container, and clean it up.
+            </summary>
+            <param name="o">The object to tear down.</param>
+        </member>
+        <member name="M:Microsoft.Practices.Unity.IUnityContainer.AddExtension(Microsoft.Practices.Unity.UnityContainerExtension)">
+            <summary>
+            Add an extension object to the container.
+            </summary>
+            <param name="extension"><see cref="T:Microsoft.Practices.Unity.UnityContainerExtension"/> to add.</param>
+            <returns>The <see cref="T:Microsoft.Practices.Unity.UnityContainer"/> object that this method was called on (this in C#, Me in Visual Basic).</returns>
+        </member>
+        <member name="M:Microsoft.Practices.Unity.IUnityContainer.Configure(System.Type)">
+            <summary>
+            Resolve access to a configuration interface exposed by an extension.
+            </summary>
+            <remarks>Extensions can expose configuration interfaces as well as adding
+            strategies and policies to the container. This method walks the list of
+            added extensions and returns the first one that implements the requested type.
+            </remarks>
+            <param name="configurationInterface"><see cref="T:System.Type"/> of configuration interface required.</param>
+            <returns>The requested extension's configuration interface, or null if not found.</returns>
+        </member>
+        <member name="M:Microsoft.Practices.Unity.IUnityContainer.RemoveAllExtensions">
+            <summary>
+            Remove all installed extensions from this container.
+            </summary>
+            <remarks>
+            <para>
+            This method removes all extensions from the container, including the default ones
+            that implement the out-of-the-box behavior. After this method, if you want to use
+            the container again you will need to either read the default extensions or replace
+            them with your own.
+            </para>
+            <para>
+            The registered instances and singletons that have already been set up in this container
+            do not get removed.
+            </para>
+            </remarks>
+            <returns>The <see cref="T:Microsoft.Practices.Unity.UnityContainer"/> object that this method was called on (this in C#, Me in Visual Basic).</returns>
+        </member>
+        <member name="M:Microsoft.Practices.Unity.IUnityContainer.CreateChildContainer">
+            <summary>
+            Create a child container.
+            </summary>
+            <remarks>
+            A child container shares the parent's configuration, but can be configured with different
+            settings or lifetime.</remarks>
+            <returns>The new child container.</returns>
+        </member>
+        <member name="P:Microsoft.Practices.Unity.IUnityContainer.Parent">
+            <summary>
+            The parent of this container.
+            </summary>
+            <value>The parent container, or null if this container doesn't have one.</value>
+        </member>
+        <member name="P:Microsoft.Practices.Unity.IUnityContainer.Registrations">
+            <summary>
+            Get a sequence of <see cref="T:Microsoft.Practices.Unity.ContainerRegistration"/> that describe the current state
+            of the container.
+            </summary>
+        </member>
+        <member name="T:Microsoft.Practices.Unity.ExternallyControlledLifetimeManager">
+            <summary>
+            A <see cref="T:Microsoft.Practices.Unity.LifetimeManager"/> that holds a weak reference to
+            it's managed instance.
+            </summary>
+        </member>
+        <member name="M:Microsoft.Practices.Unity.ExternallyControlledLifetimeManager.GetValue">
+            <summary>
+            Retrieve a value from the backing store associated with this Lifetime policy.
+            </summary>
+            <returns>the object desired, or null if no such object is currently stored.</returns>
+        </member>
+        <member name="M:Microsoft.Practices.Unity.ExternallyControlledLifetimeManager.SetValue(System.Object)">
+            <summary>
+            Stores the given value into backing store for retrieval later.
+            </summary>
+            <param name="newValue">The object being stored.</param>
+        </member>
+        <member name="M:Microsoft.Practices.Unity.ExternallyControlledLifetimeManager.RemoveValue">
+            <summary>
+            Remove the given object from backing store.
+            </summary>
+        </member>
+        <member name="T:Microsoft.Practices.Unity.LifetimeManagerFactory">
+            <summary>
+            An implementation of <see cref="T:Microsoft.Practices.ObjectBuilder2.ILifetimeFactoryPolicy"/> that
+            creates instances of the type of the given Lifetime Manager
+            by resolving them through the container.
+            </summary>
+        </member>
+        <member name="M:Microsoft.Practices.Unity.LifetimeManagerFactory.#ctor(Microsoft.Practices.Unity.ExtensionContext,System.Type)">
+            <summary>
+            Create a new <see cref="T:Microsoft.Practices.Unity.LifetimeManagerFactory"/> that will
+            return instances of the given type, creating them by
+            resolving through the container.
+            </summary>
+            <param name="containerContext">Container to resolve with.</param>
+            <param name="lifetimeType">Type of LifetimeManager to create.</param>
+        </member>
+        <member name="M:Microsoft.Practices.Unity.LifetimeManagerFactory.CreateLifetimePolicy">
+            <summary>
+            Create a new instance of <see cref="T:Microsoft.Practices.ObjectBuilder2.ILifetimePolicy"/>.
+            </summary>
+            <returns>The new instance.</returns>
+        </member>
+        <member name="P:Microsoft.Practices.Unity.LifetimeManagerFactory.LifetimeType">
+            <summary>
+            The type of Lifetime manager that will be created by this factory.
+            </summary>
+        </member>
+        <member name="T:Microsoft.Practices.Unity.PerThreadLifetimeManager">
+            <summary>
+            A <see cref="T:Microsoft.Practices.Unity.LifetimeManager"/> that holds the instances given to it, 
+            keeping one instance per thread.
+            </summary>
+            <remarks>
+            <para>
+            This LifetimeManager does not dispose the instances it holds.
+            </para>
+            </remarks>
+        </member>
+        <member name="M:Microsoft.Practices.Unity.PerThreadLifetimeManager.#ctor">
+            <summary>
+            Initializes a new instance of the <see cref="T:Microsoft.Practices.Unity.PerThreadLifetimeManager"/> class.
+            </summary>
+        </member>
+        <member name="M:Microsoft.Practices.Unity.PerThreadLifetimeManager.GetValue">
+            <summary>
+            Retrieve a value from the backing store associated with this Lifetime policy for the 
+            current thread.
+            </summary>
+            <returns>the object desired, or <see langword="null"/> if no such object is currently 
+            stored for the current thread.</returns>
+        </member>
+        <member name="M:Microsoft.Practices.Unity.PerThreadLifetimeManager.SetValue(System.Object)">
+            <summary>
+            Stores the given value into backing store for retrieval later when requested
+            in the current thread.
+            </summary>
+            <param name="newValue">The object being stored.</param>
+        </member>
+        <member name="M:Microsoft.Practices.Unity.PerThreadLifetimeManager.RemoveValue">
+            <summary>
+            Remove the given object from backing store.
+            </summary>
+            <remarks>Not implemented for this lifetime manager.</remarks>
+        </member>
+        <member name="T:Microsoft.Practices.Unity.TransientLifetimeManager">
+            <summary>
+            An <see cref="T:Microsoft.Practices.Unity.LifetimeManager"/> implementation that does nothing,
+            thus ensuring that instances are created new every time.
+            </summary>
+        </member>
+        <member name="M:Microsoft.Practices.Unity.TransientLifetimeManager.GetValue">
+            <summary>
+            Retrieve a value from the backing store associated with this Lifetime policy.
+            </summary>
+            <returns>the object desired, or null if no such object is currently stored.</returns>
+        </member>
+        <member name="M:Microsoft.Practices.Unity.TransientLifetimeManager.SetValue(System.Object)">
+            <summary>
+            Stores the given value into backing store for retrieval later.
+            </summary>
+            <param name="newValue">The object being stored.</param>
+        </member>
+        <member name="M:Microsoft.Practices.Unity.TransientLifetimeManager.RemoveValue">
+            <summary>
+            Remove the given object from backing store.
+            </summary>
+        </member>
+        <member name="T:Microsoft.Practices.Unity.ArrayResolutionStrategy">
+            <summary>
+            This strategy implements the logic that will call container.ResolveAll
+            when an array parameter is detected.
+            </summary>
+        </member>
+        <member name="M:Microsoft.Practices.Unity.ArrayResolutionStrategy.PreBuildUp(Microsoft.Practices.ObjectBuilder2.IBuilderContext)">
+            <summary>
+            Do the PreBuildUp stage of construction. This is where the actual work is performed.
+            </summary>
+            <param name="context">Current build context.</param>
+        </member>
+        <member name="T:Microsoft.Practices.Unity.ObjectBuilder.DefaultUnityConstructorSelectorPolicy">
+            <summary>
+            An implementation of <see cref="T:Microsoft.Practices.ObjectBuilder2.IConstructorSelectorPolicy"/> that is
+            aware of the build keys used by the Unity container.
+            </summary>
+        </member>
+        <member name="M:Microsoft.Practices.Unity.ObjectBuilder.DefaultUnityConstructorSelectorPolicy.CreateResolver(System.Reflection.ParameterInfo)">
+            <summary>
+            Create a <see cref="T:Microsoft.Practices.ObjectBuilder2.IDependencyResolverPolicy"/> instance for the given
+            <see cref="T:System.Reflection.ParameterInfo"/>.
+            </summary>
+            <remarks>
+            This implementation looks for the Unity <see cref="T:Microsoft.Practices.Unity.DependencyAttribute"/> on the
+            parameter and uses it to create an instance of <see cref="T:Microsoft.Practices.Unity.ObjectBuilder.NamedTypeDependencyResolverPolicy"/>
+            for this parameter.</remarks>
+            <param name="parameter">Parameter to create the resolver for.</param>
+            <returns>The resolver object.</returns>
+        </member>
+        <member name="T:Microsoft.Practices.Unity.ObjectBuilder.DefaultUnityMethodSelectorPolicy">
+            <summary>
+            An implementation of <see cref="T:Microsoft.Practices.ObjectBuilder2.IMethodSelectorPolicy"/> that is aware
+            of the build keys used by the Unity container.
+            </summary>
+        </member>
+        <member name="M:Microsoft.Practices.Unity.ObjectBuilder.DefaultUnityMethodSelectorPolicy.CreateResolver(System.Reflection.ParameterInfo)">
+            <summary>
+            Create a <see cref="T:Microsoft.Practices.ObjectBuilder2.IDependencyResolverPolicy"/> instance for the given
+            <see cref="T:System.Reflection.ParameterInfo"/>.
+            </summary>
+            <param name="parameter">Parameter to create the resolver for.</param>
+            <returns>The resolver object.</returns>
+        </member>
+        <member name="T:Microsoft.Practices.Unity.ObjectBuilder.DefaultUnityPropertySelectorPolicy">
+            <summary>
+            An implementation of <see cref="T:Microsoft.Practices.ObjectBuilder2.IPropertySelectorPolicy"/> that is aware of
+            the build keys used by the unity container.
+            </summary>
+        </member>
+        <member name="M:Microsoft.Practices.Unity.ObjectBuilder.DefaultUnityPropertySelectorPolicy.CreateResolver(System.Reflection.PropertyInfo)">
+            <summary>
+            Create a <see cref="T:Microsoft.Practices.ObjectBuilder2.IDependencyResolverPolicy"/> for the given
+            property.
+            </summary>
+            <param name="property">Property to create resolver for.</param>
+            <returns>The resolver object.</returns>
+        </member>
+        <member name="T:Microsoft.Practices.Unity.ObjectBuilder.LiteralValueDependencyResolverPolicy">
+            <summary>
+            A <see cref="T:Microsoft.Practices.ObjectBuilder2.IDependencyResolverPolicy"/> implementation that returns
+            the value set in the constructor.
+            </summary>
+        </member>
+        <member name="M:Microsoft.Practices.Unity.ObjectBuilder.LiteralValueDependencyResolverPolicy.#ctor(System.Object)">
+            <summary>
+            Create a new instance of <see cref="T:Microsoft.Practices.Unity.ObjectBuilder.LiteralValueDependencyResolverPolicy"/>
+            which will return the given value when resolved.
+            </summary>
+            <param name="dependencyValue">The value to return.</param>
+        </member>
+        <member name="M:Microsoft.Practices.Unity.ObjectBuilder.LiteralValueDependencyResolverPolicy.Resolve(Microsoft.Practices.ObjectBuilder2.IBuilderContext)">
+            <summary>
+            Get the value for a dependency.
+            </summary>
+            <param name="context">Current build context.</param>
+            <returns>The value for the dependency.</returns>
+        </member>
+        <member name="T:Microsoft.Practices.Unity.ObjectBuilder.NamedTypeDependencyResolverPolicy">
+            <summary>
+            An implementation of <see cref="T:Microsoft.Practices.ObjectBuilder2.IDependencyResolverPolicy"/> that stores a
+            type and name, and at resolution time puts them together into a
+            <see cref="T:Microsoft.Practices.ObjectBuilder2.NamedTypeBuildKey"/>.
+            </summary>
+        </member>
+        <member name="M:Microsoft.Practices.Unity.ObjectBuilder.NamedTypeDependencyResolverPolicy.#ctor(System.Type,System.String)">
+            <summary>
+            Create an instance of <see cref="T:Microsoft.Practices.Unity.ObjectBuilder.NamedTypeDependencyResolverPolicy"/>
+            with the given type and name.
+            </summary>
+            <param name="type">The type.</param>
+            <param name="name">The name (may be null).</param>
+        </member>
+        <member name="M:Microsoft.Practices.Unity.ObjectBuilder.NamedTypeDependencyResolverPolicy.Resolve(Microsoft.Practices.ObjectBuilder2.IBuilderContext)">
+            <summary>
+            Resolve the value for a dependency.
+            </summary>
+            <param name="context">Current build context.</param>
+            <returns>The value for the dependency.</returns>
+        </member>
+        <member name="P:Microsoft.Practices.Unity.ObjectBuilder.NamedTypeDependencyResolverPolicy.Type">
+            <summary>
+            The type that this resolver resolves.
+            </summary>
+        </member>
+        <member name="P:Microsoft.Practices.Unity.ObjectBuilder.NamedTypeDependencyResolverPolicy.Name">
+            <summary>
+            The name that this resolver resolves.
+            </summary>
+        </member>
+        <member name="T:Microsoft.Practices.Unity.ResolvedArrayWithElementsResolverPolicy">
+            <summary>
+            An implementation of <see cref="T:Microsoft.Practices.ObjectBuilder2.IDependencyResolverPolicy"/> that resolves to
+            to an array populated with the values that result from resolving other instances
+            of <see cref="T:Microsoft.Practices.ObjectBuilder2.IDependencyResolverPolicy"/>.
+            </summary>
+        </member>
+        <member name="M:Microsoft.Practices.Unity.ResolvedArrayWithElementsResolverPolicy.#ctor(System.Type,Microsoft.Practices.ObjectBuilder2.IDependencyResolverPolicy[])">
+            <summary>
+            Create an instance of <see cref="T:Microsoft.Practices.Unity.ResolvedArrayWithElementsResolverPolicy"/>
+            with the given type and a collection of <see cref="T:Microsoft.Practices.ObjectBuilder2.IDependencyResolverPolicy"/>
+            instances to use when populating the result.
+            </summary>
+            <param name="elementType">The type.</param>
+            <param name="elementPolicies">The resolver policies to use when populating an array.</param>
+        </member>
+        <member name="M:Microsoft.Practices.Unity.ResolvedArrayWithElementsResolverPolicy.Resolve(Microsoft.Practices.ObjectBuilder2.IBuilderContext)">
+            <summary>
+            Resolve the value for a dependency.
+            </summary>
+            <param name="context">Current build context.</param>
+            <returns>An array populated with the results of resolving the resolver policies.</returns>
+        </member>
+        <member name="T:Microsoft.Practices.Unity.ObjectBuilder.SpecifiedConstructorSelectorPolicy">
+            <summary>
+            An implementation of <see cref="T:Microsoft.Practices.ObjectBuilder2.IConstructorSelectorPolicy"/> that selects
+            the given constructor and creates the appropriate resolvers to call it with
+            the specified parameters.
+            </summary>
+        </member>
+        <member name="M:Microsoft.Practices.Unity.ObjectBuilder.SpecifiedConstructorSelectorPolicy.#ctor(System.Reflection.ConstructorInfo,Microsoft.Practices.Unity.InjectionParameterValue[])">
+            <summary>
+            Create an instance of <see cref="T:Microsoft.Practices.Unity.ObjectBuilder.SpecifiedConstructorSelectorPolicy"/> that
+            will return the given constructor, being passed the given injection values
+            as parameters.
+            </summary>
+            <param name="ctor">The constructor to call.</param>
+            <param name="parameterValues">Set of <see cref="T:Microsoft.Practices.Unity.InjectionParameterValue"/> objects
+            that describes how to obtain the values for the constructor parameters.</param>
+        </member>
+        <member name="M:Microsoft.Practices.Unity.ObjectBuilder.SpecifiedConstructorSelectorPolicy.SelectConstructor(Microsoft.Practices.ObjectBuilder2.IBuilderContext,Microsoft.Practices.ObjectBuilder2.IPolicyList)">
+            <summary>
+            Choose the constructor to call for the given type.
+            </summary>
+            <param name="context">Current build context</param>
+            <param name="resolverPolicyDestination">The <see cref="T:Microsoft.Practices.ObjectBuilder2.IPolicyList"/> to add any
+            generated resolver objects into.</param>
+            <returns>The chosen constructor.</returns>
+        </member>
+        <member name="T:Microsoft.Practices.Unity.ObjectBuilder.SpecifiedMemberSelectorHelper">
+            <summary>
+            Helper class for implementing selector policies that need to
+            set up dependency resolver policies.
+            </summary>
+        </member>
+        <member name="M:Microsoft.Practices.Unity.ObjectBuilder.SpecifiedMemberSelectorHelper.AddParameterResolvers(System.Type,Microsoft.Practices.ObjectBuilder2.IPolicyList,System.Collections.Generic.IEnumerable{Microsoft.Practices.Unity.InjectionParameterValue},Microsoft.Practices.ObjectBuilder2.SelectedMemberWithParameters)">
+            <summary>
+            Add dependency resolvers to the parameter set.
+            </summary>
+            <param name="typeToBuild">Type that's currently being built (used to resolve open generics).</param>
+            <param name="policies">PolicyList to add the resolvers to.</param>
+            <param name="parameterValues">Objects supplying the dependency resolvers.</param>
+            <param name="result">Result object to store the keys in.</param>
+        </member>
+        <member name="T:Microsoft.Practices.Unity.ObjectBuilder.SpecifiedMethodsSelectorPolicy">
+            <summary>
+            A <see cref="T:Microsoft.Practices.ObjectBuilder2.IMethodSelectorPolicy"/> implementation that calls the specific
+            methods with the given parameters.
+            </summary>
+        </member>
+        <member name="M:Microsoft.Practices.Unity.ObjectBuilder.SpecifiedMethodsSelectorPolicy.AddMethodAndParameters(System.Reflection.MethodInfo,System.Collections.Generic.IEnumerable{Microsoft.Practices.Unity.InjectionParameterValue})">
+            <summary>
+            Add the given method and parameter collection to the list of methods
+            that will be returned when the selector's <see cref="M:Microsoft.Practices.ObjectBuilder2.IMethodSelectorPolicy.SelectMethods(Microsoft.Practices.ObjectBuilder2.IBuilderContext,Microsoft.Practices.ObjectBuilder2.IPolicyList)"/>
+            method is called.
+            </summary>
+            <param name="method">Method to call.</param>
+            <param name="parameters">sequence of <see cref="T:Microsoft.Practices.Unity.InjectionParameterValue"/> objects
+            that describe how to create the method parameter values.</param>
+        </member>
+        <member name="M:Microsoft.Practices.Unity.ObjectBuilder.SpecifiedMethodsSelectorPolicy.SelectMethods(Microsoft.Practices.ObjectBuilder2.IBuilderContext,Microsoft.Practices.ObjectBuilder2.IPolicyList)">
+            <summary>
+            Return the sequence of methods to call while building the target object.
+            </summary>
+            <param name="context">Current build context.</param>
+            <param name="resolverPolicyDestination">The <see cref="T:Microsoft.Practices.ObjectBuilder2.IPolicyList"/> to add any
+            generated resolver objects into.</param>
+            <returns>Sequence of methods to call.</returns>
+        </member>
+        <member name="T:Microsoft.Practices.Unity.ObjectBuilder.SpecifiedPropertiesSelectorPolicy">
+            <summary>
+            An implementation of <see cref="T:Microsoft.Practices.ObjectBuilder2.IPropertySelectorPolicy"/> which returns
+            the set of specific properties that the selector was configured with.
+            </summary>
+        </member>
+        <member name="M:Microsoft.Practices.Unity.ObjectBuilder.SpecifiedPropertiesSelectorPolicy.AddPropertyAndValue(System.Reflection.PropertyInfo,Microsoft.Practices.Unity.InjectionParameterValue)">
+            <summary>
+            Add a property that will be par of the set returned when the 
+            <see cref="M:Microsoft.Practices.Unity.ObjectBuilder.SpecifiedPropertiesSelectorPolicy.SelectProperties(Microsoft.Practices.ObjectBuilder2.IBuilderContext,Microsoft.Practices.ObjectBuilder2.IPolicyList)"/> is called.
+            </summary>
+            <param name="property">The property to set.</param>
+            <param name="value"><see cref="T:Microsoft.Practices.Unity.InjectionParameterValue"/> object describing
+            how to create the value to inject.</param>
+        </member>
+        <member name="M:Microsoft.Practices.Unity.ObjectBuilder.SpecifiedPropertiesSelectorPolicy.SelectProperties(Microsoft.Practices.ObjectBuilder2.IBuilderContext,Microsoft.Practices.ObjectBuilder2.IPolicyList)">
+            <summary>
+            Returns sequence of properties on the given type that
+            should be set as part of building that object.
+            </summary>
+            <param name="context">Current build context.</param>
+            <param name="resolverPolicyDestination">The <see cref="T:Microsoft.Practices.ObjectBuilder2.IPolicyList"/> to add any
+            generated resolver objects into.</param>
+            <returns>Sequence of <see cref="T:System.Reflection.PropertyInfo"/> objects
+            that contain the properties to set.</returns>
+        </member>
+        <member name="T:Microsoft.Practices.Unity.ResolutionFailedException">
+            <summary>
+            The exception thrown by the Unity container when
+            an attempt to resolve a dependency fails.
+            </summary>
+        </member>
+        <member name="M:Microsoft.Practices.Unity.ResolutionFailedException.#ctor(System.Type,System.String,System.Exception,Microsoft.Practices.ObjectBuilder2.IBuilderContext)">
+            <summary>
+            Create a new <see cref="T:Microsoft.Practices.Unity.ResolutionFailedException"/> that records
+            the exception for the given type and name.
+            </summary>
+            <param name="typeRequested">Type requested from the container.</param>
+            <param name="nameRequested">Name requested from the container.</param>
+            <param name="innerException">The actual exception that caused the failure of the build.</param>
+            <param name="context">The build context representing the failed operation.</param>
+        </member>
+        <member name="P:Microsoft.Practices.Unity.ResolutionFailedException.TypeRequested">
+            <summary>
+            The type that was being requested from the container at the time of failure.
+            </summary>
+        </member>
+        <member name="P:Microsoft.Practices.Unity.ResolutionFailedException.NameRequested">
+            <summary>
+            The name that was being requested from the container at the time of failure.
+            </summary>
+        </member>
+        <member name="T:Microsoft.Practices.Unity.ResolvedParameter">
+            <summary>
+            A class that stores a name and type, and generates a 
+            resolver object that resolves the parameter via the
+            container.
+            </summary>
+        </member>
+        <member name="M:Microsoft.Practices.Unity.ResolvedParameter.#ctor(System.Type)">
+            <summary>
+            Construct a new <see cref="T:Microsoft.Practices.Unity.ResolvedParameter"/> that
+            resolves to the given type.
+            </summary>
+            <param name="parameterType">Type of this parameter.</param>
+        </member>
+        <member name="M:Microsoft.Practices.Unity.ResolvedParameter.#ctor(System.Type,System.String)">
+            <summary>
+            Construct a new <see cref="T:Microsoft.Practices.Unity.ResolvedParameter"/> that
+            resolves the given type and name.
+            </summary>
+            <param name="parameterType">Type of this parameter.</param>
+            <param name="name">Name to use when resolving parameter.</param>
+        </member>
+        <member name="M:Microsoft.Practices.Unity.ResolvedParameter.GetResolverPolicy(System.Type)">
+            <summary>
+            Return a <see cref="T:Microsoft.Practices.ObjectBuilder2.IDependencyResolverPolicy"/> instance that will
+            return this types value for the parameter.
+            </summary>
+            <param name="typeToBuild">Type that contains the member that needs this parameter. Used
+            to resolve open generic parameters.</param>
+            <returns>The <see cref="T:Microsoft.Practices.ObjectBuilder2.IDependencyResolverPolicy"/>.</returns>
+        </member>
+        <member name="T:Microsoft.Practices.Unity.ResolvedParameter`1">
+            <summary>
+            A generic version of <see cref="T:Microsoft.Practices.Unity.ResolvedParameter"/> for convenience
+            when creating them by hand.
+            </summary>
+            <typeparam name="TParameter">Type of the parameter</typeparam>
+        </member>
+        <member name="M:Microsoft.Practices.Unity.ResolvedParameter`1.#ctor">
+            <summary>
+            Create a new <see cref="T:Microsoft.Practices.Unity.ResolvedParameter`1"/> for the given
+            generic type and the default name.
+            </summary>
+        </member>
+        <member name="M:Microsoft.Practices.Unity.ResolvedParameter`1.#ctor(System.String)">
+            <summary>
+            Create a new <see cref="T:Microsoft.Practices.Unity.ResolvedParameter`1"/> for the given
+            generic type and name.
+            </summary>
+            <param name="name">Name to use to resolve this parameter.</param>
+        </member>
+        <member name="T:Microsoft.Practices.Unity.UnityServiceLocator">
+            <summary>
+            An implementation of <see cref="T:Microsoft.Practices.ServiceLocation.IServiceLocator"/> that wraps a Unity container.
+            </summary>
+        </member>
+        <member name="M:Microsoft.Practices.Unity.UnityServiceLocator.#ctor(Microsoft.Practices.Unity.IUnityContainer)">
+            <summary>
+            Initializes a new instance of the <see cref="T:Microsoft.Practices.Unity.UnityServiceLocator"/> class for a container.
+            </summary>
+            <param name="container">The <see cref="T:Microsoft.Practices.Unity.IUnityContainer"/> to wrap with the <see cref="T:Microsoft.Practices.ServiceLocation.IServiceLocator"/>
+            interface implementation.</param>
+        </member>
+        <member name="M:Microsoft.Practices.Unity.UnityServiceLocator.Dispose">
+            <summary>
+            Performs application-defined tasks associated with freeing, releasing, or resetting unmanaged resources.
+            </summary>
+            <filterpriority>2</filterpriority>
+        </member>
+        <member name="M:Microsoft.Practices.Unity.UnityServiceLocator.DoGetInstance(System.Type,System.String)">
+            <summary>
+            When implemented by inheriting classes, this method will do the actual work of resolving
+                        the requested service instance.
+            </summary>
+            <param name="serviceType">Type of instance requested.</param><param name="key">Name of registered service you want. May be null.</param>
+            <returns>
+            The requested service instance.
+            </returns>
+        </member>
+        <member name="M:Microsoft.Practices.Unity.UnityServiceLocator.DoGetAllInstances(System.Type)">
+            <summary>
+            When implemented by inheriting classes, this method will do the actual work of
+                        resolving all the requested service instances.
+            </summary>
+            <param name="serviceType">Type of service requested.</param>
+            <returns>
+            Sequence of service instance objects.
+            </returns>
+        </member>
+        <member name="T:Microsoft.Practices.Unity.Utility.Guard">
+            <summary>
+            A static helper class that includes various parameter checking routines.
+            </summary>
+        </member>
+        <member name="M:Microsoft.Practices.Unity.Utility.Guard.ArgumentNotNull(System.Object,System.String)">
+            <summary>
+            Throws <see cref="T:System.ArgumentNullException"/> if the given argument is null.
+            </summary>
+            <exception cref="T:System.ArgumentNullException"> if tested value if null.</exception>
+            <param name="argumentValue">Argument value to test.</param>
+            <param name="argumentName">Name of the argument being tested.</param>
+        </member>
+        <member name="M:Microsoft.Practices.Unity.Utility.Guard.ArgumentNotNullOrEmpty(System.String,System.String)">
+            <summary>
+            Throws an exception if the tested string argument is null or the empty string.
+            </summary>
+            <exception cref="T:System.ArgumentNullException">Thrown if string value is null.</exception>
+            <exception cref="T:System.ArgumentException">Thrown if the string is empty</exception>
+            <param name="argumentValue">Argument value to check.</param>
+            <param name="argumentName">Name of argument being checked.</param>
+        </member>
+        <member name="M:Microsoft.Practices.Unity.Utility.Guard.TypeIsAssignable(System.Type,System.Type,System.String)">
+            <summary>
+            Verifies that an argument type is assignable from the provided type (meaning
+            interfaces are implemented, or classes exist in the base class hierarchy).
+            </summary>
+            <param name="assignmentTargetType">The argument type that will be assigned to.</param>
+            <param name="assignmentValueType">The type of the value being assigned.</param>
+            <param name="argumentName">Argument name.</param>
+        </member>
+        <member name="M:Microsoft.Practices.Unity.Utility.Guard.InstanceIsAssignable(System.Type,System.Object,System.String)">
+            <summary>
+            Verifies that an argument instance is assignable from the provided type (meaning
+            interfaces are implemented, or classes exist in the base class hierarchy, or instance can be 
+            assigned through a runtime wrapper, as is the case for COM Objects).
+            </summary>
+            <param name="assignmentTargetType">The argument type that will be assigned to.</param>
+            <param name="assignmentInstance">The instance that will be assigned.</param>
+            <param name="argumentName">Argument name.</param>
+        </member>
+        <member name="T:Microsoft.Practices.Unity.NamedTypesRegistry">
+            <summary>
+            A helper class to manage the names that get registered in the container  
+            </summary>
+        </member>
+        <member name="T:Microsoft.Practices.Unity.ObjectBuilder.UnityBuildStage">
+            <summary>
+            The build stages we use in the Unity container
+            strategy pipeline.
+            </summary>
+        </member>
+        <member name="F:Microsoft.Practices.Unity.ObjectBuilder.UnityBuildStage.Setup">
+            <summary>
+            First stage. By default, nothing happens here.
+            </summary>
+        </member>
+        <member name="F:Microsoft.Practices.Unity.ObjectBuilder.UnityBuildStage.TypeMapping">
+            <summary>
+            Second stage. Type mapping occurs here.
+            </summary>
+        </member>
+        <member name="F:Microsoft.Practices.Unity.ObjectBuilder.UnityBuildStage.Lifetime">
+            <summary>
+            Third stage. lifetime managers are checked here,
+            and if they're available the rest of the pipeline is skipped.
+            </summary>
+        </member>
+        <member name="F:Microsoft.Practices.Unity.ObjectBuilder.UnityBuildStage.PreCreation">
+            <summary>
+            Fourth stage. Reflection over constructors, properties, etc. is
+            performed here.
+            </summary>
+        </member>
+        <member name="F:Microsoft.Practices.Unity.ObjectBuilder.UnityBuildStage.Creation">
+            <summary>
+            Fifth stage. Instance creation happens here.
+            </summary>
+        </member>
+        <member name="F:Microsoft.Practices.Unity.ObjectBuilder.UnityBuildStage.Initialization">
+            <summary>
+            Sixth stage. Property sets and method injection happens here.
+            </summary>
+        </member>
+        <member name="F:Microsoft.Practices.Unity.ObjectBuilder.UnityBuildStage.PostInitialization">
+            <summary>
+            Seventh and final stage. By default, nothing happens here.
+            </summary>
+        </member>
+        <member name="T:Microsoft.Practices.Unity.Properties.Resources">
+            <summary>
+              A strongly-typed resource class, for looking up localized strings, etc.
+            </summary>
+        </member>
+        <member name="P:Microsoft.Practices.Unity.Properties.Resources.ResourceManager">
+            <summary>
+              Returns the cached ResourceManager instance used by this class.
+            </summary>
+        </member>
+        <member name="P:Microsoft.Practices.Unity.Properties.Resources.Culture">
+            <summary>
+              Overrides the current thread's CurrentUICulture property for all
+              resource lookups using this strongly typed resource class.
+            </summary>
+        </member>
+        <member name="P:Microsoft.Practices.Unity.Properties.Resources.AmbiguousInjectionConstructor">
+            <summary>
+              Looks up a localized string similar to The type {0} has multiple constructors of length {1}. Unable to disambiguate..
+            </summary>
+        </member>
+        <member name="P:Microsoft.Practices.Unity.Properties.Resources.ArgumentMustNotBeEmpty">
+            <summary>
+              Looks up a localized string similar to The provided string argument must not be empty..
+            </summary>
+        </member>
+        <member name="P:Microsoft.Practices.Unity.Properties.Resources.BuildFailedException">
+            <summary>
+              Looks up a localized string similar to The current build operation (build key {2}) failed: {3} (Strategy type {0}, index {1}).
+            </summary>
+        </member>
+        <member name="P:Microsoft.Practices.Unity.Properties.Resources.CannotConstructAbstractClass">
+            <summary>
+              Looks up a localized string similar to The current type, {0}, is an abstract class and cannot be constructed. Are you missing a type mapping?.
+            </summary>
+        </member>
+        <member name="P:Microsoft.Practices.Unity.Properties.Resources.CannotConstructDelegate">
+            <summary>
+              Looks up a localized string similar to The current type, {0}, is delegate and cannot be constructed. Unity only supports resolving Func&lt;T&gt; and Func&lt;IEnumerable&lt;T&gt;&gt; by default..
+            </summary>
+        </member>
+        <member name="P:Microsoft.Practices.Unity.Properties.Resources.CannotConstructInterface">
+            <summary>
+              Looks up a localized string similar to The current type, {0}, is an interface and cannot be constructed. Are you missing a type mapping?.
+            </summary>
+        </member>
+        <member name="P:Microsoft.Practices.Unity.Properties.Resources.CannotExtractTypeFromBuildKey">
+            <summary>
+              Looks up a localized string similar to Cannot extract type from build key {0}..
+            </summary>
+        </member>
+        <member name="P:Microsoft.Practices.Unity.Properties.Resources.CannotInjectGenericMethod">
+            <summary>
+              Looks up a localized string similar to The method {0}.{1}({2}) is an open generic method. Open generic methods cannot be injected..
+            </summary>
+        </member>
+        <member name="P:Microsoft.Practices.Unity.Properties.Resources.CannotInjectIndexer">
+            <summary>
+              Looks up a localized string similar to The property {0} on type {1} is an indexer. Indexed properties cannot be injected..
+            </summary>
+        </member>
+        <member name="P:Microsoft.Practices.Unity.Properties.Resources.CannotInjectMethodWithOutParam">
+            <summary>
+              Looks up a localized string similar to The method {1} on type {0} has an out parameter. Injection cannot be performed..
+            </summary>
+        </member>
+        <member name="P:Microsoft.Practices.Unity.Properties.Resources.CannotInjectMethodWithOutParams">
+            <summary>
+              Looks up a localized string similar to The method {0}.{1}({2}) has at least one out parameter. Methods with out parameters cannot be injected..
+            </summary>
+        </member>
+        <member name="P:Microsoft.Practices.Unity.Properties.Resources.CannotInjectMethodWithRefParams">
+            <summary>
+              Looks up a localized string similar to The method {0}.{1}({2}) has at least one ref parameter.Methods with ref parameters cannot be injected..
+            </summary>
+        </member>
+        <member name="P:Microsoft.Practices.Unity.Properties.Resources.CannotInjectOpenGenericMethod">
+            <summary>
+              Looks up a localized string similar to The method {1} on type {0} is marked for injection, but it is an open generic method. Injection cannot be performed..
+            </summary>
+        </member>
+        <member name="P:Microsoft.Practices.Unity.Properties.Resources.CannotInjectStaticMethod">
+            <summary>
+              Looks up a localized string similar to The method {0}.{1}({2}) is static. Static methods cannot be injected..
+            </summary>
+        </member>
+        <member name="P:Microsoft.Practices.Unity.Properties.Resources.CannotResolveOpenGenericType">
+            <summary>
+              Looks up a localized string similar to The type {0} is an open generic type. An open generic type cannot be resolved..
+            </summary>
+        </member>
+        <member name="P:Microsoft.Practices.Unity.Properties.Resources.ConstructorArgumentResolveOperation">
+            <summary>
+              Looks up a localized string similar to Resolving parameter &quot;{0}&quot; of constructor {1}.
+            </summary>
+        </member>
+        <member name="P:Microsoft.Practices.Unity.Properties.Resources.ConstructorParameterResolutionFailed">
+            <summary>
+              Looks up a localized string similar to The parameter {0} could not be resolved when attempting to call constructor {1}..
+            </summary>
+        </member>
+        <member name="P:Microsoft.Practices.Unity.Properties.Resources.ExceptionNullParameterValue">
+            <summary>
+              Looks up a localized string similar to Parameter type inference does not work for null values. Indicate the parameter type explicitly using a properly configured instance of the InjectionParameter or InjectionParameter&lt;T&gt; classes..
+            </summary>
+        </member>
+        <member name="P:Microsoft.Practices.Unity.Properties.Resources.InvokingConstructorOperation">
+            <summary>
+              Looks up a localized string similar to Calling constructor {0}.
+            </summary>
+        </member>
+        <member name="P:Microsoft.Practices.Unity.Properties.Resources.InvokingMethodOperation">
+            <summary>
+              Looks up a localized string similar to Calling method {0}.{1}.
+            </summary>
+        </member>
+        <member name="P:Microsoft.Practices.Unity.Properties.Resources.KeyAlreadyPresent">
+            <summary>
+              Looks up a localized string similar to An item with the given key is already present in the dictionary..
+            </summary>
+        </member>
+        <member name="P:Microsoft.Practices.Unity.Properties.Resources.LifetimeManagerInUse">
+            <summary>
+              Looks up a localized string similar to The lifetime manager is already registered. Lifetime managers cannot be reused, please create a new one..
+            </summary>
+        </member>
+        <member name="P:Microsoft.Practices.Unity.Properties.Resources.MarkerBuildPlanInvoked">
+            <summary>
+              Looks up a localized string similar to The override marker build plan policy has been invoked. This should never happen, looks like a bug in the container..
+            </summary>
+        </member>
+        <member name="P:Microsoft.Practices.Unity.Properties.Resources.MethodArgumentResolveOperation">
+            <summary>
+              Looks up a localized string similar to Resolving parameter &quot;{0}&quot; of method {1}.{2}.
+            </summary>
+        </member>
+        <member name="P:Microsoft.Practices.Unity.Properties.Resources.MethodParameterResolutionFailed">
+            <summary>
+              Looks up a localized string similar to The value for parameter &quot;{1}&quot; of method {0} could not be resolved. .
+            </summary>
+        </member>
+        <member name="P:Microsoft.Practices.Unity.Properties.Resources.MissingDependency">
+            <summary>
+              Looks up a localized string similar to Could not resolve dependency for build key {0}..
+            </summary>
+        </member>
+        <member name="P:Microsoft.Practices.Unity.Properties.Resources.MultipleInjectionConstructors">
+            <summary>
+              Looks up a localized string similar to The type {0} has multiple constructors marked with the InjectionConstructor attribute. Unable to disambiguate..
+            </summary>
+        </member>
+        <member name="P:Microsoft.Practices.Unity.Properties.Resources.MustHaveOpenGenericType">
+            <summary>
+              Looks up a localized string similar to The supplied type {0} must be an open generic type..
+            </summary>
+        </member>
+        <member name="P:Microsoft.Practices.Unity.Properties.Resources.MustHaveSameNumberOfGenericArguments">
+            <summary>
+              Looks up a localized string similar to The supplied type {0} does not have the same number of generic arguments as the target type {1}..
+            </summary>
+        </member>
+        <member name="P:Microsoft.Practices.Unity.Properties.Resources.NoConstructorFound">
+            <summary>
+              Looks up a localized string similar to The type {0} does not have an accessible constructor..
+            </summary>
+        </member>
+        <member name="P:Microsoft.Practices.Unity.Properties.Resources.NoMatchingGenericArgument">
+            <summary>
+              Looks up a localized string similar to The type {0} does not have a generic argument named &quot;{1}&quot;.
+            </summary>
+        </member>
+        <member name="P:Microsoft.Practices.Unity.Properties.Resources.NoOperationExceptionReason">
+            <summary>
+              Looks up a localized string similar to while resolving.
+            </summary>
+        </member>
+        <member name="P:Microsoft.Practices.Unity.Properties.Resources.NoSuchConstructor">
+            <summary>
+              Looks up a localized string similar to The type {0} does not have a constructor that takes the parameters ({1})..
+            </summary>
+        </member>
+        <member name="P:Microsoft.Practices.Unity.Properties.Resources.NoSuchMethod">
+            <summary>
+              Looks up a localized string similar to The type {0} does not have a public method named {1} that takes the parameters ({2})..
+            </summary>
+        </member>
+        <member name="P:Microsoft.Practices.Unity.Properties.Resources.NoSuchProperty">
+            <summary>
+              Looks up a localized string similar to The type {0} does not contain an instance property named {1}..
+            </summary>
+        </member>
+        <member name="P:Microsoft.Practices.Unity.Properties.Resources.NotAGenericType">
+            <summary>
+              Looks up a localized string similar to The type {0} is not a generic type, and you are attempting to inject a generic parameter named &quot;{1}&quot;..
+            </summary>
+        </member>
+        <member name="P:Microsoft.Practices.Unity.Properties.Resources.NotAnArrayTypeWithRankOne">
+            <summary>
+              Looks up a localized string similar to The type {0} is not an array type with rank 1, and you are attempting to use a [DependencyArray] attribute on a parameter or property with this type..
+            </summary>
+        </member>
+        <member name="P:Microsoft.Practices.Unity.Properties.Resources.OptionalDependenciesMustBeReferenceTypes">
+            <summary>
+              Looks up a localized string similar to Optional dependencies must be reference types. The type {0} is a value type..
+            </summary>
+        </member>
+        <member name="P:Microsoft.Practices.Unity.Properties.Resources.PropertyNotSettable">
+            <summary>
+              Looks up a localized string similar to The property {0} on type {1} is not settable..
+            </summary>
+        </member>
+        <member name="P:Microsoft.Practices.Unity.Properties.Resources.PropertyTypeMismatch">
+            <summary>
+              Looks up a localized string similar to The property {0} on type {1} is of type {2}, and cannot be injected with a value of type {3}..
+            </summary>
+        </member>
+        <member name="P:Microsoft.Practices.Unity.Properties.Resources.PropertyValueResolutionFailed">
+            <summary>
+              Looks up a localized string similar to The value for the property &quot;{0}&quot; could not be resolved..
+            </summary>
+        </member>
+        <member name="P:Microsoft.Practices.Unity.Properties.Resources.ProvidedStringArgMustNotBeEmpty">
+            <summary>
+              Looks up a localized string similar to The provided string argument must not be empty..
+            </summary>
+        </member>
+        <member name="P:Microsoft.Practices.Unity.Properties.Resources.ResolutionFailed">
+             <summary>
+               Looks up a localized string similar to Resolution of the dependency failed, type = &quot;{0}&quot;, name = &quot;{1}&quot;.
+            Exception occurred while: {2}.
+            Exception is: {3} - {4}
+            -----------------------------------------------
+            At the time of the exception, the container was:
+            .
+             </summary>
+        </member>
+        <member name="P:Microsoft.Practices.Unity.Properties.Resources.ResolutionTraceDetail">
+            <summary>
+              Looks up a localized string similar to Resolving {0},{1}.
+            </summary>
+        </member>
+        <member name="P:Microsoft.Practices.Unity.Properties.Resources.ResolutionWithMappingTraceDetail">
+            <summary>
+              Looks up a localized string similar to Resolving {0},{1} (mapped from {2}, {3}).
+            </summary>
+        </member>
+        <member name="P:Microsoft.Practices.Unity.Properties.Resources.ResolvingPropertyValueOperation">
+            <summary>
+              Looks up a localized string similar to Resolving value for property {0}.{1}.
+            </summary>
+        </member>
+        <member name="P:Microsoft.Practices.Unity.Properties.Resources.SelectedConstructorHasRefParameters">
+            <summary>
+              Looks up a localized string similar to The constructor {1} selected for type {0} has ref or out parameters. Such parameters are not supported for constructor injection..
+            </summary>
+        </member>
+        <member name="P:Microsoft.Practices.Unity.Properties.Resources.SettingPropertyOperation">
+            <summary>
+              Looks up a localized string similar to Setting value for property {0}.{1}.
+            </summary>
+        </member>
+        <member name="P:Microsoft.Practices.Unity.Properties.Resources.TypeIsNotConstructable">
+            <summary>
+              Looks up a localized string similar to The type {0} cannot be constructed. You must configure the container to supply this value..
+            </summary>
+        </member>
+        <member name="P:Microsoft.Practices.Unity.Properties.Resources.TypesAreNotAssignable">
+            <summary>
+              Looks up a localized string similar to The type {1} cannot be assigned to variables of type {0}..
+            </summary>
+        </member>
+        <member name="P:Microsoft.Practices.Unity.Properties.Resources.UnknownType">
+            <summary>
+              Looks up a localized string similar to &lt;unknown&gt;.
+            </summary>
+        </member>
+        <member name="T:Microsoft.Practices.Unity.UnityContainer">
+            <summary>
+            A simple, extensible dependency injection container.
+            </summary>
+        </member>
+        <member name="M:Microsoft.Practices.Unity.UnityContainer.#ctor">
+            <summary>
+            Create a default <see cref="T:Microsoft.Practices.Unity.UnityContainer"/>.
+            </summary>
+        </member>
+        <member name="M:Microsoft.Practices.Unity.UnityContainer.#ctor(Microsoft.Practices.Unity.UnityContainer)">
+            <summary>
+            Create a <see cref="T:Microsoft.Practices.Unity.UnityContainer"/> with the given parent container.
+            </summary>
+            <param name="parent">The parent <see cref="T:Microsoft.Practices.Unity.UnityContainer"/>. The current object
+            will apply its own settings first, and then check the parent for additional ones.</param>
+        </member>
+        <member name="M:Microsoft.Practices.Unity.UnityContainer.RegisterType(System.Type,System.Type,System.String,Microsoft.Practices.Unity.LifetimeManager,Microsoft.Practices.Unity.InjectionMember[])">
+            <summary>
+            RegisterType a type mapping with the container, where the created instances will use
+            the given <see cref="T:Microsoft.Practices.Unity.LifetimeManager"/>.
+            </summary>
+            <param name="from"><see cref="T:System.Type"/> that will be requested.</param>
+            <param name="to"><see cref="T:System.Type"/> that will actually be returned.</param>
+            <param name="name">Name to use for registration, null if a default registration.</param>
+            <param name="lifetimeManager">The <see cref="T:Microsoft.Practices.Unity.LifetimeManager"/> that controls the lifetime
+            of the returned instance.</param>
+            <param name="injectionMembers">Injection configuration objects.</param>
+            <returns>The <see cref="T:Microsoft.Practices.Unity.UnityContainer"/> object that this method was called on (this in C#, Me in Visual Basic).</returns>
+        </member>
+        <member name="M:Microsoft.Practices.Unity.UnityContainer.RegisterInstance(System.Type,System.String,System.Object,Microsoft.Practices.Unity.LifetimeManager)">
+            <summary>
+            RegisterType an instance with the container.
+            </summary>
+            <remarks>
+            <para>
+            Instance registration is much like setting a type as a singleton, except that instead
+            of the container creating the instance the first time it is requested, the user
+            creates the instance ahead of type and adds that instance to the container.
+            </para>
+            </remarks>
+            <param name="t">Type of instance to register (may be an implemented interface instead of the full type).</param>
+            <param name="instance">Object to returned.</param>
+            <param name="name">Name for registration.</param>
+            <param name="lifetime">
+            <para>If true, the container will take over the lifetime of the instance,
+            calling Dispose on it (if it's <see cref="T:System.IDisposable"/>) when the container is Disposed.</para>
+            <para>
+             If false, container will not maintain a strong reference to <paramref name="instance"/>. User is responsible
+            for disposing instance, and for keeping the instance from being garbage collected.</para></param>
+            <returns>The <see cref="T:Microsoft.Practices.Unity.UnityContainer"/> object that this method was called on (this in C#, Me in Visual Basic).</returns>
+        </member>
+        <member name="M:Microsoft.Practices.Unity.UnityContainer.Resolve(System.Type,System.String,Microsoft.Practices.Unity.ResolverOverride[])">
+            <summary>
+            Get an instance of the requested type with the given name from the container.
+            </summary>
+            <param name="t"><see cref="T:System.Type"/> of object to get from the container.</param>
+            <param name="name">Name of the object to retrieve.</param>
+            <param name="resolverOverrides">Any overrides for the resolve call.</param>
+            <returns>The retrieved object.</returns>
+        </member>
+        <member name="M:Microsoft.Practices.Unity.UnityContainer.ResolveAll(System.Type,Microsoft.Practices.Unity.ResolverOverride[])">
+            <summary>
+            Return instances of all registered types requested.
+            </summary>
+            <remarks>
+            <para>
+            This method is useful if you've registered multiple types with the same
+            <see cref="T:System.Type"/> but different names.
+            </para>
+            <para>
+            Be aware that this method does NOT return an instance for the default (unnamed) registration.
+            </para>
+            </remarks>
+            <param name="t">The type requested.</param>
+            <param name="resolverOverrides">Any overrides for the resolve calls.</param>
+            <returns>Set of objects of type <paramref name="t"/>.</returns>
+        </member>
+        <member name="M:Microsoft.Practices.Unity.UnityContainer.BuildUp(System.Type,System.Object,System.String,Microsoft.Practices.Unity.ResolverOverride[])">
+            <summary>
+            Run an existing object through the container and perform injection on it.
+            </summary>
+            <remarks>
+            <para>
+            This method is useful when you don't control the construction of an
+            instance (ASP.NET pages or objects created via XAML, for instance)
+            but you still want properties and other injection performed.
+            </para></remarks>
+            <param name="t"><see cref="T:System.Type"/> of object to perform injection on.</param>
+            <param name="existing">Instance to build up.</param>
+            <param name="name">name to use when looking up the typemappings and other configurations.</param>
+            <param name="resolverOverrides">Any overrides for the buildup.</param>
+            <returns>The resulting object. By default, this will be <paramref name="existing"/>, but
+            container extensions may add things like automatic proxy creation which would
+            cause this to return a different object (but still type compatible with <paramref name="t"/>).</returns>
+        </member>
+        <member name="M:Microsoft.Practices.Unity.UnityContainer.Teardown(System.Object)">
+            <summary>
+            Run an existing object through the container, and clean it up.
+            </summary>
+            <param name="o">The object to tear down.</param>
+        </member>
+        <member name="M:Microsoft.Practices.Unity.UnityContainer.AddExtension(Microsoft.Practices.Unity.UnityContainerExtension)">
+            <summary>
+            Add an extension object to the container.
+            </summary>
+            <param name="extension"><see cref="T:Microsoft.Practices.Unity.UnityContainerExtension"/> to add.</param>
+            <returns>The <see cref="T:Microsoft.Practices.Unity.UnityContainer"/> object that this method was called on (this in C#, Me in Visual Basic).</returns>
+        </member>
+        <member name="M:Microsoft.Practices.Unity.UnityContainer.Configure(System.Type)">
+            <summary>
+            Get access to a configuration interface exposed by an extension.
+            </summary>
+            <remarks>Extensions can expose configuration interfaces as well as adding
+            strategies and policies to the container. This method walks the list of
+            added extensions and returns the first one that implements the requested type.
+            </remarks>
+            <param name="configurationInterface"><see cref="T:System.Type"/> of configuration interface required.</param>
+            <returns>The requested extension's configuration interface, or null if not found.</returns>
+        </member>
+        <member name="M:Microsoft.Practices.Unity.UnityContainer.RemoveAllExtensions">
+            <summary>
+            Remove all installed extensions from this container.
+            </summary>
+            <remarks>
+            <para>
+            This method removes all extensions from the container, including the default ones
+            that implement the out-of-the-box behavior. After this method, if you want to use
+            the container again you will need to either read the default extensions or replace
+            them with your own.
+            </para>
+            <para>
+            The registered instances and singletons that have already been set up in this container
+            do not get removed.
+            </para>
+            </remarks>
+            <returns>The <see cref="T:Microsoft.Practices.Unity.UnityContainer"/> object that this method was called on (this in C#, Me in Visual Basic).</returns>
+        </member>
+        <member name="M:Microsoft.Practices.Unity.UnityContainer.CreateChildContainer">
+            <summary>
+            Create a child container.
+            </summary>
+            <remarks>
+            A child container shares the parent's configuration, but can be configured with different
+            settings or lifetime.</remarks>
+            <returns>The new child container.</returns>
+        </member>
+        <member name="M:Microsoft.Practices.Unity.UnityContainer.Dispose">
+            <summary>
+            Dispose this container instance.
+            </summary>
+            <remarks>
+            Disposing the container also disposes any child containers,
+            and disposes any instances whose lifetimes are managed
+            by the container.
+            </remarks>
+        </member>
+        <member name="M:Microsoft.Practices.Unity.UnityContainer.Dispose(System.Boolean)">
+            <summary>
+            Dispose this container instance.
+            </summary>
+            <remarks>
+            This class doesn't have a finalizer, so <paramref name="disposing"/> will always be true.</remarks>
+            <param name="disposing">True if being called from the IDisposable.Dispose
+            method, false if being called from a finalizer.</param>
+        </member>
+        <member name="M:Microsoft.Practices.Unity.UnityContainer.ClearExistingBuildPlan(System.Type,System.String)">
+            <summary>
+            Remove policies associated with building this type. This removes the
+            compiled build plan so that it can be rebuilt with the new settings
+            the next time this type is resolved.
+            </summary>
+            <param name="typeToInject">Type of object to clear the plan for.</param>
+            <param name="name">Name the object is being registered with.</param>
+        </member>
+        <member name="P:Microsoft.Practices.Unity.UnityContainer.Parent">
+            <summary>
+            The parent of this container.
+            </summary>
+            <value>The parent container, or null if this container doesn't have one.</value>
+        </member>
+        <member name="P:Microsoft.Practices.Unity.UnityContainer.Registrations">
+            <summary>
+            Get a sequence of <see cref="T:Microsoft.Practices.Unity.ContainerRegistration"/> that describe the current state
+            of the container.
+            </summary>
+        </member>
+        <member name="T:Microsoft.Practices.Unity.UnityContainer.ExtensionContextImpl">
+            <summary>
+            Implementation of the ExtensionContext that is actually used
+            by the UnityContainer implementation.
+            </summary>
+            <remarks>
+            This is a nested class so that it can access state in the
+            container that would otherwise be inaccessible.
+            </remarks>
+        </member>
+        <member name="E:Microsoft.Practices.Unity.UnityContainer.ExtensionContextImpl.RegisteringInstance">
+            <summary>
+            This event is raised when the <see cref="M:Microsoft.Practices.Unity.UnityContainer.RegisterInstance(System.Type,System.String,System.Object,Microsoft.Practices.Unity.LifetimeManager)"/> method,
+            or one of its overloads, is called.
+            </summary>
+        </member>
+        <member name="T:Microsoft.Practices.Unity.UnityDefaultBehaviorExtension">
+            <summary>
+            This extension supplies the default behavior of the UnityContainer API
+            by handling the context events and setting policies.
+            </summary>
+        </member>
+        <member name="M:Microsoft.Practices.Unity.UnityDefaultBehaviorExtension.Initialize">
+            <summary>
+            Install the default container behavior into the container.
+            </summary>
+        </member>
+        <member name="M:Microsoft.Practices.Unity.UnityDefaultBehaviorExtension.Remove">
+            <summary>
+            Remove the default behavior from the container.
+            </summary>
+        </member>
+        <member name="T:Microsoft.Practices.Unity.UnityDefaultStrategiesExtension">
+            <summary>
+            This extension installs the default strategies and policies into the container
+            to implement the standard behavior of the Unity container.
+            </summary>
+        </member>
+        <member name="M:Microsoft.Practices.Unity.UnityDefaultStrategiesExtension.Initialize">
+            <summary>
+            Add the default ObjectBuilder strategies &amp; policies to the container.
+            </summary>
+        </member>
+        <member name="T:Microsoft.Practices.Unity.Utility.MethodReflectionHelper">
+            <summary>
+            Helper class to wrap common reflection stuff dealing with
+            methods.
+            </summary>
+        </member>
+        <member name="M:Microsoft.Practices.Unity.Utility.MethodReflectionHelper.#ctor(System.Reflection.MethodBase)">
+            <summary>
+            Create a new <see cref="T:Microsoft.Practices.Unity.Utility.MethodReflectionHelper"/> instance that
+            lets us do more reflection stuff on that method.
+            </summary>
+            <param name="method">The method to reflect on.</param>
+        </member>
+        <member name="M:Microsoft.Practices.Unity.Utility.MethodReflectionHelper.GetClosedParameterTypes(System.Type[])">
+            <summary>
+            Given our set of generic type arguments, 
+            </summary>
+            <param name="genericTypeArguments">The generic type arguments.</param>
+            <returns>An array with closed parameter types. </returns>
+        </member>
+        <member name="P:Microsoft.Practices.Unity.Utility.MethodReflectionHelper.MethodHasOpenGenericParameters">
+            <summary>
+            Returns true if any of the parameters of this method
+            are open generics.
+            </summary>
+        </member>
+        <member name="P:Microsoft.Practices.Unity.Utility.MethodReflectionHelper.ParameterTypes">
+            <summary>
+            Return the <see cref="T:System.Type"/> of each parameter for this
+            method.
+            </summary>
+            <returns>Sequence of <see cref="T:System.Type"/> objects, one for
+            each parameter in order.</returns>
+        </member>
+        <member name="T:Microsoft.Practices.Unity.Utility.Pair`2">
+            <summary>
+            A helper class that encapsulates two different
+            data items together into a a single item.
+            </summary>
+        </member>
+        <member name="M:Microsoft.Practices.Unity.Utility.Pair`2.#ctor(`0,`1)">
+            <summary>
+            Create a new <see cref="T:Microsoft.Practices.Unity.Utility.Pair`2"/> containing
+            the two values give.
+            </summary>
+            <param name="first">First value</param>
+            <param name="second">Second value</param>
+        </member>
+        <member name="P:Microsoft.Practices.Unity.Utility.Pair`2.First">
+            <summary>
+            The first value of the pair.
+            </summary>
+        </member>
+        <member name="P:Microsoft.Practices.Unity.Utility.Pair`2.Second">
+            <summary>
+            The second value of the pair.
+            </summary>
+        </member>
+        <member name="T:Microsoft.Practices.Unity.Utility.Pair">
+            <summary>
+            Container for a Pair helper method.
+            </summary>
+        </member>
+        <member name="M:Microsoft.Practices.Unity.Utility.Pair.Make``2(``0,``1)">
+            <summary>
+            A helper factory method that lets users take advantage of type inference.
+            </summary>
+            <typeparam name="TFirstParameter">Type of first value.</typeparam>
+            <typeparam name="TSecondParameter">Type of second value.</typeparam>
+            <param name="first">First value.</param>
+            <param name="second">Second value.</param>
+            <returns>A new <see cref="T:Microsoft.Practices.Unity.Utility.Pair`2"/> instance.</returns>
+        </member>
+        <member name="T:Microsoft.Practices.Unity.Utility.ParameterMatcher">
+            <summary>
+            A utility class that handles the logic of matching parameter
+            lists, so we can find the right constructor and method overloads.
+            </summary>
+        </member>
+        <member name="M:Microsoft.Practices.Unity.Utility.ParameterMatcher.#ctor(System.Collections.Generic.IEnumerable{Microsoft.Practices.Unity.InjectionParameterValue})">
+            <summary>
+            Create a new <see cref="T:Microsoft.Practices.Unity.Utility.ParameterMatcher"/> that will attempt to
+            match the given parameter types.
+            </summary>
+            <param name="parametersToMatch">Target parameters to match against.</param>
+        </member>
+        <member name="M:Microsoft.Practices.Unity.Utility.ParameterMatcher.Matches(System.Collections.Generic.IEnumerable{System.Type})">
+            <summary>
+            Tests to see if the given set of types matches the ones
+            we're looking for.
+            </summary>
+            <param name="candidate">parameter list to look for.</param>
+            <returns>true if they match, false if they don't.</returns>
+        </member>
+        <member name="M:Microsoft.Practices.Unity.Utility.ParameterMatcher.Matches(System.Collections.Generic.IEnumerable{System.Reflection.ParameterInfo})">
+            <summary>
+            Tests to see if the given set of types matches the ones we're looking for.
+            </summary>
+            <param name="candidate">Candidate method signature to look for.</param>
+            <returns>True if they match, false if they don't.</returns>
+        </member>
+        <member name="T:Microsoft.Practices.Unity.Utility.ParameterReflectionHelper">
+            <summary>
+            Another reflection helper class that has extra methods
+            for dealing with ParameterInfo.
+            </summary>
+        </member>
+        <member name="T:Microsoft.Practices.Unity.Utility.ReflectionHelper">
+            <summary>
+            A small helper class to encapsulate details of the
+            reflection API, particularly around generics.
+            </summary>
+        </member>
+        <member name="M:Microsoft.Practices.Unity.Utility.ReflectionHelper.#ctor(System.Type)">
+            <summary>
+            Create a new <see cref="T:Microsoft.Practices.Unity.Utility.ReflectionHelper"/> instance that
+            lets you look at information about the given type.
+            </summary>
+            <param name="typeToReflect">Type to do reflection on.</param>
+        </member>
+        <member name="M:Microsoft.Practices.Unity.Utility.ReflectionHelper.MethodHasOpenGenericParameters(System.Reflection.MethodBase)">
+            <summary>
+            Test the given <see cref="T:System.Reflection.MethodBase"/> object, looking at
+            the parameters. Determine if any of the parameters are
+            open generic types that need type attributes filled in.
+            </summary>
+            <param name="method">The method to check.</param>
+            <returns>True if any of the parameters are open generics. False if not.</returns>
+        </member>
+        <member name="M:Microsoft.Practices.Unity.Utility.ReflectionHelper.GetClosedParameterType(System.Type[])">
+            <summary>
+            If this type is an open generic, use the
+            given <paramref name="genericArguments"/> array to
+            determine what the required closed type is and return that.
+            </summary>
+            <remarks>If the parameter is not an open type, just
+            return this parameter's type.</remarks>
+            <param name="genericArguments">Type arguments to substitute in for
+            the open type parameters.</param>
+            <returns>Corresponding closed type of this parameter.</returns>
+        </member>
+        <member name="M:Microsoft.Practices.Unity.Utility.ReflectionHelper.GetNamedGenericParameter(System.String)">
+            <summary>
+            Given a generic argument name, return the corresponding type for this
+            closed type. For example, if the current type is SomeType&lt;User&gt;, and the
+            corresponding definition was SomeType&lt;TSomething&gt;, calling this method
+            and passing "TSomething" will return typeof(User).
+            </summary>
+            <param name="parameterName">Name of the generic parameter.</param>
+            <returns>Type of the corresponding generic parameter, or null if there
+            is no matching name.</returns>
+        </member>
+        <member name="P:Microsoft.Practices.Unity.Utility.ReflectionHelper.Type">
+            <summary>
+            The <see cref="P:Microsoft.Practices.Unity.Utility.ReflectionHelper.Type"/> object we're reflecting over.
+            </summary>
+        </member>
+        <member name="P:Microsoft.Practices.Unity.Utility.ReflectionHelper.IsGenericType">
+            <summary>
+            Is this type generic?
+            </summary>
+        </member>
+        <member name="P:Microsoft.Practices.Unity.Utility.ReflectionHelper.IsOpenGeneric">
+            <summary>
+            Is this type an open generic (no type parameter specified)
+            </summary>
+        </member>
+        <member name="P:Microsoft.Practices.Unity.Utility.ReflectionHelper.IsArray">
+            <summary>
+            Is this type an array type?
+            </summary>
+        </member>
+        <member name="P:Microsoft.Practices.Unity.Utility.ReflectionHelper.IsGenericArray">
+            <summary>
+            Is this type an array of generic elements?
+            </summary>
+        </member>
+        <member name="P:Microsoft.Practices.Unity.Utility.ReflectionHelper.ArrayElementType">
+            <summary>
+            The type of the elements in this type (if it's an array).
+            </summary>
+        </member>
+        <member name="P:Microsoft.Practices.Unity.Utility.ReflectionHelper.InstanceConstructors">
+            <summary>
+            Returns all the public constructors defined for the current reflected <see cref="P:Microsoft.Practices.Unity.Utility.ReflectionHelper.Type"/>.
+            </summary>
+            <value>
+            An enumeration of <see cref="T:System.Reflection.ConstructorInfo"/> ConstructorInfo objects representing all the public instance constructors defined for the 
+            current reflected <see cref="P:Microsoft.Practices.Unity.Utility.ReflectionHelper.Type"/>, but not including the type initializer (static constructor).
+            </value>
+        </member>
+        <member name="M:Microsoft.Practices.Unity.Utility.ParameterReflectionHelper.#ctor(System.Reflection.ParameterInfo)">
+            <summary>
+            Create a new instance of <see cref="T:Microsoft.Practices.Unity.Utility.ParameterReflectionHelper"/> that
+            lets you query information about the given ParameterInfo object.
+            </summary>
+            <param name="parameter">Parameter to query.</param>
+        </member>
+        <member name="T:Microsoft.Practices.Unity.Utility.StaticReflection">
+            <summary>
+            A set of helper methods to pick through lambdas and pull out
+            <see cref="T:System.Reflection.MethodInfo"/> from them.
+            </summary>
+        </member>
+        <member name="M:Microsoft.Practices.Unity.Utility.StaticReflection.GetMethodInfo(System.Linq.Expressions.Expression{System.Action})">
+            <summary>
+            Pull out a <see cref="T:System.Reflection.MethodInfo"/> object from an expression of the form
+            () =&gt; SomeClass.SomeMethod()
+            </summary>
+            <param name="expression">Expression describing the method to call.</param>
+            <returns>Corresponding <see cref="T:System.Reflection.MethodInfo"/>.</returns>
+        </member>
+        <member name="M:Microsoft.Practices.Unity.Utility.StaticReflection.GetMethodInfo``1(System.Linq.Expressions.Expression{System.Action{``0}})">
+            <summary>
+            Pull out a <see cref="T:System.Reflection.MethodInfo"/> object from an expression of the form
+            x =&gt; x.SomeMethod()
+            </summary>
+            <typeparam name="T">The type where the method is defined.</typeparam>
+            <param name="expression">Expression describing the method to call.</param>
+            <returns>Corresponding <see cref="T:System.Reflection.MethodInfo"/>.</returns>
+        </member>
+        <member name="M:Microsoft.Practices.Unity.Utility.StaticReflection.GetPropertyGetMethodInfo``2(System.Linq.Expressions.Expression{System.Func{``0,``1}})">
+            <summary>
+            Pull out a <see cref="T:System.Reflection.MethodInfo"/> object for the get method from an expression of the form
+            x =&gt; x.SomeProperty
+            </summary>
+            <typeparam name="T">The type where the method is defined.</typeparam>
+            <typeparam name="TProperty">The type for the property.</typeparam>
+            <param name="expression">Expression describing the property for which the get method is to be extracted.</param>
+            <returns>Corresponding <see cref="T:System.Reflection.MethodInfo"/>.</returns>
+        </member>
+        <member name="M:Microsoft.Practices.Unity.Utility.StaticReflection.GetPropertySetMethodInfo``2(System.Linq.Expressions.Expression{System.Func{``0,``1}})">
+            <summary>
+            Pull out a <see cref="T:System.Reflection.MethodInfo"/> object for the set method from an expression of the form
+            x =&gt; x.SomeProperty
+            </summary>
+            <typeparam name="T">The type where the method is defined.</typeparam>
+            <typeparam name="TProperty">The type for the property.</typeparam>
+            <param name="expression">Expression describing the property for which the set method is to be extracted.</param>
+            <returns>Corresponding <see cref="T:System.Reflection.MethodInfo"/>.</returns>
+        </member>
+        <member name="M:Microsoft.Practices.Unity.Utility.StaticReflection.GetMemberInfo``2(System.Linq.Expressions.Expression{System.Func{``0,``1}})">
+            <summary>
+            
+            </summary>
+            <typeparam name="T"></typeparam>
+            <typeparam name="TProperty"></typeparam>
+            <param name="expression"></param>
+            <returns></returns>
+        </member>
+        <member name="M:Microsoft.Practices.Unity.Utility.StaticReflection.GetConstructorInfo``1(System.Linq.Expressions.Expression{System.Func{``0}})">
+            <summary>
+            Pull out a <see cref="T:System.Reflection.ConstructorInfo"/> object from an expression of the form () =&gt; new SomeType()
+            </summary>
+            <typeparam name="T">The type where the constructor is defined.</typeparam>
+            <param name="expression">Expression invoking the desired constructor.</param>
+            <returns>Corresponding <see cref="T:System.Reflection.ConstructorInfo"/>.</returns>
+        </member>
+    </members>
+</doc>
