Copying //UE4/Dev-Core to //UE4/Main

diff --git a/Engine/Source/Programs/CrashReporter/CrashReportWebSite/packages/Unity.4.0.1/lib/win8/Microsoft.Practices.Unity.RegistrationByConvention.XML b/Engine/Source/Programs/CrashReporter/CrashReportWebSite/packages/Unity.4.0.1/lib/win8/Microsoft.Practices.Unity.RegistrationByConvention.XML
new file mode 100644
--- /dev/null
+++ b/Engine/Source/Programs/CrashReporter/CrashReportWebSite/packages/Unity.4.0.1/lib/win8/Microsoft.Practices.Unity.RegistrationByConvention.XML
@@ -0,0 +1,306 @@
+<?xml version="1.0"?>
+<doc>
+    <assembly>
+        <name>Microsoft.Practices.Unity.RegistrationByConvention</name>
+    </assembly>
+    <members>
+        <member name="T:Microsoft.Practices.Unity.AllClasses">
+            <summary>
+            Provides helper methods to retrieve classes from assemblies.
+            </summary>
+        </member>
+        <member name="M:Microsoft.Practices.Unity.AllClasses.FromAssemblies(System.Reflection.Assembly[])">
+            <summary>
+            Returns all visible, non-abstract classes from <paramref name="assemblies"/>.
+            </summary>
+            <param name="assemblies">The assemblies.</param>
+            <returns>All visible, non-abstract classes found in the assemblies.</returns>
+            <exception cref="T:System.ArgumentNullException"><paramref name="assemblies"/> is <see langword="null"/>.</exception>
+            <exception cref="T:System.ArgumentException"><paramref name="assemblies"/> contains <see langword="null"/> elements.</exception>
+            <remarks>All exceptions thrown while getting types from the assemblies are ignored, and the types that can be retrieved are returned.</remarks>
+        </member>
+        <member name="M:Microsoft.Practices.Unity.AllClasses.FromAssemblies(System.Boolean,System.Reflection.Assembly[])">
+            <summary>
+            Returns all visible, non-abstract classes from <paramref name="assemblies"/>, and optionally skips errors.
+            </summary>
+            <param name="skipOnError"><see langword="true"/> to skip errors; otherwise, <see langword="true"/>.</param>
+            <param name="assemblies">The assemblies.</param>
+            <returns>
+            All visible, non-abstract classes.
+            </returns>
+            <exception cref="T:System.ArgumentNullException"><paramref name="assemblies"/> is <see langword="null"/>.</exception>
+            <exception cref="T:System.ArgumentException"><paramref name="assemblies"/> contains <see langword="null"/> elements.</exception>
+            <remarks>
+            If <paramref name="skipOnError"/> is <see langword="true"/>, all exceptions thrown while getting types from the assemblies are ignored, and the types 
+            that can be retrieved are returned; otherwise, the original exception is thrown.
+            </remarks>
+        </member>
+        <member name="M:Microsoft.Practices.Unity.AllClasses.FromAssemblies(System.Collections.Generic.IEnumerable{System.Reflection.Assembly},System.Boolean)">
+            <summary>
+            Returns all visible, non-abstract classes from <paramref name="assemblies"/>.
+            </summary>
+            <param name="skipOnError"><see langword="true"/> to skip errors; otherwise, <see langword="true"/>.</param>
+            <param name="assemblies">The assemblies.</param>
+            <returns>
+            All visible, non-abstract classes.
+            </returns>
+            <exception cref="T:System.ArgumentNullException"><paramref name="assemblies"/> is <see langword="null"/>.</exception>
+            <exception cref="T:System.ArgumentException"><paramref name="assemblies"/> contains <see langword="null"/> elements.</exception>
+            <remarks>
+            If <paramref name="skipOnError"/> is <see langword="true"/>, all exceptions thrown while getting types from the assemblies are ignored, and the types 
+            that can be retrieved are returned; otherwise, the original exception is thrown.
+            </remarks>
+        </member>
+        <member name="M:Microsoft.Practices.Unity.AllClasses.FromApplication(System.Boolean,System.Boolean)">
+            <summary>
+            Returns all visible, non-abstract classes from all assemblies located where the application is installed.
+            </summary>
+            <param name="includeUnityAssemblies"><see langword="false"/> to include the Unity assemblies; otherwise, <see langword="false"/>. Defaults to <see langword="false"/>.</param>
+            <param name="skipOnError"><see langword="true"/> to skip errors; otherwise, <see langword="true"/>.</param>
+            <returns>
+            All visible, non-abstract classes.
+            </returns>
+            <remarks>
+            If <paramref name="skipOnError"/> is <see langword="true"/>, all exceptions thrown while loading assemblies or getting types from the assemblies 
+            are ignored, and the types that can be retrieved are returned; otherwise, the original exception is thrown. These exceptions might be wrapped in a
+            <see cref="T:System.AggregateException"/>.
+            </remarks>
+        </member>
+        <member name="T:Microsoft.Practices.Unity.DuplicateTypeMappingException">
+            <summary>
+            The exception that is thrown when registering multiple types would result in an type mapping being overwritten.
+            </summary>
+        </member>
+        <member name="M:Microsoft.Practices.Unity.DuplicateTypeMappingException.#ctor(System.String,System.Type,System.Type,System.Type)">
+            <summary>
+            Initializes a new instance of the <see cref="T:Microsoft.Practices.Unity.DuplicateTypeMappingException"/> class.
+            </summary>
+            <param name="name">The name for the mapping.</param>
+            <param name="mappedFromType">The source type for the mapping.</param>
+            <param name="currentMappedToType">The type currently mapped.</param>
+            <param name="newMappedToType">The new type to map.</param>
+        </member>
+        <member name="P:Microsoft.Practices.Unity.DuplicateTypeMappingException.Name">
+            <summary>
+            Gets the name for the mapping.
+            </summary>
+        </member>
+        <member name="P:Microsoft.Practices.Unity.DuplicateTypeMappingException.MappedFromType">
+            <summary>
+            Gets the source type for the mapping.
+            </summary>
+        </member>
+        <member name="P:Microsoft.Practices.Unity.DuplicateTypeMappingException.CurrentMappedToType">
+            <summary>
+            Gets the type currently mapped.
+            </summary>
+        </member>
+        <member name="P:Microsoft.Practices.Unity.DuplicateTypeMappingException.NewMappedToType">
+            <summary>
+            Gets the new type to map.
+            </summary>
+        </member>
+        <member name="T:Microsoft.Practices.Unity.RegistrationConvention">
+            <summary>
+            Represents a set of types to register and their registration settings.
+            </summary>
+        </member>
+        <member name="M:Microsoft.Practices.Unity.RegistrationConvention.GetTypes">
+            <summary>
+            Gets types to register.
+            </summary>
+        </member>
+        <member name="M:Microsoft.Practices.Unity.RegistrationConvention.GetFromTypes">
+            <summary>
+            Gets a function to get the types that will be requested for each type to configure.
+            </summary>
+        </member>
+        <member name="M:Microsoft.Practices.Unity.RegistrationConvention.GetName">
+            <summary>
+            Gets a function to get the name to use for the registration of each type.
+            </summary>
+        </member>
+        <member name="M:Microsoft.Practices.Unity.RegistrationConvention.GetLifetimeManager">
+            <summary>
+            Gets a function to get the <see cref="T:Microsoft.Practices.Unity.LifetimeManager"/> for the registration of each type. Defaults to no lifetime management.
+            </summary>
+            <returns></returns>
+        </member>
+        <member name="M:Microsoft.Practices.Unity.RegistrationConvention.GetInjectionMembers">
+            <summary>
+            Gets a function to get the additional <see cref="T:Microsoft.Practices.Unity.InjectionMember"/> objects for the registration of each type. Defaults to no injection members.
+            </summary>
+            <returns></returns>
+        </member>
+        <member name="T:Microsoft.Practices.Unity.Resources">
+            <summary>
+              A strongly-typed resource class, for looking up localized strings, etc.
+            </summary>
+        </member>
+        <member name="P:Microsoft.Practices.Unity.Resources.ResourceManager">
+            <summary>
+              Returns the cached ResourceManager instance used by this class.
+            </summary>
+        </member>
+        <member name="P:Microsoft.Practices.Unity.Resources.DuplicateTypeMappingException">
+            <summary>
+            An attempt to override an existing mapping was detected for type {1} wit...
+            </summary>
+        </member>
+        <member name="P:Microsoft.Practices.Unity.Resources.ExceptionNullAssembly">
+            <summary>
+            The set of assemblies contains a null element.
+            </summary>
+        </member>
+        <member name="T:Microsoft.Practices.Unity.UnityContainerRegistrationByConventionExtensions">
+            <summary>
+            Provides a set of convenience overloads to the
+            <see cref="T:Microsoft.Practices.Unity.IUnityContainer"/> interface to support registration of multiple types.
+            </summary>
+        </member>
+        <member name="M:Microsoft.Practices.Unity.UnityContainerRegistrationByConventionExtensions.RegisterTypes(Microsoft.Practices.Unity.IUnityContainer,System.Collections.Generic.IEnumerable{System.Type},System.Func{System.Type,System.Collections.Generic.IEnumerable{System.Type}},System.Func{System.Type,System.String},System.Func{System.Type,Microsoft.Practices.Unity.LifetimeManager},System.Func{System.Type,System.Collections.Generic.IEnumerable{Microsoft.Practices.Unity.InjectionMember}},System.Boolean)">
+            <summary>
+            Registers the supplied types by using the specified rules for name, lifetime manager, injection members, and registration types.
+            </summary>
+            <param name="container">The container to configure.</param>
+            <param name="types">The types to register. The methods in the <see cref="T:Microsoft.Practices.Unity.AllClasses"/> class can be used to scan assemblies to get types, and further filtering can be performed using LINQ queries.</param>
+            <param name="getFromTypes">A function that gets the types that will be requested for each type to configure. It can be a method from the <see cref="T:Microsoft.Practices.Unity.WithMappings"/> class or a custom function. Defaults to no registration types, and registers only the supplied types.</param>
+            <param name="getName">A function that gets the name to use for the registration of each type. It can be a method from the <see cref="T:Microsoft.Practices.Unity.WithName"/> or a custom function. Defaults to no name.</param>
+            <param name="getLifetimeManager">A function that gets the <see cref="T:Microsoft.Practices.Unity.LifetimeManager"/> for the registration of each type. It can be a method from the <see cref="T:Microsoft.Practices.Unity.WithLifetime"/> class or a custom function. Defaults to no lifetime management.</param>
+            <param name="getInjectionMembers">A function that gets the additional <see cref="T:Microsoft.Practices.Unity.InjectionMember"/> objects for the registration of each type. Defaults to no injection members.</param>
+            <param name="overwriteExistingMappings"><see langword="true"/> to overwrite existing mappings; otherwise, <see langword="false"/>. Defaults to <see langword="false"/>.</param>
+            <returns>
+            The container that this method was called on.
+            </returns>
+            <exception cref="T:System.ArgumentException">A new registration would overwrite an existing mapping and <paramref name="overwriteExistingMappings"/> is <see langword="false"/>.</exception>
+        </member>
+        <member name="M:Microsoft.Practices.Unity.UnityContainerRegistrationByConventionExtensions.RegisterTypes(Microsoft.Practices.Unity.IUnityContainer,Microsoft.Practices.Unity.RegistrationConvention,System.Boolean)">
+            <summary>
+            Registers the types according to the <paramref name="convention"/>.
+            </summary>
+            <param name="container">The container to configure.</param>
+            <param name="convention">The convention to determine which types will be registered and how.</param>
+            <param name="overwriteExistingMappings"><see langword="true"/> to overwrite existing mappings; otherwise, <see langword="false"/>. Defaults to <see langword="false"/>.</param>
+            <returns>
+            The container that this method was called on.
+            </returns>
+        </member>
+        <member name="T:Microsoft.Practices.Unity.WithLifetime">
+            <summary>
+            Provides helper methods to specify the lifetime for a type with registration by convention.
+            </summary>
+        </member>
+        <member name="M:Microsoft.Practices.Unity.WithLifetime.None(System.Type)">
+            <summary>
+            Returns a <see langword="null"/> <see cref="T:Microsoft.Practices.Unity.LifetimeManager"/>.
+            </summary>
+            <param name="type">The type.</param>
+            <returns>A lifetime manager</returns>
+        </member>
+        <member name="M:Microsoft.Practices.Unity.WithLifetime.ContainerControlled(System.Type)">
+            <summary>
+            Returns a <see cref="T:Microsoft.Practices.Unity.ContainerControlledLifetimeManager"/>.
+            </summary>
+            <param name="type">The type.</param>
+            <returns>A container controlled lifetime manager.</returns>
+        </member>
+        <member name="M:Microsoft.Practices.Unity.WithLifetime.ExternallyControlled(System.Type)">
+            <summary>
+            Returns a <see cref="T:Microsoft.Practices.Unity.ExternallyControlledLifetimeManager"/>.
+            </summary>
+            <param name="type">The type.</param>
+            <returns>An externally controlled lifetime manager.</returns>
+        </member>
+        <member name="M:Microsoft.Practices.Unity.WithLifetime.Hierarchical(System.Type)">
+            <summary>
+            Returns a <see cref="T:Microsoft.Practices.Unity.HierarchicalLifetimeManager"/>.
+            </summary>
+            <param name="type">The type.</param>
+            <returns>A hierarchical lifetime manager.</returns>
+        </member>
+        <member name="M:Microsoft.Practices.Unity.WithLifetime.PerResolve(System.Type)">
+            <summary>
+            Returns a <see cref="T:Microsoft.Practices.Unity.PerResolveLifetimeManager"/>.
+            </summary>
+            <param name="type">The type.</param>
+            <returns>A per resolve lifetime manager.</returns>
+        </member>
+        <member name="M:Microsoft.Practices.Unity.WithLifetime.Transient(System.Type)">
+            <summary>
+            Returns a <see cref="T:Microsoft.Practices.Unity.TransientLifetimeManager"/>.
+            </summary>
+            <param name="type">The type.</param>
+            <returns>A transient lifetime manager.</returns>
+        </member>
+        <member name="M:Microsoft.Practices.Unity.WithLifetime.Custom``1(System.Type)">
+            <summary>
+            Returns a <see cref="T:Microsoft.Practices.Unity.LifetimeManager"/>.
+            </summary>
+            <typeparam name="T">The custom <see cref="T:Microsoft.Practices.Unity.LifetimeManager"/> type.</typeparam>
+            <param name="type">The type.</param>
+            <returns>
+            A lifetime manager.
+            </returns>
+        </member>
+        <member name="M:Microsoft.Practices.Unity.WithLifetime.PerThread(System.Type)">
+            <summary>
+            Returns a <see cref="T:Microsoft.Practices.Unity.PerThreadLifetimeManager"/>.
+            </summary>
+            <param name="type">The type.</param>
+            <returns>A per thread lifetime manager.</returns>
+        </member>
+        <member name="T:Microsoft.Practices.Unity.WithMappings">
+            <summary>
+            Provides helper methods to map types to the types interfaces to which register them.
+            </summary>
+        </member>
+        <member name="M:Microsoft.Practices.Unity.WithMappings.None(System.Type)">
+            <summary>
+            Returns no types.
+            </summary>
+            <param name="implementationType">The type to register.</param>
+            <returns>An empty enumeration.</returns>
+        </member>
+        <member name="M:Microsoft.Practices.Unity.WithMappings.FromMatchingInterface(System.Type)">
+            <summary>
+            Returns an enumeration with the interface that matches the name of <paramref name="implementationType"/>.
+            </summary>
+            <param name="implementationType">The type to register.</param>
+            <returns>An enumeration with the first interface matching the name of <paramref name="implementationType"/> (for example, if type is MyType, a matching interface is IMyType),
+            or an empty enumeration if no such interface is found.</returns>
+        </member>
+        <member name="M:Microsoft.Practices.Unity.WithMappings.FromAllInterfaces(System.Type)">
+            <summary>
+            Returns an enumeration with all the interfaces implemented by <paramref name="implementationType"/>.
+            </summary>
+            <param name="implementationType">The type to register.</param>
+            <returns>An enumeration with all the interfaces implemented by the implementation type except <see cref="T:System.IDisposable"/>.</returns>
+        </member>
+        <member name="M:Microsoft.Practices.Unity.WithMappings.FromAllInterfacesInSameAssembly(System.Type)">
+            <summary>
+            Returns an enumeration with all the interfaces implemented by <paramref name="implementationType"/> that belong to the same assembly as implementationType.
+            </summary>
+            <param name="implementationType">The type to register.</param>
+            <returns>An enumeration with all the interfaces implemented by the implementation type that belong to the same assembly.</returns>
+        </member>
+        <member name="T:Microsoft.Practices.Unity.WithName">
+            <summary>
+            Provides helper methods to get type names.
+            </summary>
+        </member>
+        <member name="M:Microsoft.Practices.Unity.WithName.TypeName(System.Type)">
+            <summary>
+            Returns the type name.
+            </summary>
+            <param name="type">The type.</param>
+            <returns>The type name.</returns>
+        </member>
+        <member name="M:Microsoft.Practices.Unity.WithName.Default(System.Type)">
+            <summary>
+            Returns null for the registration name.
+            </summary>
+            <param name="type">The type.</param>
+            <returns><see langword="null"/></returns>
+        </member>
+    </members>
+</doc>
diff --git a/Engine/Source/Programs/CrashReporter/CrashReportWebSite/packages/Unity.4.0.1/lib/wp80/Microsoft.Practices.Unity.RegistrationByConvention.XML b/Engine/Source/Programs/CrashReporter/CrashReportWebSite/packages/Unity.4.0.1/lib/wp80/Microsoft.Practices.Unity.RegistrationByConvention.XML
new file mode 100644
--- /dev/null
+++ b/Engine/Source/Programs/CrashReporter/CrashReportWebSite/packages/Unity.4.0.1/lib/wp80/Microsoft.Practices.Unity.RegistrationByConvention.XML
@@ -0,0 +1,312 @@
+<?xml version="1.0"?>
+<doc>
+    <assembly>
+        <name>Microsoft.Practices.Unity.RegistrationByConvention</name>
+    </assembly>
+    <members>
+        <member name="T:Microsoft.Practices.Unity.AllClasses">
+            <summary>
+            Provides helper methods to retrieve classes from assemblies.
+            </summary>
+        </member>
+        <member name="M:Microsoft.Practices.Unity.AllClasses.FromAssemblies(System.Reflection.Assembly[])">
+            <summary>
+            Returns all visible, non-abstract classes from <paramref name="assemblies"/>.
+            </summary>
+            <param name="assemblies">The assemblies.</param>
+            <returns>All visible, non-abstract classes found in the assemblies.</returns>
+            <exception cref="T:System.ArgumentNullException"><paramref name="assemblies"/> is <see langword="null"/>.</exception>
+            <exception cref="T:System.ArgumentException"><paramref name="assemblies"/> contains <see langword="null"/> elements.</exception>
+            <remarks>All exceptions thrown while getting types from the assemblies are ignored, and the types that can be retrieved are returned.</remarks>
+        </member>
+        <member name="M:Microsoft.Practices.Unity.AllClasses.FromAssemblies(System.Boolean,System.Reflection.Assembly[])">
+            <summary>
+            Returns all visible, non-abstract classes from <paramref name="assemblies"/>, and optionally skips errors.
+            </summary>
+            <param name="skipOnError"><see langword="true"/> to skip errors; otherwise, <see langword="true"/>.</param>
+            <param name="assemblies">The assemblies.</param>
+            <returns>
+            All visible, non-abstract classes.
+            </returns>
+            <exception cref="T:System.ArgumentNullException"><paramref name="assemblies"/> is <see langword="null"/>.</exception>
+            <exception cref="T:System.ArgumentException"><paramref name="assemblies"/> contains <see langword="null"/> elements.</exception>
+            <remarks>
+            If <paramref name="skipOnError"/> is <see langword="true"/>, all exceptions thrown while getting types from the assemblies are ignored, and the types 
+            that can be retrieved are returned; otherwise, the original exception is thrown.
+            </remarks>
+        </member>
+        <member name="M:Microsoft.Practices.Unity.AllClasses.FromAssemblies(System.Collections.Generic.IEnumerable{System.Reflection.Assembly},System.Boolean)">
+            <summary>
+            Returns all visible, non-abstract classes from <paramref name="assemblies"/>.
+            </summary>
+            <param name="skipOnError"><see langword="true"/> to skip errors; otherwise, <see langword="true"/>.</param>
+            <param name="assemblies">The assemblies.</param>
+            <returns>
+            All visible, non-abstract classes.
+            </returns>
+            <exception cref="T:System.ArgumentNullException"><paramref name="assemblies"/> is <see langword="null"/>.</exception>
+            <exception cref="T:System.ArgumentException"><paramref name="assemblies"/> contains <see langword="null"/> elements.</exception>
+            <remarks>
+            If <paramref name="skipOnError"/> is <see langword="true"/>, all exceptions thrown while getting types from the assemblies are ignored, and the types 
+            that can be retrieved are returned; otherwise, the original exception is thrown.
+            </remarks>
+        </member>
+        <member name="M:Microsoft.Practices.Unity.AllClasses.FromApplication(System.Boolean,System.Boolean)">
+            <summary>
+            Returns all visible, non-abstract classes from all assemblies located where the application is installed.
+            </summary>
+            <param name="includeUnityAssemblies"><see langword="false"/> to include the Unity assemblies; otherwise, <see langword="false"/>. Defaults to <see langword="false"/>.</param>
+            <param name="skipOnError"><see langword="true"/> to skip errors; otherwise, <see langword="true"/>.</param>
+            <returns>
+            All visible, non-abstract classes.
+            </returns>
+            <remarks>
+            If <paramref name="skipOnError"/> is <see langword="true"/>, all exceptions thrown while loading assemblies or getting types from the assemblies 
+            are ignored, and the types that can be retrieved are returned; otherwise, the original exception is thrown. These exceptions might be wrapped in a
+            <see cref="T:System.AggregateException"/>.
+            </remarks>
+        </member>
+        <member name="T:Microsoft.Practices.Unity.DuplicateTypeMappingException">
+            <summary>
+            The exception that is thrown when registering multiple types would result in an type mapping being overwritten.
+            </summary>
+        </member>
+        <member name="M:Microsoft.Practices.Unity.DuplicateTypeMappingException.#ctor(System.String,System.Type,System.Type,System.Type)">
+            <summary>
+            Initializes a new instance of the <see cref="T:Microsoft.Practices.Unity.DuplicateTypeMappingException"/> class.
+            </summary>
+            <param name="name">The name for the mapping.</param>
+            <param name="mappedFromType">The source type for the mapping.</param>
+            <param name="currentMappedToType">The type currently mapped.</param>
+            <param name="newMappedToType">The new type to map.</param>
+        </member>
+        <member name="P:Microsoft.Practices.Unity.DuplicateTypeMappingException.Name">
+            <summary>
+            Gets the name for the mapping.
+            </summary>
+        </member>
+        <member name="P:Microsoft.Practices.Unity.DuplicateTypeMappingException.MappedFromType">
+            <summary>
+            Gets the source type for the mapping.
+            </summary>
+        </member>
+        <member name="P:Microsoft.Practices.Unity.DuplicateTypeMappingException.CurrentMappedToType">
+            <summary>
+            Gets the type currently mapped.
+            </summary>
+        </member>
+        <member name="P:Microsoft.Practices.Unity.DuplicateTypeMappingException.NewMappedToType">
+            <summary>
+            Gets the new type to map.
+            </summary>
+        </member>
+        <member name="T:Microsoft.Practices.Unity.Properties.Resources">
+            <summary>
+              A strongly-typed resource class, for looking up localized strings, etc.
+            </summary>
+        </member>
+        <member name="P:Microsoft.Practices.Unity.Properties.Resources.ResourceManager">
+            <summary>
+              Returns the cached ResourceManager instance used by this class.
+            </summary>
+        </member>
+        <member name="P:Microsoft.Practices.Unity.Properties.Resources.Culture">
+            <summary>
+              Overrides the current thread's CurrentUICulture property for all
+              resource lookups using this strongly typed resource class.
+            </summary>
+        </member>
+        <member name="P:Microsoft.Practices.Unity.Properties.Resources.DuplicateTypeMappingException">
+            <summary>
+              Looks up a localized string similar to An attempt to override an existing mapping was detected for type {1} with name &quot;{0}&quot;, currently mapped to type {2}, to type {3}..
+            </summary>
+        </member>
+        <member name="P:Microsoft.Practices.Unity.Properties.Resources.ExceptionNullAssembly">
+            <summary>
+              Looks up a localized string similar to The set of assemblies contains a null element..
+            </summary>
+        </member>
+        <member name="T:Microsoft.Practices.Unity.RegistrationConvention">
+            <summary>
+            Represents a set of types to register and their registration settings.
+            </summary>
+        </member>
+        <member name="M:Microsoft.Practices.Unity.RegistrationConvention.GetTypes">
+            <summary>
+            Gets types to register.
+            </summary>
+        </member>
+        <member name="M:Microsoft.Practices.Unity.RegistrationConvention.GetFromTypes">
+            <summary>
+            Gets a function to get the types that will be requested for each type to configure.
+            </summary>
+        </member>
+        <member name="M:Microsoft.Practices.Unity.RegistrationConvention.GetName">
+            <summary>
+            Gets a function to get the name to use for the registration of each type.
+            </summary>
+        </member>
+        <member name="M:Microsoft.Practices.Unity.RegistrationConvention.GetLifetimeManager">
+            <summary>
+            Gets a function to get the <see cref="T:Microsoft.Practices.Unity.LifetimeManager"/> for the registration of each type. Defaults to no lifetime management.
+            </summary>
+            <returns></returns>
+        </member>
+        <member name="M:Microsoft.Practices.Unity.RegistrationConvention.GetInjectionMembers">
+            <summary>
+            Gets a function to get the additional <see cref="T:Microsoft.Practices.Unity.InjectionMember"/> objects for the registration of each type. Defaults to no injection members.
+            </summary>
+            <returns></returns>
+        </member>
+        <member name="T:Microsoft.Practices.Unity.UnityContainerRegistrationByConventionExtensions">
+            <summary>
+            Provides a set of convenience overloads to the
+            <see cref="T:Microsoft.Practices.Unity.IUnityContainer"/> interface to support registration of multiple types.
+            </summary>
+        </member>
+        <member name="M:Microsoft.Practices.Unity.UnityContainerRegistrationByConventionExtensions.RegisterTypes(Microsoft.Practices.Unity.IUnityContainer,System.Collections.Generic.IEnumerable{System.Type},System.Func{System.Type,System.Collections.Generic.IEnumerable{System.Type}},System.Func{System.Type,System.String},System.Func{System.Type,Microsoft.Practices.Unity.LifetimeManager},System.Func{System.Type,System.Collections.Generic.IEnumerable{Microsoft.Practices.Unity.InjectionMember}},System.Boolean)">
+            <summary>
+            Registers the supplied types by using the specified rules for name, lifetime manager, injection members, and registration types.
+            </summary>
+            <param name="container">The container to configure.</param>
+            <param name="types">The types to register. The methods in the <see cref="T:Microsoft.Practices.Unity.AllClasses"/> class can be used to scan assemblies to get types, and further filtering can be performed using LINQ queries.</param>
+            <param name="getFromTypes">A function that gets the types that will be requested for each type to configure. It can be a method from the <see cref="T:Microsoft.Practices.Unity.WithMappings"/> class or a custom function. Defaults to no registration types, and registers only the supplied types.</param>
+            <param name="getName">A function that gets the name to use for the registration of each type. It can be a method from the <see cref="T:Microsoft.Practices.Unity.WithName"/> or a custom function. Defaults to no name.</param>
+            <param name="getLifetimeManager">A function that gets the <see cref="T:Microsoft.Practices.Unity.LifetimeManager"/> for the registration of each type. It can be a method from the <see cref="T:Microsoft.Practices.Unity.WithLifetime"/> class or a custom function. Defaults to no lifetime management.</param>
+            <param name="getInjectionMembers">A function that gets the additional <see cref="T:Microsoft.Practices.Unity.InjectionMember"/> objects for the registration of each type. Defaults to no injection members.</param>
+            <param name="overwriteExistingMappings"><see langword="true"/> to overwrite existing mappings; otherwise, <see langword="false"/>. Defaults to <see langword="false"/>.</param>
+            <returns>
+            The container that this method was called on.
+            </returns>
+            <exception cref="T:System.ArgumentException">A new registration would overwrite an existing mapping and <paramref name="overwriteExistingMappings"/> is <see langword="false"/>.</exception>
+        </member>
+        <member name="M:Microsoft.Practices.Unity.UnityContainerRegistrationByConventionExtensions.RegisterTypes(Microsoft.Practices.Unity.IUnityContainer,Microsoft.Practices.Unity.RegistrationConvention,System.Boolean)">
+            <summary>
+            Registers the types according to the <paramref name="convention"/>.
+            </summary>
+            <param name="container">The container to configure.</param>
+            <param name="convention">The convention to determine which types will be registered and how.</param>
+            <param name="overwriteExistingMappings"><see langword="true"/> to overwrite existing mappings; otherwise, <see langword="false"/>. Defaults to <see langword="false"/>.</param>
+            <returns>
+            The container that this method was called on.
+            </returns>
+        </member>
+        <member name="T:Microsoft.Practices.Unity.WithLifetime">
+            <summary>
+            Provides helper methods to specify the lifetime for a type with registration by convention.
+            </summary>
+        </member>
+        <member name="M:Microsoft.Practices.Unity.WithLifetime.None(System.Type)">
+            <summary>
+            Returns a <see langword="null"/> <see cref="T:Microsoft.Practices.Unity.LifetimeManager"/>.
+            </summary>
+            <param name="type">The type.</param>
+            <returns>A lifetime manager</returns>
+        </member>
+        <member name="M:Microsoft.Practices.Unity.WithLifetime.ContainerControlled(System.Type)">
+            <summary>
+            Returns a <see cref="T:Microsoft.Practices.Unity.ContainerControlledLifetimeManager"/>.
+            </summary>
+            <param name="type">The type.</param>
+            <returns>A container controlled lifetime manager.</returns>
+        </member>
+        <member name="M:Microsoft.Practices.Unity.WithLifetime.ExternallyControlled(System.Type)">
+            <summary>
+            Returns a <see cref="T:Microsoft.Practices.Unity.ExternallyControlledLifetimeManager"/>.
+            </summary>
+            <param name="type">The type.</param>
+            <returns>An externally controlled lifetime manager.</returns>
+        </member>
+        <member name="M:Microsoft.Practices.Unity.WithLifetime.Hierarchical(System.Type)">
+            <summary>
+            Returns a <see cref="T:Microsoft.Practices.Unity.HierarchicalLifetimeManager"/>.
+            </summary>
+            <param name="type">The type.</param>
+            <returns>A hierarchical lifetime manager.</returns>
+        </member>
+        <member name="M:Microsoft.Practices.Unity.WithLifetime.PerResolve(System.Type)">
+            <summary>
+            Returns a <see cref="T:Microsoft.Practices.Unity.PerResolveLifetimeManager"/>.
+            </summary>
+            <param name="type">The type.</param>
+            <returns>A per resolve lifetime manager.</returns>
+        </member>
+        <member name="M:Microsoft.Practices.Unity.WithLifetime.Transient(System.Type)">
+            <summary>
+            Returns a <see cref="T:Microsoft.Practices.Unity.TransientLifetimeManager"/>.
+            </summary>
+            <param name="type">The type.</param>
+            <returns>A transient lifetime manager.</returns>
+        </member>
+        <member name="M:Microsoft.Practices.Unity.WithLifetime.Custom``1(System.Type)">
+            <summary>
+            Returns a <see cref="T:Microsoft.Practices.Unity.LifetimeManager"/>.
+            </summary>
+            <typeparam name="T">The custom <see cref="T:Microsoft.Practices.Unity.LifetimeManager"/> type.</typeparam>
+            <param name="type">The type.</param>
+            <returns>
+            A lifetime manager.
+            </returns>
+        </member>
+        <member name="M:Microsoft.Practices.Unity.WithLifetime.PerThread(System.Type)">
+            <summary>
+            Returns a <see cref="T:Microsoft.Practices.Unity.PerThreadLifetimeManager"/>.
+            </summary>
+            <param name="type">The type.</param>
+            <returns>A per thread lifetime manager.</returns>
+        </member>
+        <member name="T:Microsoft.Practices.Unity.WithMappings">
+            <summary>
+            Provides helper methods to map types to the types interfaces to which register them.
+            </summary>
+        </member>
+        <member name="M:Microsoft.Practices.Unity.WithMappings.None(System.Type)">
+            <summary>
+            Returns no types.
+            </summary>
+            <param name="implementationType">The type to register.</param>
+            <returns>An empty enumeration.</returns>
+        </member>
+        <member name="M:Microsoft.Practices.Unity.WithMappings.FromMatchingInterface(System.Type)">
+            <summary>
+            Returns an enumeration with the interface that matches the name of <paramref name="implementationType"/>.
+            </summary>
+            <param name="implementationType">The type to register.</param>
+            <returns>An enumeration with the first interface matching the name of <paramref name="implementationType"/> (for example, if type is MyType, a matching interface is IMyType),
+            or an empty enumeration if no such interface is found.</returns>
+        </member>
+        <member name="M:Microsoft.Practices.Unity.WithMappings.FromAllInterfaces(System.Type)">
+            <summary>
+            Returns an enumeration with all the interfaces implemented by <paramref name="implementationType"/>.
+            </summary>
+            <param name="implementationType">The type to register.</param>
+            <returns>An enumeration with all the interfaces implemented by the implementation type except <see cref="T:System.IDisposable"/>.</returns>
+        </member>
+        <member name="M:Microsoft.Practices.Unity.WithMappings.FromAllInterfacesInSameAssembly(System.Type)">
+            <summary>
+            Returns an enumeration with all the interfaces implemented by <paramref name="implementationType"/> that belong to the same assembly as implementationType.
+            </summary>
+            <param name="implementationType">The type to register.</param>
+            <returns>An enumeration with all the interfaces implemented by the implementation type that belong to the same assembly.</returns>
+        </member>
+        <member name="T:Microsoft.Practices.Unity.WithName">
+            <summary>
+            Provides helper methods to get type names.
+            </summary>
+        </member>
+        <member name="M:Microsoft.Practices.Unity.WithName.TypeName(System.Type)">
+            <summary>
+            Returns the type name.
+            </summary>
+            <param name="type">The type.</param>
+            <returns>The type name.</returns>
+        </member>
+        <member name="M:Microsoft.Practices.Unity.WithName.Default(System.Type)">
+            <summary>
+            Returns null for the registration name.
+            </summary>
+            <param name="type">The type.</param>
+            <returns><see langword="null"/></returns>
+        </member>
+    </members>
+</doc>
diff --git a/Engine/Source/Programs/CrashReporter/CrashReportWebSite/packages/Unity.Mvc.4.0.1/lib/net45/Microsoft.Practices.Unity.Mvc.XML b/Engine/Source/Programs/CrashReporter/CrashReportWebSite/packages/Unity.Mvc.4.0.1/lib/net45/Microsoft.Practices.Unity.Mvc.XML
new file mode 100644
--- /dev/null
+++ b/Engine/Source/Programs/CrashReporter/CrashReportWebSite/packages/Unity.Mvc.4.0.1/lib/net45/Microsoft.Practices.Unity.Mvc.XML
@@ -0,0 +1,145 @@
+<?xml version="1.0"?>
+<doc>
+    <assembly>
+        <name>Microsoft.Practices.Unity.Mvc</name>
+    </assembly>
+    <members>
+        <member name="T:Microsoft.Practices.Unity.PerRequestLifetimeManager">
+            <summary>
+            A <see cref="T:Microsoft.Practices.Unity.LifetimeManager"/> that holds onto the instance given to it during
+            the lifetime of a single HTTP request.
+            This lifetime manager enables you to create instances of registered types that behave like
+            singletons within the scope of an HTTP request.
+            See remarks for important usage information.
+            </summary>
+            <remarks>
+            <para>
+            Although the <see cref="T:Microsoft.Practices.Unity.PerRequestLifetimeManager"/> lifetime manager works correctly and can help
+            in working with stateful or thread-unsafe dependencies within the scope of an HTTP request, it is
+            generally not a good idea to use it when it can be avoided, as it can often lead to bad practices or
+            hard to find bugs in the end-user's application code when used incorrectly. 
+            It is recommended that the dependencies you register are stateless and if there is a need to share
+            common state between several objects during the lifetime of an HTTP request, then you can
+            have a stateless service that explicitly stores and retrieves this state using the
+            <see cref="P:System.Web.HttpContext.Items"/> collection of the <see cref="P:System.Web.HttpContext.Current"/> object.
+            </para>
+            <para>
+            For the instance of the registered type to be disposed automatically when the HTTP request completes,
+            make sure to register the <see cref="T:Microsoft.Practices.Unity.Mvc.UnityPerRequestHttpModule"/> with the web application.
+            To do this, invoke the following in the Unity bootstrapping class (typically UnityMvcActivator.cs):
+            <code>DynamicModuleUtility.RegisterModule(typeof(UnityPerRequestHttpModule));</code>
+            </para>
+            </remarks>
+        </member>
+        <member name="M:Microsoft.Practices.Unity.PerRequestLifetimeManager.GetValue">
+            <summary>
+            Retrieves a value from the backing store associated with this lifetime policy.
+            </summary>
+            <returns>The desired object, or null if no such object is currently stored.</returns>
+        </member>
+        <member name="M:Microsoft.Practices.Unity.PerRequestLifetimeManager.SetValue(System.Object)">
+            <summary>
+            Stores the given value into the backing store for retrieval later.
+            </summary>
+            <param name="newValue">The object being stored.</param>
+        </member>
+        <member name="M:Microsoft.Practices.Unity.PerRequestLifetimeManager.RemoveValue">
+            <summary>
+            Removes the given object from the backing store.
+            </summary>
+        </member>
+        <member name="T:Microsoft.Practices.Unity.Mvc.Properties.Resources">
+            <summary>
+              A strongly-typed resource class, for looking up localized strings, etc.
+            </summary>
+        </member>
+        <member name="P:Microsoft.Practices.Unity.Mvc.Properties.Resources.ResourceManager">
+            <summary>
+              Returns the cached ResourceManager instance used by this class.
+            </summary>
+        </member>
+        <member name="P:Microsoft.Practices.Unity.Mvc.Properties.Resources.Culture">
+            <summary>
+              Overrides the current thread's CurrentUICulture property for all
+              resource lookups using this strongly typed resource class.
+            </summary>
+        </member>
+        <member name="P:Microsoft.Practices.Unity.Mvc.Properties.Resources.ErrorHttpContextNotAvailable">
+            <summary>
+              Looks up a localized string similar to The PerRequestLifetimeManager can only be used in the context of an HTTP request. Possible causes for this error are using the lifetime manager on a non-ASP.NET application, or using it in a thread that is not associated with the appropriate synchronization context..
+            </summary>
+        </member>
+        <member name="T:Microsoft.Practices.Unity.Mvc.UnityDependencyResolver">
+            <summary>
+            An implementation of <see cref="T:System.Web.Mvc.IDependencyResolver"/> that wraps a Unity container.
+            </summary>
+        </member>
+        <member name="M:Microsoft.Practices.Unity.Mvc.UnityDependencyResolver.#ctor(Microsoft.Practices.Unity.IUnityContainer)">
+            <summary>
+            Initializes a new instance of the <see cref="T:Microsoft.Practices.Unity.Mvc.UnityDependencyResolver"/> class for a container.
+            </summary>
+            <param name="container">The <see cref="T:Microsoft.Practices.Unity.IUnityContainer"/> to wrap with the <see cref="T:System.Web.Mvc.IDependencyResolver"/>
+            interface implementation.</param>
+        </member>
+        <member name="M:Microsoft.Practices.Unity.Mvc.UnityDependencyResolver.GetService(System.Type)">
+            <summary>
+            Resolves an instance of the default requested type from the container.
+            </summary>
+            <param name="serviceType">The <see cref="T:System.Type"/> of the object to get from the container.</param>
+            <returns>The requested object.</returns>
+        </member>
+        <member name="M:Microsoft.Practices.Unity.Mvc.UnityDependencyResolver.GetServices(System.Type)">
+            <summary>
+            Resolves multiply registered services.
+            </summary>
+            <param name="serviceType">The type of the requested services.</param>
+            <returns>The requested services.</returns>
+        </member>
+        <member name="T:Microsoft.Practices.Unity.Mvc.UnityFilterAttributeFilterProvider">
+            <summary>
+            Defines a filter provider for filter attributes that support injection of Unity dependencies.
+            </summary>
+        </member>
+        <member name="M:Microsoft.Practices.Unity.Mvc.UnityFilterAttributeFilterProvider.#ctor(Microsoft.Practices.Unity.IUnityContainer)">
+            <summary>
+            Initializes a new instance of the <see cref="T:Microsoft.Practices.Unity.Mvc.UnityFilterAttributeFilterProvider"/> class.
+            </summary>
+            <param name="container">The <see cref="T:Microsoft.Practices.Unity.IUnityContainer"/> that will be used to inject the filters.</param>
+        </member>
+        <member name="M:Microsoft.Practices.Unity.Mvc.UnityFilterAttributeFilterProvider.GetActionAttributes(System.Web.Mvc.ControllerContext,System.Web.Mvc.ActionDescriptor)">
+            <summary>
+            Gets a collection of custom action attributes, and injects them using a Unity container.
+            </summary>
+            <param name="controllerContext">The controller context.</param>
+            <param name="actionDescriptor">The action descriptor.</param>
+            <returns>A collection of custom action attributes.</returns>
+        </member>
+        <member name="M:Microsoft.Practices.Unity.Mvc.UnityFilterAttributeFilterProvider.GetControllerAttributes(System.Web.Mvc.ControllerContext,System.Web.Mvc.ActionDescriptor)">
+            <summary>
+            Gets a collection of controller attributes, and injects them using a Unity container.
+            </summary>
+            <param name="controllerContext">The controller context.</param>
+            <param name="actionDescriptor">The action descriptor.</param>
+            <returns>A collection of controller attributes.</returns>
+        </member>
+        <member name="T:Microsoft.Practices.Unity.Mvc.UnityPerRequestHttpModule">
+            <summary>
+            Implementation of the <see cref="T:System.Web.IHttpModule"/> interface that provides support for using the
+            <see cref="T:Microsoft.Practices.Unity.PerRequestLifetimeManager"/> lifetime manager, and enables it to
+            dispose the instances after the HTTP request ends.
+            </summary>
+        </member>
+        <member name="M:Microsoft.Practices.Unity.Mvc.UnityPerRequestHttpModule.Dispose">
+            <summary>
+            Disposes the resources used by this module.
+            </summary>
+        </member>
+        <member name="M:Microsoft.Practices.Unity.Mvc.UnityPerRequestHttpModule.Init(System.Web.HttpApplication)">
+            <summary>
+            Initializes a module and prepares it to handle requests.
+            </summary>
+            <param name="context">An <see cref="T:System.Web.HttpApplication"/> that provides access to the methods, properties,
+            and events common to all application objects within an ASP.NET application.</param>
+        </member>
+    </members>
+</doc>
diff --git a/Engine/Source/Programs/CrashReporter/CrashReportWebSite/packages/repositories.config b/Engine/Source/Programs/CrashReporter/CrashReportWebSite/packages/repositories.config
new file mode 100644
--- /dev/null
+++ b/Engine/Source/Programs/CrashReporter/CrashReportWebSite/packages/repositories.config
@@ -0,0 +1,4 @@
+﻿<?xml version="1.0" encoding="utf-8"?>
+<repositories>
+  <repository path="..\packages.config" />
+</repositories>
\ No newline at end of file
diff --git a/Engine/Source/Programs/DotNETCommon/DotNETUtilities/LaunchProcess.cs b/Engine/Source/Programs/DotNETCommon/DotNETUtilities/LaunchProcess.cs
--- a/Engine/Source/Programs/DotNETCommon/DotNETUtilities/LaunchProcess.cs
+++ b/Engine/Source/Programs/DotNETCommon/DotNETUtilities/LaunchProcess.cs
@@ -120,14 +120,14 @@ namespace Tools.DotNETCommon.LaunchProcess
 		/// <summary>
 		/// Wait for the launched process to exit, and return its exit code.
 		/// </summary>
-		/// <param name="Timeout">Number of milliseconds to wait for the process to exit. Default is forever.</param>
+		/// <param name="TimeoutMilliseconds">Number of milliseconds to wait for the process to exit. Default is forever.</param>
 		/// <returns>The exit code of the launched process.</returns>
 		/// <remarks>false is returned if the process failed to finish before the requested timeout.</remarks>
-		public EWaitResult WaitForExit( int Timeout = Int32.MaxValue )
+		public EWaitResult WaitForExit( int TimeoutMilliseconds = Int32.MaxValue )
 		{
 			if( LaunchedProcess != null )
 			{
-				LaunchedProcess.WaitForExit( Timeout );
+				LaunchedProcess.WaitForExit(TimeoutMilliseconds);
 				if( !bIsFinished )
 				{
 					// Calling Kill() here seems to be a race condition on the process terminating after WaitForExit (TTP#315685). Catch anything it throws, and make sure the Kill() finishes.
diff --git a/Engine/Source/Programs/UnrealBuildTool/Configuration/UEBuildTarget.cs b/Engine/Source/Programs/UnrealBuildTool/Configuration/UEBuildTarget.cs
--- a/Engine/Source/Programs/UnrealBuildTool/Configuration/UEBuildTarget.cs
+++ b/Engine/Source/Programs/UnrealBuildTool/Configuration/UEBuildTarget.cs
@@ -3609,7 +3609,7 @@ namespace UnrealBuildTool
 				LinkEnvironmentConfiguration RulesLinkEnvConfig = GlobalLinkEnvironment.Config;
 				CPPEnvironmentConfiguration RulesCPPEnvConfig = GlobalCompileEnvironment.Config;
 
-				if (!bUseSharedBuildEnvironment)
+                if (!bUseSharedBuildEnvironment)
 				{
 					Rules.SetupGlobalEnvironment(
 						TargetInfo,
@@ -3624,6 +3624,21 @@ namespace UnrealBuildTool
 				GlobalCompileEnvironment.Config.Definitions.Add(String.Format("IS_PROGRAM={0}", TargetType == TargetRules.TargetType.Program ? "1" : "0"));
 			}
 
+			// See if the target rules defined any signing keys and add the appropriate defines
+			if (!string.IsNullOrEmpty(Rules.PakSigningKeysFile))
+			{
+				string FullFilename = Path.Combine(ProjectDirectory.FullName, Rules.PakSigningKeysFile);
+				if (File.Exists(FullFilename))
+				{
+					string[] Lines = File.ReadAllLines(FullFilename);
+					if (Lines.Length == 3)
+					{
+						GlobalCompileEnvironment.Config.Definitions.Add("DECRYPTION_KEY_MODULUS=\"" + Lines[1] + "\"");
+						GlobalCompileEnvironment.Config.Definitions.Add("DECRYPTION_KEY_EXPONENT=\"" + Lines[2] + "\"");
+					}
+				}
+			}
+
 			// Validate UE configuration - needs to happen before setting any environment mojo and after argument parsing.
 			PlatformContext.ValidateUEBuildConfiguration();
 			UEBuildConfiguration.ValidateConfiguration();
diff --git a/Engine/Source/Programs/UnrealBuildTool/System/RulesCompiler.cs b/Engine/Source/Programs/UnrealBuildTool/System/RulesCompiler.cs
--- a/Engine/Source/Programs/UnrealBuildTool/System/RulesCompiler.cs
+++ b/Engine/Source/Programs/UnrealBuildTool/System/RulesCompiler.cs
@@ -656,6 +656,11 @@ namespace UnrealBuildTool
 		/// </summary>
 		public List<string> AdditionalPlugins = new List<string>();
 
+		/// <summary>
+		/// Path to the set of pak signing keys to embed in the executable
+		/// </summary>
+		public string PakSigningKeysFile = "";
+
 		/// <summary>
 		/// Is the given type a 'game' type (Game/Editor/Server) wrt building?
 		/// </summary>
diff --git a/Engine/Source/Programs/UnrealHeaderTool/Private/HeaderParser.cpp b/Engine/Source/Programs/UnrealHeaderTool/Private/HeaderParser.cpp
--- a/Engine/Source/Programs/UnrealHeaderTool/Private/HeaderParser.cpp
+++ b/Engine/Source/Programs/UnrealHeaderTool/Private/HeaderParser.cpp
@@ -1127,9 +1127,9 @@ void AddModuleRelativePathToMetadata(UField* Type, TMap<FName, FString> &MetaDat
 UEnum* FHeaderParser::CompileEnum()
 {
 	FUnrealSourceFile* CurrentSrcFile = GetCurrentSourceFile();
-	auto Scope = CurrentSrcFile->GetScope();
+	TSharedPtr<FFileScope> Scope = CurrentSrcFile->GetScope();
 
-	CheckAllow( TEXT("'Enum'"), ALLOW_TypeDecl );
+	CheckAllow( TEXT("'Enum'"), ENestAllowFlags::TypeDecl );
 
 	// Get the enum specifier list
 	FToken                     EnumToken;
@@ -1711,7 +1711,7 @@ UScriptStruct* FHeaderParser::CompileStructDeclaration(FClasses& AllClasses)
 	auto Scope = CurrentSrcFile->GetScope();
 
 	// Make sure structs can be declared here.
-	CheckAllow( TEXT("'struct'"), ALLOW_TypeDecl );//@TODO: UCREMOVAL: After the switch: Make this require global scope
+	CheckAllow( TEXT("'struct'"), ENestAllowFlags::TypeDecl );
 
 	FScriptLocation StructDeclaration;
 
@@ -2182,14 +2182,14 @@ const TCHAR *FHeaderParser::NestTypeName( ENestType NestType )
 {
 	switch( NestType )
 	{
-		case NEST_GlobalScope:
+		case ENestType::GlobalScope:
 			return TEXT("Global Scope");
-		case NEST_Class:
+		case ENestType::Class:
 			return TEXT("Class");
-		case NEST_NativeInterface:
-		case NEST_Interface:
+		case ENestType::NativeInterface:
+		case ENestType::Interface:
 			return TEXT("Interface");
-		case NEST_FunctionDeclaration:
+		case ENestType::FunctionDeclaration:
 			return TEXT("Function");
 		default:
 			check(false);
@@ -2198,9 +2198,9 @@ const TCHAR *FHeaderParser::NestTypeName( ENestType NestType )
 }
 
 // Checks to see if a particular kind of command is allowed on this nesting level.
-bool FHeaderParser::IsAllowedInThisNesting(uint32 AllowFlags)
+bool FHeaderParser::IsAllowedInThisNesting(ENestAllowFlags AllowFlags)
 {
-	return ((TopNest->Allow & AllowFlags) != 0);
+	return (TopNest->Allow & AllowFlags) != ENestAllowFlags::None;
 }
 
 //
@@ -2208,11 +2208,11 @@ bool FHeaderParser::IsAllowedInThisNesting(uint32 AllowFlags)
 // If it's not, issues a compiler error referring to the token and the current
 // nesting level.
 //
-void FHeaderParser::CheckAllow( const TCHAR* Thing, uint32 AllowFlags )
+void FHeaderParser::CheckAllow( const TCHAR* Thing, ENestAllowFlags AllowFlags )
 {
 	if (!IsAllowedInThisNesting(AllowFlags))
 	{
-		if (TopNest->NestType == NEST_GlobalScope)
+		if (TopNest->NestType == ENestType::GlobalScope)
 		{
 			FError::Throwf(TEXT("%s is not allowed before the Class definition"), Thing );
 		}
@@ -2240,7 +2240,7 @@ void FHeaderParser::PushNest(ENestType NestType, UStruct* InNode, FUnrealSourceF
 {
 	// Update pointer to top nesting level.
 	TopNest = &Nest[NestLevel++];
-	TopNest->SetScope(NestType == NEST_GlobalScope ? &SourceFile->GetScope().Get() : &FScope::GetTypeScope(InNode).Get());
+	TopNest->SetScope(NestType == ENestType::GlobalScope ? &SourceFile->GetScope().Get() : &FScope::GetTypeScope(InNode).Get());
 	TopNest->NestType = NestType;
 
 	// Prevent overnesting.
@@ -2250,7 +2250,7 @@ void FHeaderParser::PushNest(ENestType NestType, UStruct* InNode, FUnrealSourceF
 	}
 
 	// Inherit info from stack node above us.
-	if (NestLevel > 1 && NestType == NEST_GlobalScope)
+	if (NestLevel > 1 && NestType == ENestType::GlobalScope)
 	{
 		// Use the existing stack node.
 		TopNest->SetScope(TopNest[-1].GetScope());
@@ -2259,21 +2259,21 @@ void FHeaderParser::PushNest(ENestType NestType, UStruct* InNode, FUnrealSourceF
 	// NestType specific logic.
 	switch (NestType)
 	{
-	case NEST_GlobalScope:
-		TopNest->Allow = ALLOW_Class | ALLOW_TypeDecl | ALLOW_Function;
+	case ENestType::GlobalScope:
+		TopNest->Allow = ENestAllowFlags::Class | ENestAllowFlags::TypeDecl | ENestAllowFlags::ImplicitDelegateDecl;
 		break;
 
-	case NEST_Class:
-		TopNest->Allow = ALLOW_VarDecl | ALLOW_Function | ALLOW_TypeDecl;
+	case ENestType::Class:
+		TopNest->Allow = ENestAllowFlags::VarDecl | ENestAllowFlags::Function | ENestAllowFlags::ImplicitDelegateDecl;
 		break;
 
-	case NEST_NativeInterface:
-	case NEST_Interface:
-		TopNest->Allow = ALLOW_Function | ALLOW_TypeDecl;
+	case ENestType::NativeInterface:
+	case ENestType::Interface:
+		TopNest->Allow = ENestAllowFlags::Function;
 		break;
 
-	case NEST_FunctionDeclaration:
-		TopNest->Allow = ALLOW_VarDecl;
+	case ENestType::FunctionDeclaration:
+		TopNest->Allow = ENestAllowFlags::VarDecl;
 
 		break;
 
@@ -2301,19 +2301,19 @@ void FHeaderParser::PopNest(ENestType NestType, const TCHAR* Descr)
 		FError::Throwf(TEXT("Unexpected end of %s in '%s' block"), Descr, NestTypeName(TopNest->NestType));
 	}
 
-	if (NestType != NEST_GlobalScope && NestType != NEST_Class && NestType != NEST_Interface && NestType != NEST_NativeInterface && NestType != NEST_FunctionDeclaration)
+	if (NestType != ENestType::GlobalScope && NestType != ENestType::Class && NestType != ENestType::Interface && NestType != ENestType::NativeInterface && NestType != ENestType::FunctionDeclaration)
 	{
 		FError::Throwf(TEXT("Bad first pass NestType %i"), (uint8)NestType);
 	}
 
 	bool bLinkProps = true;
-	if (NestType == NEST_Class)
+	if (NestType == ENestType::Class)
 	{
 		UClass* TopClass = CastChecked<UClass>(GetCurrentClass());
 		bLinkProps = !TopClass->HasAnyClassFlags(CLASS_Intrinsic);
 	}
 
-	if (NestType != NEST_GlobalScope)
+	if (NestType != ENestType::GlobalScope)
 	{
 		GetCurrentClass()->StaticLink(bLinkProps);
 	}
@@ -4345,16 +4345,16 @@ bool FHeaderParser::CompileDeclaration(FClasses& AllClasses, FToken& Token)
 	EAccessSpecifier AccessSpecifier = ParseAccessProtectionSpecifier(Token);
 	if (AccessSpecifier)
 	{
-		if (!IsAllowedInThisNesting(ALLOW_VarDecl) && !IsAllowedInThisNesting(ALLOW_Function))
+		if (!IsAllowedInThisNesting(ENestAllowFlags::VarDecl) && !IsAllowedInThisNesting(ENestAllowFlags::Function))
 		{
 			FError::Throwf(TEXT("Access specifier %s not allowed here."), Token.Identifier);
 		}
-		check(TopNest->NestType == NEST_Class || TopNest->NestType == NEST_Interface || TopNest->NestType == NEST_NativeInterface);
+		check(TopNest->NestType == ENestType::Class || TopNest->NestType == ENestType::Interface || TopNest->NestType == ENestType::NativeInterface);
 		CurrentAccessSpecifier = AccessSpecifier;
 		return true;
 	}
 
-	if (Token.Matches(TEXT("class")) && (TopNest->NestType == NEST_GlobalScope))
+	if (Token.Matches(TEXT("class")) && (TopNest->NestType == ENestType::GlobalScope))
 	{
 		// Make sure the previous class ended with valid nesting.
 		if (bEncounteredNewStyleClass_UnmatchedBrackets)
@@ -4375,9 +4375,9 @@ bool FHeaderParser::CompileDeclaration(FClasses& AllClasses, FToken& Token)
 		return true;
 	}
 
-	if (Token.Matches(TEXT("GENERATED_IINTERFACE_BODY")) || (Token.Matches(TEXT("GENERATED_BODY")) && TopNest->NestType == NEST_NativeInterface))
+	if (Token.Matches(TEXT("GENERATED_IINTERFACE_BODY")) || (Token.Matches(TEXT("GENERATED_BODY")) && TopNest->NestType == ENestType::NativeInterface))
 	{
-		if (TopNest->NestType != NEST_NativeInterface)
+		if (TopNest->NestType != ENestType::NativeInterface)
 		{
 			FError::Throwf(TEXT("%s must occur inside the native interface definition"), Token.Identifier);
 		}
@@ -4404,9 +4404,9 @@ bool FHeaderParser::CompileDeclaration(FClasses& AllClasses, FToken& Token)
 		return true;
 	}
 
-	if (Token.Matches(TEXT("GENERATED_UINTERFACE_BODY")) || (Token.Matches(TEXT("GENERATED_BODY")) && TopNest->NestType == NEST_Interface))
+	if (Token.Matches(TEXT("GENERATED_UINTERFACE_BODY")) || (Token.Matches(TEXT("GENERATED_BODY")) && TopNest->NestType == ENestType::Interface))
 	{
-		if (TopNest->NestType != NEST_Interface)
+		if (TopNest->NestType != ENestType::Interface)
 		{
 			FError::Throwf(TEXT("%s must occur inside the interface definition"), Token.Identifier);
 		}
@@ -4428,9 +4428,9 @@ bool FHeaderParser::CompileDeclaration(FClasses& AllClasses, FToken& Token)
 		return true;
 	}
 
-	if (Token.Matches(TEXT("GENERATED_UCLASS_BODY")) || (Token.Matches(TEXT("GENERATED_BODY")) && TopNest->NestType == NEST_Class))
+	if (Token.Matches(TEXT("GENERATED_UCLASS_BODY")) || (Token.Matches(TEXT("GENERATED_BODY")) && TopNest->NestType == ENestType::Class))
 	{
-		if (TopNest->NestType != NEST_Class)
+		if (TopNest->NestType != ENestType::Class)
 		{
 			FError::Throwf(TEXT("%s must occur inside the class definition"), Token.Identifier);
 		}
@@ -4463,7 +4463,7 @@ bool FHeaderParser::CompileDeclaration(FClasses& AllClasses, FToken& Token)
 		return true;
 	}
 
-	if (Token.Matches(TEXT("UCLASS"), ESearchCase::CaseSensitive) && (TopNest->Allow & ALLOW_Class))
+	if (Token.Matches(TEXT("UCLASS"), ESearchCase::CaseSensitive))
 	{
 		bHaveSeenUClass = true;
 		bEncounteredNewStyleClass_UnmatchedBrackets = true;
@@ -4471,7 +4471,7 @@ bool FHeaderParser::CompileDeclaration(FClasses& AllClasses, FToken& Token)
 		return true;
 	}
 
-	if (Token.Matches(TEXT("UINTERFACE")) && (TopNest->Allow & ALLOW_Class))
+	if (Token.Matches(TEXT("UINTERFACE")))
 	{
 		bHaveSeenUClass = true;
 		bEncounteredNewStyleClass_UnmatchedBrackets = true;
@@ -4499,8 +4499,8 @@ bool FHeaderParser::CompileDeclaration(FClasses& AllClasses, FToken& Token)
 
 	if (Token.Matches(TEXT("UPROPERTY"), ESearchCase::CaseSensitive))
 	{
-		CheckAllow(TEXT("'Member variable declaration'"), ALLOW_VarDecl);
-		check(TopNest->NestType == NEST_Class);
+		CheckAllow(TEXT("'Member variable declaration'"), ENestAllowFlags::VarDecl);
+		check(TopNest->NestType == ENestType::Class);
 
 		CompileVariableDeclaration(AllClasses, GetCurrentClass());
 		return true;
@@ -4545,7 +4545,7 @@ bool FHeaderParser::CompileDeclaration(FClasses& AllClasses, FToken& Token)
 		// Pop nesting here to allow other non UClass declarations in the header file.
 		if (CurrentClass->ClassFlags & CLASS_Interface)
 		{
-			checkf(TopNest->NestType == NEST_Interface || TopNest->NestType == NEST_NativeInterface, TEXT("Unexpected end of interface block."));
+			checkf(TopNest->NestType == ENestType::Interface || TopNest->NestType == ENestType::NativeInterface, TEXT("Unexpected end of interface block."));
 			PopNest(TopNest->NestType, TEXT("'Interface'"));
 			PostPopNestInterface(AllClasses, CurrentClass);
 
@@ -4563,7 +4563,7 @@ bool FHeaderParser::CompileDeclaration(FClasses& AllClasses, FToken& Token)
 		}
 		else
 		{
-			PopNest(NEST_Class, TEXT("'Class'"));
+			PopNest(ENestType::Class, TEXT("'Class'"));
 			PostPopNestClass(CurrentClass);
 
 			// Ensure classes have a GENERATED_BODY declaration
@@ -4920,7 +4920,7 @@ void PostParsingClassSetup(UClass* Class)
 void FHeaderParser::CompileClassDeclaration(FClasses& AllClasses)
 {
 	// Start of a class block.
-	CheckAllow(TEXT("'class'"), ALLOW_Class);
+	CheckAllow(TEXT("'class'"), ENestAllowFlags::Class);
 
 	// New-style UCLASS() syntax
 	TMap<FName, FString> MetaData;
@@ -4954,7 +4954,7 @@ void FHeaderParser::CompileClassDeclaration(FClasses& AllClasses)
 
 	Class->ClassFlags |= CLASS_Parsed;
 
-	PushNest(ENestType::NEST_Class, Class);
+	PushNest(ENestType::Class, Class);
 	
 	const uint32 PrevClassFlags = Class->ClassFlags;
 	ResetClassData();
@@ -5114,7 +5114,7 @@ bool FHeaderParser::TryParseIInterfaceClass(FClasses& AllClasses)
 	RequireSymbol(TEXT("{"), *ErrorMsg);
 
 	// Push the interface class nesting again.
-	PushNest(NEST_NativeInterface, FoundClass);
+	PushNest(ENestType::NativeInterface, FoundClass);
 
 	return true;
 }
@@ -5127,7 +5127,7 @@ void FHeaderParser::CompileInterfaceDeclaration(FClasses& AllClasses)
 	FUnrealSourceFile* CurrentSrcFile = GetCurrentSourceFile();
 	// Start of an interface block. Since Interfaces and Classes are always at the same nesting level,
 	// whereever a class declaration is allowed, an interface declaration is also allowed.
-	CheckAllow( TEXT("'interface'"), ALLOW_Class );
+	CheckAllow( TEXT("'interface'"), ENestAllowFlags::Class );
 
 	FString DeclaredInterfaceName;
 	FString RequiredAPIMacroIfPresent;
@@ -5226,8 +5226,8 @@ void FHeaderParser::CompileInterfaceDeclaration(FClasses& AllClasses)
 	check(InterfaceClass->HasAnyFlags(RF_Standalone));
 
 	// Push the interface class nesting.
-	// we need a more specific set of allow flags for NEST_Interface, only function declaration is allowed, no other stuff are allowed
-	PushNest(NEST_Interface, InterfaceClass);
+	// we need a more specific set of allow flags for ENestType::Interface, only function declaration is allowed, no other stuff are allowed
+	PushNest(ENestType::Interface, InterfaceClass);
 }
 
 // Returns true if the token is a dynamic delegate declaration
@@ -5364,6 +5364,8 @@ void FHeaderParser::ParseParameterList(FClasses& AllClasses, UFunction* Function
 }
 void FHeaderParser::CompileDelegateDeclaration(FClasses& AllClasses, const TCHAR* DelegateIdentifier, EDelegateSpecifierAction::Type SpecifierAction)
 {
+	const TCHAR* CurrentScopeName = TEXT("Delegate Declaration");
+
 	FUnrealSourceFile* CurrentSrcFile = GetCurrentSourceFile();
 	TMap<FName, FString> MetaData;
 	AddModuleRelativePathToMetadata(*CurrentSrcFile, MetaData);
@@ -5386,16 +5388,14 @@ void FHeaderParser::CompileDelegateDeclaration(FClasses& AllClasses, const TCHAR
 			FError::Throwf(TEXT("Unexpected token following UDELEGATE(): %s"), Token.Identifier);
 
 		DelegateMacro = Token.Identifier;
+		CheckAllow(CurrentScopeName, ENestAllowFlags::TypeDecl);
 	}
 	else
 	{
 		DelegateMacro = DelegateIdentifier;
+		CheckAllow(CurrentScopeName, ENestAllowFlags::ImplicitDelegateDecl);
 	}
 
-	// Make sure we can have a delegate declaration here
-	const TCHAR* CurrentScopeName = TEXT("Delegate Declaration");
-	CheckAllow(CurrentScopeName, ALLOW_TypeDecl);//@TODO: UCREMOVAL: After the switch: Make this require global scope
-
 	// Break the delegate declaration macro down into parts
 	const bool bHasReturnValue = DelegateMacro.Contains(TEXT("_RetVal"));
 	const bool bDeclaredConst  = DelegateMacro.Contains(TEXT("_Const"));
@@ -5562,7 +5562,7 @@ void FHeaderParser::CompileDelegateDeclaration(FClasses& AllClasses, const TCHAR
  */
 void FHeaderParser::CompileFunctionDeclaration(FClasses& AllClasses)
 {
-	CheckAllow(TEXT("'Function'"), ALLOW_Function);
+	CheckAllow(TEXT("'Function'"), ENestAllowFlags::Function);
 
 	FUnrealSourceFile* CurrentSrcFile = GetCurrentSourceFile();
 	TMap<FName, FString> MetaData;
@@ -6390,7 +6390,7 @@ void FHeaderParser::CompileVariableDeclaration(FClasses& AllClasses, UStruct* St
 		// Deprecation validation
 		ValidatePropertyIsDeprecatedIfNecessary(Property, NULL);
 
-		if (TopNest->NestType != NEST_FunctionDeclaration)
+		if (TopNest->NestType != ENestType::FunctionDeclaration)
 		{
 			if (NewProperties.Num())
 			{
@@ -6640,7 +6640,7 @@ ECompilationResult::Type FHeaderParser::ParseHeader(FClasses& AllClasses, FUnrea
 	// Init nesting.
 	NestLevel = 0;
 	TopNest = NULL;
-	PushNest(NEST_GlobalScope, nullptr, CurrentSrcFile);
+	PushNest(ENestType::GlobalScope, nullptr, CurrentSrcFile);
 
 	// C++ classes default to private access level
 	CurrentAccessSpecifier = ACCESS_Private; 
@@ -6665,7 +6665,7 @@ ECompilationResult::Type FHeaderParser::ParseHeader(FClasses& AllClasses, FUnrea
 			StatementsParsed++;
 		}
 
-		PopNest(NEST_GlobalScope, TEXT("Global scope"));
+		PopNest(ENestType::GlobalScope, TEXT("Global scope"));
 
 		auto ScopeTypeIterator = CurrentSrcFile->GetScope()->GetTypeIterator();
 		while (ScopeTypeIterator.MoveNext())
@@ -6747,7 +6747,7 @@ ECompilationResult::Type FHeaderParser::ParseHeader(FClasses& AllClasses, FUnrea
 		if (NestLevel == 0)
 		{
 			// Pushing nest so there is a file context for this error.
-			PushNest(NEST_GlobalScope, nullptr, CurrentSrcFile);
+			PushNest(ENestType::GlobalScope, nullptr, CurrentSrcFile);
 		}
 
 		// Handle compiler error.
diff --git a/Engine/Source/Programs/UnrealHeaderTool/Private/HeaderParser.h b/Engine/Source/Programs/UnrealHeaderTool/Private/HeaderParser.h
--- a/Engine/Source/Programs/UnrealHeaderTool/Private/HeaderParser.h
+++ b/Engine/Source/Programs/UnrealHeaderTool/Private/HeaderParser.h
@@ -25,25 +25,29 @@ extern double GHeaderCodeGenTime;
 enum {MAX_NEST_LEVELS = 16};
 
 /* Code nesting types. */
-enum ENestType
+enum class ENestType
 {
-	NEST_GlobalScope,
-	NEST_Class,
-	NEST_FunctionDeclaration,
-	NEST_Interface,
-	NEST_NativeInterface
+	GlobalScope,
+	Class,
+	FunctionDeclaration,
+	Interface,
+	NativeInterface
 };
 
 /** Types of statements to allow within a particular nesting block. */
-enum ENestAllowFlags
+enum class ENestAllowFlags
 {
-	ALLOW_Function			= 1,	// Allow Event declarations at this level.
-	ALLOW_VarDecl			= 2,	// Allow variable declarations at this level.
-	ALLOW_Class				= 4,	// Allow class definition heading.
-	ALLOW_Return			= 8,	// Allow 'return' within a function.
-	ALLOW_TypeDecl			= 16,	// Allow declarations which do not affect memory layout, such as structs, enums, and consts
+	None                 = 0,
+	Function             = 1,  // Allow Event declarations at this level.
+	VarDecl              = 2,  // Allow variable declarations at this level.
+	Class                = 4,  // Allow class definition heading.
+	Return               = 8,  // Allow 'return' within a function.
+	TypeDecl             = 16, // Allow declarations which do not affect memory layout, such as structs, enums, and consts, but not implicit delegates
+	ImplicitDelegateDecl = 32, // Allow implicit delegates (i.e. those not decorated with UDELEGATE) to be declared
 };
 
+ENUM_CLASS_FLAGS(ENestAllowFlags)
+
 namespace EDelegateSpecifierAction
 {
 	enum Type
@@ -92,7 +96,7 @@ public:
 	ENestType NestType;
 
 	/** Types of statements to allow at this nesting level. */
-	int32 Allow;
+	ENestAllowFlags Allow;
 };
 
 struct FIndexRange
@@ -243,7 +247,7 @@ protected:
 	FFileScope* GetCurrentFileScope() const
 	{
 		int32 Index = 0;
-		while (TopNest[Index].NestType != NEST_GlobalScope)
+		while (TopNest[Index].NestType != ENestType::GlobalScope)
 		{
 			--Index;
 		}
@@ -269,7 +273,7 @@ protected:
 	 */
 	FStructScope* GetCurrentClassScope() const
 	{
-		check(TopNest->NestType == NEST_Class || TopNest->NestType == NEST_Interface || TopNest->NestType == NEST_NativeInterface);
+		check(TopNest->NestType == ENestType::Class || TopNest->NestType == ENestType::Interface || TopNest->NestType == ENestType::NativeInterface);
 
 		return (FStructScope*)TopNest->GetScope();
 	}
@@ -280,9 +284,9 @@ protected:
 	bool IsInAClass() const
 	{
 		int32 Index = 0;
-		while (TopNest[Index].NestType != NEST_GlobalScope)
+		while (TopNest[Index].NestType != ENestType::GlobalScope)
 		{
-			if (TopNest[Index].NestType == NEST_Class || TopNest->NestType == NEST_Interface || TopNest->NestType == NEST_NativeInterface)
+			if (TopNest[Index].NestType == ENestType::Class || TopNest->NestType == ENestType::Interface || TopNest->NestType == ENestType::NativeInterface)
 			{
 				return true;
 			}
@@ -521,12 +525,12 @@ protected:
 	bool CompileStatement(FClasses& AllClasses);
 
 	// Checks to see if a particular kind of command is allowed on this nesting level.
-	bool IsAllowedInThisNesting(uint32 AllowFlags);
+	bool IsAllowedInThisNesting(ENestAllowFlags AllowFlags);
 	
 	// Make sure that a particular kind of command is allowed on this nesting level.
 	// If it's not, issues a compiler error referring to the token and the current
 	// nesting level.
-	void CheckAllow(const TCHAR* Thing, uint32 AllowFlags);
+	void CheckAllow(const TCHAR* Thing, ENestAllowFlags AllowFlags);
 
 	UStruct* GetSuperScope( UStruct* CurrentScope, const FName& SearchName );
 
diff --git a/Engine/Source/Programs/UnrealHeaderTool/Resources/UHTDebugging/TestObject.h b/Engine/Source/Programs/UnrealHeaderTool/Resources/UHTDebugging/TestObject.h
--- a/Engine/Source/Programs/UnrealHeaderTool/Resources/UHTDebugging/TestObject.h
+++ b/Engine/Source/Programs/UnrealHeaderTool/Resources/UHTDebugging/TestObject.h
@@ -11,6 +11,7 @@ UCLASS()
 class UTestObject : public UObject
 {
 	GENERATED_BODY()
+
 public:
 	UTestObject(const FObjectInitializer& ObjectInitializer = FObjectInitializer::Get());
 
@@ -24,7 +25,7 @@ public:
 	ECppEnum EnumProperty;
 
 	UPROPERTY()
-	TMap<int32, FString> TestMap;
+	TMap<int32, bool> TestMap;
 
 	UFUNCTION()
 	void CodeGenTestForEnumClasses(ECppEnum Val);
diff --git a/Engine/Source/Programs/UnrealPak/Private/KeyGenerator.cpp b/Engine/Source/Programs/UnrealPak/Private/KeyGenerator.cpp
--- a/Engine/Source/Programs/UnrealPak/Private/KeyGenerator.cpp
+++ b/Engine/Source/Programs/UnrealPak/Private/KeyGenerator.cpp
@@ -11,10 +11,10 @@
 // Global constants
 namespace
 {
-	const int256 One(1);
-	const int256 Two(2);
-	const int256 IterationStep(1000);
-	TArray<int256> PrimeLookupTable;
+	const TEncryptionInt One(1);
+	const TEncryptionInt Two(2);
+	const TEncryptionInt IterationStep(1000);
+	TArray<TEncryptionInt> PrimeLookupTable;
 }
 
 /**
@@ -25,17 +25,17 @@ class FPrimeCheckRunnable : public FRunnable
 	/** Flag indicating if a factor has been found. Shared across multiple threads. */
 	FThreadSafeCounter&	FoundFactor;
 	/** Candidate for a prime number */
-	int256 PotentialPrime;
+	TEncryptionInt PotentialPrime;
 	/** Start of a range to check for factors */
-	int256 InitialValue;
+	TEncryptionInt InitialValue;
 	/** End of a range to check for factors */
-	int256 MaxValue;
+	TEncryptionInt MaxValue;
 	/** This thread */
 	FRunnableThread* Thread;
 
 public:
 
-	FPrimeCheckRunnable(FThreadSafeCounter& InFoundFactor, int256 Candidate, int256 InInitialValue, int256 InMaxValue)
+	FPrimeCheckRunnable(FThreadSafeCounter& InFoundFactor, TEncryptionInt Candidate, TEncryptionInt InInitialValue, TEncryptionInt InMaxValue)
 		: FoundFactor(InFoundFactor)
 		, PotentialPrime(Candidate)
 		, InitialValue(InInitialValue)
@@ -61,11 +61,11 @@ public:
 	virtual bool Init() override { return true; }
 	virtual uint32 Run() override
 	{
-		int256 Remainder;
+		TEncryptionInt Remainder;
 		int32 FactorCheckTimer = 0;
-		for (int256 Factor = InitialValue; InitialValue <= MaxValue; Factor += Two)
+		for (TEncryptionInt Factor = InitialValue; InitialValue <= MaxValue; Factor += Two)
 		{
-			int256 Dividend(PotentialPrime);
+			TEncryptionInt Dividend(PotentialPrime);
 			Dividend.DivideWithRemainder(Factor, Remainder);
 			if (Remainder.IsZero())
 			{
@@ -93,7 +93,7 @@ public:
 /**
  * Checks if the value is a prime number.
  */
-bool IsPrime(const int256& InValue, bool bUseTasks)
+bool IsPrime(const TEncryptionInt& InValue, bool bUseTasks)
 {
 	// 2 is but we don't care about small numbers here.
 	if ((InValue & One) == 0)
@@ -101,13 +101,13 @@ bool IsPrime(const int256& InValue, bool bUseTasks)
 		return false;
 	}
 
-	int256 Remainder;
+	TEncryptionInt Remainder;
 
 	// Check against known prime factors
 	int32 Index;
 	for (Index = 0; Index < PrimeLookupTable.Num() && PrimeLookupTable[Index] < InValue; ++Index)
 	{
-		int256 Dividend(InValue);
+		TEncryptionInt Dividend(InValue);
 		Dividend.DivideWithRemainder(PrimeLookupTable[Index], Remainder);
 		if (Remainder.IsZero())
 		{
@@ -121,9 +121,9 @@ bool IsPrime(const int256& InValue, bool bUseTasks)
 	}
 
 	// Brute force, check all odd numbers > MaxKnownPrime < sqrt(Number)
-	int256 MaxFactorValue(InValue);
+	TEncryptionInt MaxFactorValue(InValue);
 	MaxFactorValue.Sqrt();
-	int256 Factor(PrimeLookupTable[PrimeLookupTable.Num() - 1] + Two); 
+	TEncryptionInt Factor(PrimeLookupTable[PrimeLookupTable.Num() - 1] + Two); 
 	if (bUseTasks)
 	{
 		// Mutithreaded path. Split the range we look for factors over multiple threads. If one thread finds a factor
@@ -136,13 +136,13 @@ bool IsPrime(const int256& InValue, bool bUseTasks)
 		Tasks.Reserve(TaskCount);
 		FThreadSafeCounter FoundFactors(0);
 
-		int256 Range(MaxFactorValue - Factor);
+		TEncryptionInt Range(MaxFactorValue - Factor);
 		Range /= TaskCount;
 
 		// Spawn threads
 		for (int32 TaskIndex = 0; TaskIndex < TaskCount; ++TaskIndex)
 		{
-			int256 MaxValue(Factor + Range);
+			TEncryptionInt MaxValue(Factor + Range);
 			Tasks.Add(new FPrimeCheckRunnable(FoundFactors, InValue, Factor, MaxValue)); 
 			Factor = MaxValue;
 			if ((Factor & One) == 0)
@@ -172,7 +172,7 @@ bool IsPrime(const int256& InValue, bool bUseTasks)
 		// Single threaded path (used for generating prime table)
 		while (Factor < MaxFactorValue)
 		{
-			int256 Dividend(InValue);
+			TEncryptionInt Dividend(InValue);
 			Dividend.DivideWithRemainder(Factor, Remainder);
 			if (Remainder.IsZero())
 			{
@@ -187,7 +187,7 @@ bool IsPrime(const int256& InValue, bool bUseTasks)
 /**
  * Generate two random prime numbers
  */
-void GeneratePrimeNumbers(int256& P, int256& Q)
+void GeneratePrimeNumbers(TEncryptionInt& P, TEncryptionInt& Q)
 {
 	// Generate a random odd number
 	FRandomStream Rand((int32)(FDateTime::Now().GetTicks() % (int64)MAX_int32));
@@ -199,10 +199,10 @@ void GeneratePrimeNumbers(int256& P, int256& Q)
 		0, //(uint32)Rand.RandRange(0, MAX_int32 - 1) | 0xa0ff0000, 
 		0, 0, 0, 0 
 	};
-	int256 InitialValue(RandBits);
+	TEncryptionInt InitialValue(RandBits);
 
 	// We need two primes
-	TArray<int256> DiscoveredPrimes;
+	TArray<TEncryptionInt> DiscoveredPrimes;
 
 	int64 IterationCounter = 0;
 	double TimeAccumulator = 0.0;
@@ -218,7 +218,7 @@ void GeneratePrimeNumbers(int256& P, int256& Q)
 
 		const double IterationStartTime = FPlatformTime::Seconds();
 
-		int256 MinValue(InitialValue - IterationStep);
+		TEncryptionInt MinValue(InitialValue - IterationStep);
 		while (InitialValue >= MinValue && DiscoveredPrimes.Num() < 2)
 		{
 			if (IsPrime(InitialValue, true))
@@ -256,7 +256,7 @@ void GeneratePrimeNumberTable(int64 MaxValue, const TCHAR* Filename)
 
 	UE_LOG(LogPakFile, Display, TEXT("Generating prime number table <= %lld: %s."), MaxValue, Filename);
 
-	FString PrimeTableString(TEXT("// Copyright 1998-2016 Epic Games, Inc. All Rights Reserved.\nint256 PrimeTable[] = \n{\n\t2, "));
+	FString PrimeTableString(TEXT("// Copyright 1998-2016 Epic Games, Inc. All Rights Reserved.\nTEncryptionInt PrimeTable[] = \n{\n\t2, "));
 	int64 PrimeCount = 1;
 	const double StartTime = FPlatformTime::Seconds();
 	for (int64 SmallNumber = 3; SmallNumber <= MaxValue; SmallNumber += 2)
@@ -284,17 +284,17 @@ void GeneratePrimeNumberTable(int64 MaxValue, const TCHAR* Filename)
 class FPrimeFinderRunnable : public FRunnable
 {
 	/** Range start */
-	int256 MinValue;
+	TEncryptionInt MinValue;
 	/** Range end */
-	int256 MaxValue;
+	TEncryptionInt MaxValue;
 	/** This thread */
 	FRunnableThread* Thread;
 	/** All primes found in the given range */
-	TArray<int256> FoundPrimes;
+	TArray<TEncryptionInt> FoundPrimes;
 
 public:
 
-	FPrimeFinderRunnable(int256 InMinValue, int256 InMaxValue)
+	FPrimeFinderRunnable(TEncryptionInt InMinValue, TEncryptionInt InMaxValue)
 		: MinValue(InMinValue)
 		, MaxValue(InMaxValue)
 	{
@@ -314,7 +314,7 @@ public:
 		return Thread;
 	}
 
-	const TArray<int256>& GetFoundPrimes() const
+	const TArray<TEncryptionInt>& GetFoundPrimes() const
 	{
 		return FoundPrimes;
 	}
@@ -323,9 +323,9 @@ public:
 	virtual bool Init() override { return true; }
 	virtual uint32 Run() override
 	{
-		int256 Remainder;
+		TEncryptionInt Remainder;
 		int32 FactorCheckTimer = 0;
-		for (int256 Candidate = MinValue; Candidate <= MaxValue; Candidate += Two)
+		for (TEncryptionInt Candidate = MinValue; Candidate <= MaxValue; Candidate += Two)
 		{
 			if (IsPrime(Candidate, false))
 			{
@@ -352,20 +352,20 @@ void GeneratePrimeNumberLookupTable()
 
 	FillPrimeLookupTableWithPrecompiledNumbers();
 
-	int256 MinPrimeValue(PrimeLookupTable[PrimeLookupTable.Num() - 1] + Two);
-	int256 MaxPrimeValue(MinPrimeValue);
+	TEncryptionInt MinPrimeValue(PrimeLookupTable[PrimeLookupTable.Num() - 1] + Two);
+	TEncryptionInt MaxPrimeValue(MinPrimeValue);
 	MaxPrimeValue *= 100;
 
 	const int32 TaskCount = FPlatformMisc::NumberOfCoresIncludingHyperthreads();
 	TArray<FPrimeFinderRunnable*> Tasks;
 	Tasks.Reserve(TaskCount);
 
-	int256 Range(MaxPrimeValue - MinPrimeValue);
+	TEncryptionInt Range(MaxPrimeValue - MinPrimeValue);
 	Range /= TaskCount;
 
 	for (int32 TaskIndex = 0; TaskIndex < TaskCount; ++TaskIndex)
 	{
-		int256 MaxValue(MinPrimeValue + Range);
+		TEncryptionInt MaxValue(MinPrimeValue + Range);
 		Tasks.Add(new FPrimeFinderRunnable(MinPrimeValue, MaxValue)); 
 		MinPrimeValue = MaxValue;
 		if ((MinPrimeValue & One) == 0)
@@ -374,7 +374,7 @@ void GeneratePrimeNumberLookupTable()
 		}
 	}
 
-	TArray<int256> NewPrimes;
+	TArray<TEncryptionInt> NewPrimes;
 	for (int32 TaskIndex = 0; TaskIndex < Tasks.Num(); ++TaskIndex)
 	{
 		Tasks[TaskIndex]->GetThread()->WaitForCompletion();
@@ -392,27 +392,32 @@ bool GenerateKeys(const TCHAR* KeyFilename)
 
 	GeneratePrimeNumberLookupTable();
 
-	int256 P;
-	int256 Q;
+	TEncryptionInt P;
+	TEncryptionInt Q;
 
-	int64 CmdLineP = 0;
-	int64 CmdLineQ = 0;
+	FString CmdLineP;
+	FString CmdLineQ;
 	FParse::Value(FCommandLine::Get(), TEXT("P="), CmdLineP);
 	FParse::Value(FCommandLine::Get(), TEXT("Q="), CmdLineQ);
 
+	const bool bNoVerifyPrimes = FParse::Param(FCommandLine::Get(), TEXT("NoVerifyPrimes"));
+
+	P.Parse(CmdLineP);
+	Q.Parse(CmdLineQ);
+
 	// Check if we have valid primes in the command line.
 	// @todo: IsPrime check should probably go when we start to use big primes
-	bool bGeneratePrimes = !(CmdLineP > 2 && CmdLineQ > 2);
-	if (!bGeneratePrimes)
+	bool bGeneratePrimes = !(P > Two && Q > Two);
+	if (!bGeneratePrimes && !bNoVerifyPrimes)
 	{
-		if (!IsPrime(CmdLineP, false))
+		if (!IsPrime(P, false))
 		{
-			UE_LOG(LogPakFile, Warning, TEXT("P=%lld is not prime!"), CmdLineP);
+			UE_LOG(LogPakFile, Warning, TEXT("P=%s is not prime!"), *CmdLineP);
 			bGeneratePrimes = true;
 		}
-		if (!IsPrime(CmdLineQ, false))
+		if (!IsPrime(Q, false))
 		{
-			UE_LOG(LogPakFile, Warning, TEXT("Q=%lld is not prime!"), CmdLineQ);
+			UE_LOG(LogPakFile, Warning, TEXT("Q=%s is not prime!"), *CmdLineQ);
 			bGeneratePrimes = true;
 		}
 	}
@@ -427,8 +432,6 @@ bool GenerateKeys(const TCHAR* KeyFilename)
 	{
 		// Use predefined primes
 		UE_LOG(LogPakFile, Display, TEXT("Using predefined values to generate keys."));
-		P = CmdLineP;
-		Q = CmdLineQ;
 	}
 
 	// Generate key pair
@@ -436,7 +439,14 @@ bool GenerateKeys(const TCHAR* KeyFilename)
 	FKeyPair Keys;
 	FEncryption::GenerateKeyPair(P, Q, Keys.PublicKey, Keys.PrivateKey);
 
-	return SaveKeysToFile(Keys, KeyFilename);
+	if (TestKeys(Keys))
+	{
+		return SaveKeysToFile(Keys, KeyFilename);
+	}
+	else
+	{
+		return false;
+	}
 }
 
 bool SaveKeysToFile(const FKeyPair& Keys, const TCHAR* KeyFilename)
@@ -475,3 +485,41 @@ bool ReadKeysFromFile(const TCHAR* KeyFilename, FKeyPair& OutKeys)
 
 	return bResult;
 }
+
+
+bool TestKeys(FKeyPair& Pair)
+{
+	UE_LOG(LogPakFile, Display, TEXT("Testing signature keys."));
+
+	// Just some random values
+	static TEncryptionInt TestData[] =
+	{
+		11,
+		253,
+		128,
+		234,
+		56,
+		89,
+		34,
+		179,
+		29,
+		1024,
+		(int64)(MAX_int32),
+		(int64)(MAX_uint32) - 1
+	};
+
+	for (int32 TestIndex = 0; TestIndex < ARRAY_COUNT(TestData); ++TestIndex)
+	{
+		TEncryptionInt EncryptedData = FEncryption::ModularPow(TestData[TestIndex], Pair.PrivateKey.Exponent, Pair.PrivateKey.Modulus);
+		TEncryptionInt DecryptedData = FEncryption::ModularPow(EncryptedData, Pair.PublicKey.Exponent, Pair.PublicKey.Modulus);
+		if (TestData[TestIndex] != DecryptedData)
+		{
+			UE_LOG(LogPakFile, Error, TEXT("Keys do not properly encrypt/decrypt data (failed test with %lld)"), TestData[TestIndex].ToInt());
+			return false;
+		}
+	}
+
+	UE_LOG(LogPakFile, Display, TEXT("Signature keys check completed successfuly."));
+
+	return true;
+}
\ No newline at end of file
diff --git a/Engine/Source/Programs/UnrealPak/Private/KeyGenerator.h b/Engine/Source/Programs/UnrealPak/Private/KeyGenerator.h
--- a/Engine/Source/Programs/UnrealPak/Private/KeyGenerator.h
+++ b/Engine/Source/Programs/UnrealPak/Private/KeyGenerator.h
@@ -32,4 +32,6 @@ bool GenerateKeys(const TCHAR* KeyFilename);
 
 bool SaveKeysToFile(const FKeyPair& Keys, const TCHAR* KeyFilename);
 
-bool ReadKeysFromFile(const TCHAR* KeyFilename, FKeyPair& OutKeys);
\ No newline at end of file
+bool ReadKeysFromFile(const TCHAR* KeyFilename, FKeyPair& OutKeys);
+
+bool TestKeys(FKeyPair& Pair);
\ No newline at end of file
diff --git a/Engine/Source/Programs/UnrealPak/Private/Primes.inl b/Engine/Source/Programs/UnrealPak/Private/Primes.inl
--- a/Engine/Source/Programs/UnrealPak/Private/Primes.inl
+++ b/Engine/Source/Programs/UnrealPak/Private/Primes.inl
@@ -1,5 +1,5 @@
 // Copyright 1998-2016 Epic Games, Inc. All Rights Reserved.
-int256 PrimeTable[] = 
+TEncryptionInt PrimeTable[] =
 {
 	2, 3, 5, 7, 11, 13, 17, 19, 23, 29, 
 	31, 37, 41, 43, 47, 53, 59, 61, 67, 71, 
diff --git a/Engine/Source/Programs/UnrealPak/Private/SignedArchiveWriter.cpp b/Engine/Source/Programs/UnrealPak/Private/SignedArchiveWriter.cpp
--- a/Engine/Source/Programs/UnrealPak/Private/SignedArchiveWriter.cpp
+++ b/Engine/Source/Programs/UnrealPak/Private/SignedArchiveWriter.cpp
@@ -6,9 +6,10 @@
 #include "BigInt.h"
 #include "SignedArchiveWriter.h"
 
-FSignedArchiveWriter::FSignedArchiveWriter(FArchive& Pak, const FEncryptionKey& InPublicKey, const FEncryptionKey& InPrivateKey)
-	: BufferArchive(Buffer)
-	, PakWriter(Pak)
+FSignedArchiveWriter::FSignedArchiveWriter(FArchive& InPak, const FString& InPakFilename, const FEncryptionKey& InPublicKey, const FEncryptionKey& InPrivateKey)
+: BufferArchive(Buffer)
+	, PakWriter(InPak)
+	, PakSignaturesFilename(FPaths::ChangeExtension(InPakFilename, TEXT("sig")))
 	, SizeOnDisk(0)
 	, PakSize(0)
 	, PublicKey(InPublicKey)
@@ -26,25 +27,17 @@ FSignedArchiveWriter::~FSignedArchiveWriter()
 	delete &PakWriter;
 }
 
-void FSignedArchiveWriter::Encrypt(int256* EncryptedData, const uint8* Data, const int64 DataLength, const FEncryptionKey EncryptionKey)
-{
-	for (int Index = 0; Index < DataLength; ++Index)
-	{
-		EncryptedData[Index] = FEncryption::ModularPow(int256(Data[Index]), EncryptionKey.Exponent, EncryptionKey.Modulus);
-	}
-}
-
 void FSignedArchiveWriter::SerializeBufferAndSign()
 {
 	// Hash the buffer
-	uint8 Hash[GPakFileChunkHashSize];
+	FDecryptedSignature SourceSignature;
 	
 	// Compute a hash for this buffer data
-	ComputePakChunkHash(&Buffer[0], Buffer.Num(), Hash);
+	SourceSignature.Data = FCrc::MemCrc32(&Buffer[0], Buffer.Num());
 
 	// Encrypt the signature
-	FEncryptedSignature<GPakFileChunkHashSize> Signature;
-	Encrypt(Signature.Data, Hash, sizeof(Hash), PrivateKey);
+	FEncryptedSignature EncryptedSignature;
+	FEncryption::EncryptSignature(SourceSignature, EncryptedSignature, PrivateKey);
 
 	// Flush the buffer
 	PakWriter.Serialize(&Buffer[0], Buffer.Num());
@@ -52,7 +45,13 @@ void FSignedArchiveWriter::SerializeBufferAndSign()
 	Buffer.Empty(FPakInfo::MaxChunkDataSize);
 
 	// Collect the signature so we can write it out at the end
-	ChunkSignatures.Add(Signature);
+	ChunkSignatures.Add(EncryptedSignature);
+
+#if 0
+	FDecryptedSignature<GPakFileChunkHashSize> TestSignature;
+	FEncryption::DecryptBytes(TestSignature.Data, Signature.Data, GPakFileChunkHashSize, PublicKey);
+	check(FMemory::Memcmp(TestSignature.Data, Hash, sizeof(Hash)) == 0);
+#endif
 }
 
 bool FSignedArchiveWriter::Close()
@@ -62,14 +61,9 @@ bool FSignedArchiveWriter::Close()
 		SerializeBufferAndSign();
 	}
 
-	// Write out all the chunk signatures
-	for (FEncryptedSignature<GPakFileChunkHashSize>& Signature : ChunkSignatures)
-	{
-		Signature.Serialize(PakWriter);
-	}
-
-	int64 NumSignatures = ChunkSignatures.Num();
-	PakWriter << NumSignatures;
+	FArchive* SignatureWriter = IFileManager::Get().CreateFileWriter(*PakSignaturesFilename);
+	*SignatureWriter << ChunkSignatures;
+	delete SignatureWriter;
 
 	return FArchive::Close();
 }
@@ -118,33 +112,6 @@ void FSignedArchiveWriter::Seek(int64 InPos)
 	UE_LOG(LogPakFile, Fatal, TEXT("Seek is not supported in FSignedArchiveWriter."));
 }
 
-/**
-	* Encrypt a buffer
-	*/
-static int256* Encrypt(const uint8* Data, const int64 DataLength, const FEncryptionKey EncryptionKey)
-{
-	int256* EncryptedData = (int256*)FMemory::Malloc(DataLength * sizeof(uint64));
-	for (int Index = 0; Index < DataLength; ++Index)
-	{
-		EncryptedData[Index] = FEncryption::ModularPow(int256(Data[Index]), EncryptionKey.Exponent, EncryptionKey.Modulus);
-	}
-	return EncryptedData;
-}
-
-/**
-	* Decrypt a buffer
-	*/
-static uint8* Decrypt(const int256* Data, const int64 DataLength, const FEncryptionKey& DecryptionKey)
-{
-	uint8* DecryptedData = (uint8*)FMemory::Malloc(DataLength);
-	for (int64 Index = 0; Index < DataLength; ++Index)
-	{
-		int256 DecryptedByte = FEncryption::ModularPow(Data[Index], DecryptionKey.Exponent, DecryptionKey.Modulus);
-		DecryptedData[Index] = (uint8)DecryptedByte.ToInt();
-	}
-	return DecryptedData;
-}
-
 /**
  * Useful code for testing the encryption methods
  */
@@ -152,8 +119,8 @@ void TestEncryption()
 {
 	FEncryptionKey PublicKey;
 	FEncryptionKey PrivateKey;
-	int256 P = 61;
-	int256 Q = 53;
+	TEncryptionInt P(TEXT("0x21443BD2DD63E995403"));
+	TEncryptionInt Q(TEXT("0x28CBB6E5749AC65749"));
 	FEncryption::GenerateKeyPair(P, Q, PublicKey, PrivateKey);
 
 	// Generate random data
@@ -165,16 +132,17 @@ void TestEncryption()
 	}
 
 	// Generate signature
-	uint8 Signature[20];
-	FSHA1::HashBuffer(Data, DataSize, Signature);
+	FDecryptedSignature OriginalSignature;
+	FEncryptedSignature EncryptedSignature;
+	FDecryptedSignature DecryptedSignature;
+	OriginalSignature.Data = FCrc::MemCrc32(Data, DataSize);
 
 	// Encrypt signature with the private key
-	int256* EncryptedSignature = Encrypt(Signature, sizeof(Signature), PrivateKey);
-	uint8* DecryptedSignature = Decrypt(EncryptedSignature, sizeof(Signature), PublicKey);
+	FEncryption::EncryptSignature(OriginalSignature, EncryptedSignature, PrivateKey);
+	FEncryption::DecryptSignature(EncryptedSignature, DecryptedSignature, PublicKey);
 
 	// Check
-	bool Identical = FMemory::Memcmp(DecryptedSignature, Signature, sizeof(Signature)) == 0;
-	if (Identical)
+	if (OriginalSignature == DecryptedSignature)
 	{
 		UE_LOG(LogPakFile, Display, TEXT("Keys match"));
 	}
@@ -183,7 +151,23 @@ void TestEncryption()
 		UE_LOG(LogPakFile, Fatal, TEXT("Keys mismatched!"));
 	}
 
+	double OverallTime = 0.0;
+	double OverallNumTests = 0.0;
+	for (int32 TestIndex = 0; TestIndex < 10; ++TestIndex)
+	{
+		static const int64 NumTests = 500;
+		double Timer = FPlatformTime::Seconds();
+		{
+			for (int64 i = 0; i < NumTests; ++i)
+			{
+				FEncryption::DecryptSignature(EncryptedSignature, DecryptedSignature, PublicKey);
+			}
+		}
+		Timer = FPlatformTime::Seconds() - Timer;
+		OverallTime += Timer;
+		OverallNumTests += (double)NumTests;
+		UE_LOG(LogPakFile, Display, TEXT("%i signatures decrypted in %.4fs, Avg = %.4fs, OverallAvg = %.4fs"), NumTests, Timer, Timer / (double)NumTests, OverallTime / OverallNumTests);
+	}
+	
 	FMemory::Free(Data);
-	FMemory::Free(EncryptedSignature);
-	FMemory::Free(DecryptedSignature);
 }
\ No newline at end of file
diff --git a/Engine/Source/Programs/UnrealPak/Private/SignedArchiveWriter.h b/Engine/Source/Programs/UnrealPak/Private/SignedArchiveWriter.h
--- a/Engine/Source/Programs/UnrealPak/Private/SignedArchiveWriter.h
+++ b/Engine/Source/Programs/UnrealPak/Private/SignedArchiveWriter.h
@@ -10,8 +10,10 @@ class FSignedArchiveWriter : public FArchive
 	TArray<uint8> Buffer;
 	/** Buffer writer */
 	FMemoryWriter BufferArchive;
-	/** The actuall pak archive */
+	/** The actual pak archive */
 	FArchive& PakWriter;
+	/** The filename of the signature file that accompanies the pak */
+	FString PakSignaturesFilename;
 	/** Size of the archive on disk (including signatures) */
 	int64 SizeOnDisk;
 	/** Data size (excluding signatures) */
@@ -21,12 +23,8 @@ class FSignedArchiveWriter : public FArchive
 	/** Encryption key */
 	FEncryptionKey PrivateKey;
 	/** Signatures */
-	TArray<FEncryptedSignature<GPakFileChunkHashSize>> ChunkSignatures;
+	TArray<FEncryptedSignature> ChunkSignatures;
 
-	/** 
-	 * Encrypts a signature 
-	 */
-	void Encrypt(int256* EncryptedData, const uint8* Data, const int64 DataLength, const FEncryptionKey EncryptionKey);
 	/** 
 	 * Serializes and signs a buffer
 	 */
@@ -34,7 +32,7 @@ class FSignedArchiveWriter : public FArchive
 
 public:
 
-	FSignedArchiveWriter(FArchive& Pak, const FEncryptionKey& InPublicKey, const FEncryptionKey& InPrivateKey);
+	FSignedArchiveWriter(FArchive& InPak, const FString& InPakFilename, const FEncryptionKey& InPublicKey, const FEncryptionKey& InPrivateKey);
 	virtual ~FSignedArchiveWriter();
 
 	// FArchive interface
diff --git a/Engine/Source/Programs/UnrealPak/Private/UnrealPak.cpp b/Engine/Source/Programs/UnrealPak/Private/UnrealPak.cpp
--- a/Engine/Source/Programs/UnrealPak/Private/UnrealPak.cpp
+++ b/Engine/Source/Programs/UnrealPak/Private/UnrealPak.cpp
@@ -719,13 +719,17 @@ FArchive* CreatePakWriter(const TCHAR* Filename)
 		else if (!ReadKeysFromFile(*KeyFilename, Pair))
 		{
 			UE_LOG(LogPakFile, Error, TEXT("Unable to load signature keys %s."), *KeyFilename);
+		}
+
+		if (!TestKeys(Pair))
+		{
 			Pair.PrivateKey.Exponent.Zero();
 		}
 
 		if (!Pair.PrivateKey.Exponent.IsZero())
 		{
 			UE_LOG(LogPakFile, Display, TEXT("Creating signed pak %s."), Filename);
-			Writer = new FSignedArchiveWriter(*Writer, Pair.PublicKey, Pair.PrivateKey);
+			Writer = new FSignedArchiveWriter(*Writer, Filename, Pair.PublicKey, Pair.PrivateKey);
 		}
 		else
 		{
@@ -1374,6 +1378,7 @@ INT32_MAIN_INT32_ARGC_TCHAR_ARGV()
 		UE_LOG(LogPakFile, Error, TEXT("  UnrealPak GenerateKeys=<KeyFilename>"));
 		UE_LOG(LogPakFile, Error, TEXT("  UnrealPak GeneratePrimeTable=<KeyFilename> [-TableMax=<N>]"));
 		UE_LOG(LogPakFile, Error, TEXT("  UnrealPak <PakFilename1> <PakFilename2> -diff"));
+		UE_LOG(LogPakFile, Error, TEXT("  UnrealPak -TestEncryption"));
 		UE_LOG(LogPakFile, Error, TEXT("  Options:"));
 		UE_LOG(LogPakFile, Error, TEXT("    -blocksize=<BlockSize>"));
 		UE_LOG(LogPakFile, Error, TEXT("    -compress"));
@@ -1396,6 +1401,11 @@ INT32_MAIN_INT32_ARGC_TCHAR_ARGV()
 		FParse::Value(FCommandLine::Get(), TEXT("TableMax="), MaxPrimeValue);
 		GeneratePrimeNumberTable(MaxPrimeValue, *KeyFilename);
 	}
+	else if (FParse::Param(FCommandLine::Get(), TEXT("TestEncryption")))
+	{
+		void TestEncryption();
+		TestEncryption();
+	}
 	else 
 	{
 		if (FParse::Param(FCommandLine::Get(), TEXT("Test")))
diff --git a/Engine/Source/Programs/UnrealPak/UnrealPak.Target.cs b/Engine/Source/Programs/UnrealPak/UnrealPak.Target.cs
--- a/Engine/Source/Programs/UnrealPak/UnrealPak.Target.cs
+++ b/Engine/Source/Programs/UnrealPak/UnrealPak.Target.cs
@@ -48,6 +48,9 @@ public class UnrealPakTarget : TargetRules
 		UEBuildConfiguration.bCompileAgainstEngine = false;
 		UEBuildConfiguration.bCompileAgainstCoreUObject = false;
 
+		// ICU is not needed
+		UEBuildConfiguration.bCompileICU = false;
+
 		// UnrealHeaderTool is a console application, not a Windows app (sets entry point to main(), instead of WinMain())
 		OutLinkEnvironmentConfiguration.bIsBuildingConsoleApplication = true;
 	}
diff --git a/Engine/Source/Runtime/Core/Private/GenericPlatform/GenericPlatformMemory.cpp b/Engine/Source/Runtime/Core/Private/GenericPlatform/GenericPlatformMemory.cpp
--- a/Engine/Source/Runtime/Core/Private/GenericPlatform/GenericPlatformMemory.cpp
+++ b/Engine/Source/Runtime/Core/Private/GenericPlatform/GenericPlatformMemory.cpp
@@ -226,7 +226,7 @@ void FGenericPlatformMemory::DumpPlatformAndAllocatorStats( class FOutputDevice&
 	GMalloc->DumpAllocatorStats( Ar );
 }
 
-void FGenericPlatformMemory::MemswapImpl( void* RESTRICT Ptr1, void* RESTRICT Ptr2, SIZE_T Size )
+void FGenericPlatformMemory::MemswapGreaterThan8( void* RESTRICT Ptr1, void* RESTRICT Ptr2, SIZE_T Size )
 {
 	union PtrUnion
 	{
@@ -238,40 +238,27 @@ void FGenericPlatformMemory::MemswapImpl( void* RESTRICT Ptr1, void* RESTRICT Pt
 		UPTRINT PtrUint;
 	};
 
-	if (!Size)
-	{
-		return;
-	}
-
 	PtrUnion Union1 = { Ptr1 };
 	PtrUnion Union2 = { Ptr2 };
 
+	// We may skip up to 7 bytes below, so better make sure that we're swapping more than that
+	// (8 is a common case that we also want to inline before we this call, so skip that too)
+	check(Size > 8);
+
 	if (Union1.PtrUint & 1)
 	{
 		Valswap(*Union1.Ptr8++, *Union2.Ptr8++);
 		Size -= 1;
-		if (!Size)
-		{
-			return;
-		}
 	}
 	if (Union1.PtrUint & 2)
 	{
 		Valswap(*Union1.Ptr16++, *Union2.Ptr16++);
 		Size -= 2;
-		if (!Size)
-		{
-			return;
-		}
 	}
 	if (Union1.PtrUint & 4)
 	{
 		Valswap(*Union1.Ptr32++, *Union2.Ptr32++);
 		Size -= 4;
-		if (!Size)
-		{
-			return;
-		}
 	}
 
 	uint32 CommonAlignment = FMath::Min(FMath::CountTrailingZeros(Union1.PtrUint - Union2.PtrUint), 3u);
diff --git a/Engine/Source/Runtime/Core/Private/HAL/MallocLeakDetection.h b/Engine/Source/Runtime/Core/Private/HAL/MallocLeakDetection.h
--- a/Engine/Source/Runtime/Core/Private/HAL/MallocLeakDetection.h
+++ b/Engine/Source/Runtime/Core/Private/HAL/MallocLeakDetection.h
@@ -190,6 +190,11 @@ public:
 	{ 
 		return UsedMalloc->GetDescriptiveName();
 	}
+
+	virtual void FlushCurrentThreadCache() override
+	{
+		UsedMalloc->FlushCurrentThreadCache();
+	}
 };
 
 #endif // MALLOC_LEAKDETECTION
diff --git a/Engine/Source/Runtime/Core/Private/HAL/MallocVerify.cpp b/Engine/Source/Runtime/Core/Private/HAL/MallocVerify.cpp
--- a/Engine/Source/Runtime/Core/Private/HAL/MallocVerify.cpp
+++ b/Engine/Source/Runtime/Core/Private/HAL/MallocVerify.cpp
@@ -14,8 +14,6 @@ void FMallocVerify::Malloc(void* Ptr)
 {
 	if (Ptr)
 	{
-		FScopeLock Lock(&AllocatedPointersCritical);
-
 		UE_CLOG(AllocatedPointers.Contains(Ptr), LogMemory, Fatal, TEXT("Malloc allocated pointer that's already been allocated: 0x%016llx"), (int64)(PTRINT)Ptr);
 
 		AllocatedPointers.Add(Ptr);			
@@ -24,7 +22,6 @@ void FMallocVerify::Malloc(void* Ptr)
 
 void FMallocVerify::Realloc(void* OldPtr, void* NewPtr)
 {
-	FScopeLock Lock(&AllocatedPointersCritical);
 	if (OldPtr != NewPtr)
 	{			
 		if (OldPtr)
@@ -50,8 +47,6 @@ void FMallocVerify::Free(void* Ptr)
 {
 	if (Ptr)
 	{
-		FScopeLock Lock(&AllocatedPointersCritical);
-
 		UE_CLOG(!AllocatedPointers.Contains(Ptr), LogMemory, Fatal, TEXT("Free attempts to free a pointer that hasn't been allocated yet: 0x%016llx"), (int64)(PTRINT)Ptr);
 
 		AllocatedPointers.Remove(Ptr);
