Copying //UE4/Dev-Core to //UE4/Dev-Main (Source: //UE4/Dev-Core @ 3049602)

diff --git a/Engine/Source/Programs/CrashReporter/CrashReportProcess/DataRouterReportQueue.cs b/Engine/Source/Programs/CrashReporter/CrashReportProcess/DataRouterReportQueue.cs
--- a/Engine/Source/Programs/CrashReporter/CrashReportProcess/DataRouterReportQueue.cs
+++ b/Engine/Source/Programs/CrashReporter/CrashReportProcess/DataRouterReportQueue.cs
@@ -16,13 +16,11 @@ namespace Tools.CrashReporter.CrashReportProcess
 {
 	class DataRouterReportQueue : ReportQueueBase
 	{
-		private const int UncompressedSizeMax = 16*1024*1024;
-
 		protected override string QueueProcessingStartedEventName
 		{
 			get
 			{
-				return StatusReportingConstants.ProcessingStartedDataRouterEvent;
+				return StatusReportingEventNames.ProcessingStartedDataRouterEvent;
 			}
 		}
 
@@ -99,11 +97,13 @@ namespace Tools.CrashReporter.CrashReportProcess
 					if (RecordPair.Length != 2)
 					{
 						CrashReporterProcessServicer.WriteFailure("TryGetNewS3Crashes: bad SQS message was " + SQSRecord);
+						CrashReporterProcessServicer.StatusReporter.IncrementCount(StatusReportingEventNames.ReadS3FileFailedEvent);
 						continue;
 					}
 
 					string S3BucketName = RecordPair[0];
 					string S3Key = RecordPair[1];
+					string ReadableRequestString = "Bucket=" + S3BucketName + " Key=" + S3Key;
 
 					var ObjectRequest = new GetObjectRequest
 					{
@@ -119,34 +119,21 @@ namespace Tools.CrashReporter.CrashReportProcess
 							{
 								if (!TryDecompResponseStream(ResponseStream, ProtocolBufferStream))
 								{
-									CrashReporterProcessServicer.WriteFailure("! GZip fail in DecompResponseStream(): Bucket=" + S3BucketName + " Key=" + S3Key);
+									CrashReporterProcessServicer.WriteFailure("! GZip fail in DecompResponseStream(): " + ReadableRequestString);
+									CrashReporterProcessServicer.StatusReporter.IncrementCount(StatusReportingEventNames.ReadS3FileFailedEvent);
 									continue;
 								}
 							}
 						}
 
-						string UnpackErrorMessage;
-						if (!UnpackPayloadFromProtocolBuffer(ProtocolBufferStream, LandingZone, out UnpackErrorMessage))
-						{
-							// Report failure when UnpackErrorMessage is set - treat as silent failure when not set
-							if (!string.IsNullOrEmpty(UnpackErrorMessage))
-							{
-								string FailString = "! Protocol buffer fail in UnpackPayloadFromProtocolBuffer(): Bucket=" + S3BucketName +
-													" Key=" + S3Key;
-
-								FailString += '\n' + UnpackErrorMessage;
-								CrashReporterProcessServicer.WriteFailure(FailString);
-							}
-
-							continue;
-						}
+						NewCrashCount += UnpackRecordsFromDelimitedProtocolBuffers(ProtocolBufferStream, LandingZone, ReadableRequestString);
 					}
-
-					NewCrashCount++;
 				}
 				catch (Exception ex)
 				{
-					CrashReporterProcessServicer.WriteException("TryGetNewS3Crashes: failure during processing SQS record " + SQSRecord + "\n" + ex);
+					CrashReporterProcessServicer.StatusReporter.IncrementCount(StatusReportingEventNames.ReadS3FileFailedEvent);
+					CrashReporterProcessServicer.WriteException("TryGetNewS3Crashes: failure during processing SQS record " + SQSRecord +
+					                                            "\n" + ex);
 				}
 			}
 		}
@@ -168,38 +155,51 @@ namespace Tools.CrashReporter.CrashReportProcess
 			}
 		}
 
-		private static bool UnpackPayloadFromProtocolBuffer(Stream ProtocolBufferStream, string InLandingZone, out string ErrorMessage)
+		private static int UnpackRecordsFromDelimitedProtocolBuffers(Stream ProtocolBufferStream, string InLandingZone, string ReadableRequestString)
 		{
+			var UnpackedRecordCount = 0;
 			DataRouterConsumer Consumer = new DataRouterConsumer();
-			DataRouterConsumer.ProtocolBufferRecord Message;
 
-			if (!Consumer.TryParse(ProtocolBufferStream, out Message))
+			// Expect one or more pairs of varint size + protocol buffer in the stream
+			while (ProtocolBufferStream.Position < ProtocolBufferStream.Length)
 			{
-				ErrorMessage = Consumer.LastError;
-				return false;
-			}
+				DataRouterConsumer.ProtocolBufferRecord Message;
 
-			if (!DecompressDataRouterContent(Message.Payload, InLandingZone, out ErrorMessage))
-			{
-				return false;
+				if (!Consumer.TryParse(ProtocolBufferStream, out Message))
+				{
+					string FailString = "! Protocol buffer parse fail in UnpackRecordsFromDelimitedProtocolBuffers(): " + ReadableRequestString;
+					FailString += '\n' + Consumer.LastError;
+
+					CrashReporterProcessServicer.WriteFailure(FailString);
+					CrashReporterProcessServicer.StatusReporter.IncrementCount(StatusReportingEventNames.ReadS3FileFailedEvent);
+					break;
+				}
+
+				if (DecompressDataRouterContent(Message.Payload, InLandingZone))
+				{
+					UnpackedRecordCount++;
+				}
 			}
 
-			ErrorMessage = string.Empty;
-			return true;
+			return UnpackedRecordCount;
 		}
 
-		private static unsafe bool DecompressDataRouterContent(byte[] CompressedBufferArray, string InLandingZone, out string ErrorMessage)
+		private static unsafe bool DecompressDataRouterContent(byte[] CompressedBufferArray, string InLandingZone)
 		{
 			// Decompress to landing zone
-			byte[] UncompressedBufferArray = new byte[UncompressedSizeMax];
+			byte[] UncompressedBufferArray = new byte[Config.Default.MaxUncompressedS3RecordSize];
 			int UncompressedSize = 0;
 			fixed (byte* UncompressedBufferPtr = UncompressedBufferArray, CompressedBufferPtr = CompressedBufferArray)
 			{
-				Int32 UncompressResult = NativeMethods.UE4UncompressMemoryZLIB((IntPtr)UncompressedBufferPtr, UncompressedSizeMax, (IntPtr)CompressedBufferPtr, CompressedBufferArray.Length);
+				Int32 UncompressResult = NativeMethods.UE4UncompressMemoryZLIB((IntPtr) UncompressedBufferPtr,
+					UncompressedBufferArray.Length,
+					(IntPtr) CompressedBufferPtr, CompressedBufferArray.Length);
 				if (UncompressResult < 0)
 				{
-					ErrorMessage = "! DecompressDataRouterContent() failed in UE4UncompressMemoryZLIB() with " +
-					               NativeMethods.GetUncompressError(UncompressResult);
+					string FailString = "! DecompressDataRouterContent() failed in UE4UncompressMemoryZLIB() with " +
+					                    NativeMethods.GetUncompressError(UncompressResult);
+					CrashReporterProcessServicer.WriteFailure(FailString);
+					CrashReporterProcessServicer.StatusReporter.IncrementCount(StatusReportingEventNames.ReadS3RecordFailedEvent);
 					return false;
 				}
 				UncompressedSize = UncompressResult;
@@ -210,11 +210,26 @@ namespace Tools.CrashReporter.CrashReportProcess
 				char[] MarkerChars = BinaryData.ReadChars(3);
 				if (MarkerChars[0] == 'C' && MarkerChars[1] == 'R' && MarkerChars[2] == '1')
 				{
-					var CrashHeader = DataRouterReportQueue.CrashHeader.ParseCrashHeader(BinaryData);
+					CrashHeader CrashHeader = DataRouterReportQueue.CrashHeader.ParseCrashHeader(BinaryData);
 
 					// Create safe directory name and then create the directory on disk
-					string CrashFolderName = GetSafeDirectoryName(CrashHeader.DirectoryName);
+					string CrashFolderName = GetSafeFilename(CrashHeader.DirectoryName);
 					string CrashFolderPath = Path.Combine(InLandingZone, CrashFolderName);
+
+					// Early check for duplicate processed report
+					lock (ReportIndexLock)
+					{
+						if (ReportIndex.ContainsReport(CrashFolderName))
+						{
+							// Crash report not accepted by index
+							CrashReporterProcessServicer.WriteEvent(string.Format(
+								"DataRouterReportQueue: Duplicate report skipped early {0} in a DataRouterReportQueue", CrashFolderPath));
+							CrashReporterProcessServicer.StatusReporter.IncrementCount(StatusReportingEventNames.DuplicateRejected);
+							return false; // this isn't an error so don't set error message
+						}
+					}
+
+					// Create target folder
 					int TryIndex = 1;
 					while (Directory.Exists(CrashFolderPath))
 					{
@@ -232,9 +247,11 @@ namespace Tools.CrashReporter.CrashReportProcess
 
 					if (!ParseCrashFiles(BinaryData, CrashHeader.FileCount, CrashFolderPath))
 					{
-						ErrorMessage = "! DecompressDataRouterContent() failed to write files Path=" + CrashFolderPath;
+						string FailString = "! DecompressDataRouterContent() failed to write files Path=" + CrashFolderPath;
+						CrashReporterProcessServicer.WriteFailure(FailString);
+						CrashReporterProcessServicer.StatusReporter.IncrementCount(StatusReportingEventNames.ReadS3RecordFailedEvent);
 						return false;
-					}	
+					}
 				}
 				else
 				{
@@ -248,8 +265,23 @@ namespace Tools.CrashReporter.CrashReportProcess
 					var CrashHeader = DataRouterReportQueue.CrashHeader.ParseCrashHeader(BinaryData);
 
 					// Create safe directory name and then create the directory on disk
-					string CrashFolderName = GetSafeDirectoryName(CrashHeader.DirectoryName);
+					string CrashFolderName = GetSafeFilename(CrashHeader.DirectoryName);
 					string CrashFolderPath = Path.Combine(InLandingZone, CrashFolderName);
+
+					// Early check for duplicate processed report
+					lock (ReportIndexLock)
+					{
+						if (ReportIndex.ContainsReport(CrashFolderName))
+						{
+							// Crash report not accepted by index
+							CrashReporterProcessServicer.WriteEvent(string.Format(
+								"DataRouterReportQueue: Duplicate report skipped early {0} in a DataRouterReportQueue", CrashFolderPath));
+							CrashReporterProcessServicer.StatusReporter.IncrementCount(StatusReportingEventNames.DuplicateRejected);
+							return false; // this isn't an error so don't set error message
+						}
+					}
+
+					// Create target folder
 					int TryIndex = 1;
 					while (Directory.Exists(CrashFolderPath))
 					{
@@ -257,7 +289,7 @@ namespace Tools.CrashReporter.CrashReportProcess
 					}
 					Directory.CreateDirectory(CrashFolderPath);
 
-					if (UncompressedSize != CrashHeader.UncompressedSize + DataRouterReportQueue.CrashHeader.FixedSize)
+					if (UncompressedSize != CrashHeader.UncompressedSize + CrashHeader.FixedSize)
 					{
 						CrashReporterProcessServicer.WriteEvent(
 							string.Format(
@@ -266,48 +298,24 @@ namespace Tools.CrashReporter.CrashReportProcess
 					}
 
 					// Seek to start of files (header size in from start)
-					BinaryData.BaseStream.Position = DataRouterReportQueue.CrashHeader.FixedSize;
+					BinaryData.BaseStream.Position = CrashHeader.FixedSize;
 					if (!ParseCrashFiles(BinaryData, CrashHeader.FileCount, CrashFolderPath))
 					{
-						ErrorMessage = "! DecompressDataRouterContent() failed to write files Path=" + CrashFolderPath;
+						string FailString = "! DecompressDataRouterContent() failed to write files Path=" + CrashFolderPath;
+						CrashReporterProcessServicer.WriteFailure(FailString);
+						CrashReporterProcessServicer.StatusReporter.IncrementCount(StatusReportingEventNames.ReadS3RecordFailedEvent);
 						return false;
-					}	
+					}
 #else
-					ErrorMessage = string.Empty;
+					ErrorMessage = "! DecompressDataRouterContent() failed to read invalid data format. Corrupt or old format data received Path=" + CrashFolderPath;
 					return false;
 #endif
 				}
 			}
 
-			ErrorMessage = string.Empty;
 			return true;
 		}
 
-		private static string GetSafeDirectoryName(string UnsafeName)
-		{
-			char SafeReplacementChar = 'X';
-			char[] InvalidChars = Path.GetInvalidFileNameChars();
-			StringBuilder SafeString = new StringBuilder(UnsafeName.Length);
-
-			foreach (var UnsafeChar in UnsafeName)
-			{
-				char SafeChar = UnsafeChar;
-
-				foreach (var InvalidChar in InvalidChars)
-				{
-					if (UnsafeChar == InvalidChar)
-					{
-						SafeChar = SafeReplacementChar;
-						break;
-					}
-				}
-
-				SafeString.Append(SafeChar);
-			}
-
-			return SafeString.ToString();
-		}
-
 		private static bool ParseCrashFiles(BinaryReader Reader, int FileCount, string CrashFolderPath)
 		{
 			for (int FileIndex = 0; FileIndex < FileCount; FileIndex++)
diff --git a/Engine/Source/Programs/CrashReporter/CrashReportProcess/ReceiverReportQueue.cs b/Engine/Source/Programs/CrashReporter/CrashReportProcess/ReceiverReportQueue.cs
--- a/Engine/Source/Programs/CrashReporter/CrashReportProcess/ReceiverReportQueue.cs
+++ b/Engine/Source/Programs/CrashReporter/CrashReportProcess/ReceiverReportQueue.cs
@@ -12,7 +12,7 @@ namespace Tools.CrashReporter.CrashReportProcess
 		{ 
 			get
 			{
-				return StatusReportingConstants.ProcessingStartedReceiverEvent;
+				return StatusReportingEventNames.ProcessingStartedReceiverEvent;
 			}
 		}
 
diff --git a/Engine/Source/Programs/CrashReporter/CrashReportProcess/ReportIndex.cs b/Engine/Source/Programs/CrashReporter/CrashReportProcess/ReportIndex.cs
--- a/Engine/Source/Programs/CrashReporter/CrashReportProcess/ReportIndex.cs
+++ b/Engine/Source/Programs/CrashReporter/CrashReportProcess/ReportIndex.cs
@@ -92,7 +92,7 @@ namespace Tools.CrashReporter.CrashReportProcess
 
 		public static bool TryAddNewReport(string ReportKey)
 		{
-			if (Index.ContainsKey(ReportKey))
+			if (ContainsReport(ReportKey))
 			{
 				return false;
 			}
@@ -107,6 +107,11 @@ namespace Tools.CrashReporter.CrashReportProcess
 			return true;
 		}
 
+		public static bool ContainsReport(string ReportKey)
+		{
+			return Index.ContainsKey(ReportKey);
+		}
+
 		public static bool TryRemoveReport(string ReportKey)
 		{
 			return Index.Remove(ReportKey);
diff --git a/Engine/Source/Programs/CrashReporter/CrashReportProcess/ReportProcessor.cs b/Engine/Source/Programs/CrashReporter/CrashReportProcess/ReportProcessor.cs
--- a/Engine/Source/Programs/CrashReporter/CrashReportProcess/ReportProcessor.cs
+++ b/Engine/Source/Programs/CrashReporter/CrashReportProcess/ReportProcessor.cs
@@ -172,16 +172,34 @@ namespace Tools.CrashReporter.CrashReportProcess
 		/// <param name="DirInfo">The directory to delete</param>
 		public static void CleanReport(DirectoryInfo DirInfo)
 		{
+			const int MaxRetries = 3;
+
 			bool bWriteException = true;
-			for( int Retry = 0; Retry < 3; ++Retry )
+			for( int Retry = 0; Retry < MaxRetries; ++Retry )
 			{
 				try
 				{
-					foreach( FileInfo Info in DirInfo.GetFiles() )
+					if (!DirInfo.Exists)
+					{
+						break;
+					}
+
+					foreach ( FileInfo Info in DirInfo.GetFiles() )
 					{
 						Info.IsReadOnly = false;
+						Info.Attributes = FileAttributes.Normal;
 					}
 					DirInfo.Delete( true /* delete contents */);
+
+					// Random failures to delete with no exception seen regularly - retry
+					DirInfo.Refresh();
+					if (DirInfo.Exists)
+                    {
+						CrashReporterProcessServicer.WriteEvent("CleanReport: Failed to delete folder without an Exception " + DirInfo);
+						Thread.Sleep(500);
+                        continue;
+                    }
+
 					break;
 				}
 				catch( Exception Ex )
@@ -195,6 +213,12 @@ namespace Tools.CrashReporter.CrashReportProcess
 				}
 				System.Threading.Thread.Sleep( 100 );
 			}
+
+			DirInfo.Refresh();
+			if (DirInfo.Exists)
+			{
+				CrashReporterProcessServicer.WriteEvent(string.Format("CleanReport: Failed to delete folder {0} after {1} retries", DirInfo, MaxRetries));
+			}
 		}
 
 		/// <summary>
@@ -215,13 +239,12 @@ namespace Tools.CrashReporter.CrashReportProcess
 		{
 			try
 			{
-				CrashReporterProcessServicer.StatusReporter.IncrementCount(StatusReportingConstants.ProcessingFailedEvent);
+				CrashReporterProcessServicer.StatusReporter.IncrementCount(StatusReportingEventNames.ProcessingFailedEvent);
 
 				DirectoryInfo DirInfo = new DirectoryInfo( ReportName );
 				// Rename the report directory, so we should be able to quickly find the invalid reports.
 				Directory.CreateDirectory( Config.Default.InvalidReportsDirectory );
-				string CleanFilename = String.Concat( ReportNameAsFilename.Split( Path.GetInvalidFileNameChars() ) );
-				string DestinationDirectory = Path.Combine( Config.Default.InvalidReportsDirectory, CleanFilename );
+				string DestinationDirectory = Path.Combine( Config.Default.InvalidReportsDirectory, ReportQueueBase.GetSafeFilename(ReportNameAsFilename));
 
 				Directory.CreateDirectory( DestinationDirectory );
 
@@ -291,6 +314,7 @@ namespace Tools.CrashReporter.CrashReportProcess
 				// Create a new crash description for uploading
 				CrashDescription NewCrash = new CrashDescription();
 
+				NewCrash.CrashType = NewContext.PrimaryCrashProperties.GetCrashType();
 				NewCrash.BranchName = EngineVersion.GetCleanedBranch();
 				NewCrash.GameName = NewContext.PrimaryCrashProperties.GameName;
 				NewCrash.Platform = NewContext.PrimaryCrashProperties.PlatformFullName;
@@ -554,7 +578,7 @@ namespace Tools.CrashReporter.CrashReportProcess
 
 				UpdateProcessedReports();
 				WebAddCounter.AddEvent();
-				CrashReporterProcessServicer.StatusReporter.IncrementCount(StatusReportingConstants.ProcessingSucceededEvent);
+				CrashReporterProcessServicer.StatusReporter.IncrementCount(StatusReportingEventNames.ProcessingSucceededEvent);
 				double Ratio = (double)WebAddCounter.TotalEvents / (double)ProcessedReports * 100;
 
 				double AddedPerDay = (double)WebAddCounter.TotalEvents / Timer.Elapsed.TotalDays;
diff --git a/Engine/Source/Programs/CrashReporter/CrashReportProcess/ReportQueue.cs b/Engine/Source/Programs/CrashReporter/CrashReportProcess/ReportQueue.cs
--- a/Engine/Source/Programs/CrashReporter/CrashReportProcess/ReportQueue.cs
+++ b/Engine/Source/Programs/CrashReporter/CrashReportProcess/ReportQueue.cs
@@ -207,15 +207,13 @@ namespace Tools.CrashReporter.CrashReportProcess
 		}
 
 		/// <summary> Converts WER metadata xml file to the crash context. </summary>
-		private static void ConvertMetadataToCrashContext(WERReportMetadata Metadata, string NewReportPath, ref FGenericCrashContext OutCrashContext)
+		private static void ConvertMetadataToCrashContext(FReportData ReportData, string NewReportPath, ref FGenericCrashContext OutCrashContext)
 		{
 			if (OutCrashContext == null)
 			{
 				OutCrashContext = new FGenericCrashContext();
 			}
 
-			FReportData ReportData = new FReportData(Metadata, NewReportPath);
-
 			OutCrashContext.PrimaryCrashProperties.CrashVersion = (int)ECrashDescVersions.VER_1_NewCrashFormat;
 
 			//OutCrashContext.PrimaryCrashProperties.ProcessId = 0; don't overwrite valid ids, zero is default anyway
@@ -272,11 +270,6 @@ namespace Tools.CrashReporter.CrashReportProcess
 			// 			OutCrashContext.PrimaryCrashProperties.SourceContext
 			OutCrashContext.PrimaryCrashProperties.UserDescription = string.Join("\n", ReportData.UserDescription);
 
-			if (string.IsNullOrEmpty(OutCrashContext.PrimaryCrashProperties.ErrorMessage))
-			{
-				OutCrashContext.PrimaryCrashProperties.ErrorMessage = string.Join("\n", ReportData.ErrorMessage);
-			}
-
 			// 			OutCrashContext.PrimaryCrashProperties.CrashDumpMode
 			// 			OutCrashContext.PrimaryCrashProperties.Misc.NumberOfCores
 			// 			OutCrashContext.PrimaryCrashProperties.Misc.NumberOfCoresIncludingHyperthreads
@@ -303,6 +296,21 @@ namespace Tools.CrashReporter.CrashReportProcess
 			// 			OutCrashContext.PrimaryCrashProperties.MemoryStats.OOMAllocationAlignment
 			OutCrashContext.PrimaryCrashProperties.TimeofCrash = new DateTime(ReportData.Ticks);
 			OutCrashContext.PrimaryCrashProperties.bAllowToBeContacted = ReportData.AllowToBeContacted;
+
+			OutCrashContext.PrimaryCrashProperties.DeploymentName = ReportData.DeploymentName;
+			OutCrashContext.PrimaryCrashProperties.IsEnsure = ReportData.IsEnsure;
+			OutCrashContext.PrimaryCrashProperties.IsAssert = ReportData.IsAssert;
+			OutCrashContext.PrimaryCrashProperties.CrashType = ReportData.CrashType;
+
+			GetErrorMessageFromMetadata(ReportData, OutCrashContext);
+		}
+
+		private static void GetErrorMessageFromMetadata(FReportData ReportData, FGenericCrashContext CrashContext)
+		{
+			if (string.IsNullOrEmpty(CrashContext.PrimaryCrashProperties.ErrorMessage))
+			{
+				CrashContext.PrimaryCrashProperties.ErrorMessage = string.Join("\n", ReportData.ErrorMessage);
+			}
 		}
 
 		/// <summary> Looks for the WER metadata xml file, if found, will return a new instance of the WERReportMetadata. </summary>
@@ -412,15 +420,23 @@ namespace Tools.CrashReporter.CrashReportProcess
 			FGenericCrashContext GenericContext = FindCrashContext(NewReportPath);
 			FGenericCrashContext Context = GenericContext;
 
-			bool bFromWER = false;
-			if (Context == null || !Context.HasProcessedData())
+			bool bContextDirty = false;
+			WERReportMetadata MetaData = FindMetadata(NewReportPath);
+			if (MetaData != null)
 			{
-				WERReportMetadata MetaData = FindMetadata(NewReportPath);
-				if (MetaData != null)
+				if (Context == null)
 				{
+					// Missing crash context
 					FReportData ReportData = new FReportData(MetaData, NewReportPath);
-					ConvertMetadataToCrashContext(MetaData, NewReportPath, ref Context);
-					bFromWER = true;
+					ConvertMetadataToCrashContext(ReportData, NewReportPath, ref Context);
+					bContextDirty = true;
+				}
+				else if (!Context.HasProcessedData())
+				{
+					// Missing data - try to get from WER metadata
+					FReportData ReportData = new FReportData(MetaData, NewReportPath);
+					GetErrorMessageFromMetadata(ReportData, Context);
+					bContextDirty = true;
 				}
 			}
 
@@ -435,7 +451,7 @@ namespace Tools.CrashReporter.CrashReportProcess
 			Context.PrimaryCrashProperties.SetPlatformFullName();
 
 			// Added data from WER, save to the crash context file.
-			if (bFromWER)
+			if (bContextDirty)
 			{
 				Context.ToFile();
 			}
@@ -480,6 +496,7 @@ namespace Tools.CrashReporter.CrashReportProcess
 					// Crash report not accepted by index
 					CrashReporterProcessServicer.WriteEvent(string.Format(
 						"EnqueueNewReport: Duplicate report skipped {0} in queue {1}", NewReportPath, QueueName));
+					CrashReporterProcessServicer.StatusReporter.IncrementCount(StatusReportingEventNames.DuplicateRejected);
 					ReportProcessor.CleanReport(new DirectoryInfo(NewReportPath));
 					return false;
 				}
@@ -490,7 +507,7 @@ namespace Tools.CrashReporter.CrashReportProcess
 				NewCrashContexts.Enqueue(Context);
 			}
 			EnqueueCounter.AddEvent();
-			CrashReporterProcessServicer.StatusReporter.IncrementCount(StatusReportingConstants.QueuedEvent);
+			CrashReporterProcessServicer.StatusReporter.IncrementCount(StatusReportingEventNames.QueuedEvent);
 			CrashReporterProcessServicer.WriteEvent("+ Enqueued: BuiltFromCL=" + string.Format("{0,7}", BuiltFromCL) + " Path=" + NewReportPath);
 			return true;
 		}
@@ -574,11 +591,12 @@ namespace Tools.CrashReporter.CrashReportProcess
 				}
 
 				string Filename = FBinaryReaderHelper.ReadFixedSizeString(BinaryData);
+				string SafeFilename = GetSafeFilename(Filename);
 				Int32 FiledataLength = BinaryData.ReadInt32();
 				byte[] Filedata = BinaryData.ReadBytes(FiledataLength);
 
 				Directory.CreateDirectory(DestinationFolderPath);
-				File.WriteAllBytes(Path.Combine(DestinationFolderPath, Filename), Filedata);
+				File.WriteAllBytes(Path.Combine(DestinationFolderPath, SafeFilename), Filedata);
 				return true;
 			}
 			catch (Exception Ex)
@@ -587,10 +605,15 @@ namespace Tools.CrashReporter.CrashReportProcess
 			}
 		}
 
-		// FIELDS
+        public static string GetSafeFilename(string UnsafeName)
+        {
+            return string.Join("X", UnsafeName.Split(Path.GetInvalidFileNameChars()));
+        }
+
+        // FIELDS
 
-		/// <summary> A list of freshly landed crash reports. </summary>
-		private readonly Queue<FGenericCrashContext> NewCrashContexts = new Queue<FGenericCrashContext>();
+        /// <summary> A list of freshly landed crash reports. </summary>
+        private readonly Queue<FGenericCrashContext> NewCrashContexts = new Queue<FGenericCrashContext>();
 
 		/// <summary> Object used to synchronize the access to the NewReport. </summary>
 		private readonly Object NewReportsLock = new Object();
@@ -608,7 +631,7 @@ namespace Tools.CrashReporter.CrashReportProcess
 
 		private readonly EventCounter EnqueueCounter = new EventCounter(TimeSpan.FromMinutes(90), 20);
 
-		private static Object ReportIndexLock = new Object();
+		protected static Object ReportIndexLock = new Object();
 	}
 }
 
diff --git a/Engine/Source/Programs/CrashReporter/CrashReportProcess/StatusReporting.cs b/Engine/Source/Programs/CrashReporter/CrashReportProcess/StatusReporting.cs
--- a/Engine/Source/Programs/CrashReporter/CrashReportProcess/StatusReporting.cs
+++ b/Engine/Source/Programs/CrashReporter/CrashReportProcess/StatusReporting.cs
@@ -1,6 +1,7 @@
 ﻿using System;
 using System.Collections.Generic;
 using System.Linq;
+using System.Reflection;
 using System.Text;
 using System.Threading;
 using System.Threading.Tasks;
@@ -58,6 +59,14 @@ namespace Tools.CrashReporter.CrashReportProcess
 			}
 			CrashReporterProcessServicer.WriteSlack(StartupMessage.ToString());
 
+			// Init events by enumerating event names
+			FieldInfo[] EventNameFields = typeof(StatusReportingEventNames).GetFields(BindingFlags.Static | BindingFlags.Public);
+			foreach (FieldInfo EventNameField in EventNameFields)
+			{
+				string CountName = (string)EventNameField.GetValue(null);
+				Counters.Add(CountName, 0);
+			}
+
 			StatusReportingTask = Task.Factory.StartNew(() =>
 			{
 				var Cancel = CancelSource.Token;
@@ -81,9 +90,9 @@ namespace Tools.CrashReporter.CrashReportProcess
 							if (SetQueueSizesCallCount >= QueueSizes.Count)
 							{
 								int ProcessingStartedInPeriodReceiver = 0;
-								CountsInPeriod.TryGetValue(StatusReportingConstants.ProcessingStartedReceiverEvent, out ProcessingStartedInPeriodReceiver);
+								CountsInPeriod.TryGetValue(StatusReportingEventNames.ProcessingStartedReceiverEvent, out ProcessingStartedInPeriodReceiver);
 								int ProcessingStartedInPeriodDataRouter = 0;
-								CountsInPeriod.TryGetValue(StatusReportingConstants.ProcessingStartedDataRouterEvent, out ProcessingStartedInPeriodDataRouter);
+								CountsInPeriod.TryGetValue(StatusReportingEventNames.ProcessingStartedDataRouterEvent, out ProcessingStartedInPeriodDataRouter);
 								int ProcessingStartedInPeriod = ProcessingStartedInPeriodReceiver + ProcessingStartedInPeriodDataRouter;
 
 								if (ProcessingStartedInPeriod > 0)
@@ -115,10 +124,12 @@ namespace Tools.CrashReporter.CrashReportProcess
 								}
 
 								StatusReportMessage.AppendLine("Queue sizes...");
+								StatusReportMessage.Append("> ");
 								foreach (var Queue in QueueSizes)
 								{
-									StatusReportMessage.AppendLine("> " + Queue.Key + " " + Queue.Value);
+									StatusReportMessage.Append(Queue.Key + " " + Queue.Value + "                ");
 								}
+								StatusReportMessage.AppendLine();
 							}
 							if (CountsInPeriod.Count > 0)
 							{
@@ -147,6 +158,11 @@ namespace Tools.CrashReporter.CrashReportProcess
 			});
 		}
 
+		public void OnPreStopping()
+		{
+			CrashReporterProcessServicer.WriteSlack("CRP stopping...");
+		}
+
 		public void Dispose()
 		{
 			if (StatusReportingTask != null)
diff --git a/Engine/Source/Programs/CrashReporter/CrashReportProcess/StatusReportingConstants.cs b/Engine/Source/Programs/CrashReporter/CrashReportProcess/StatusReportingConstants.cs
--- a/Engine/Source/Programs/CrashReporter/CrashReportProcess/StatusReportingConstants.cs
+++ b/Engine/Source/Programs/CrashReporter/CrashReportProcess/StatusReportingConstants.cs
@@ -6,13 +6,16 @@ using System.Threading.Tasks;
 
 namespace Tools.CrashReporter.CrashReportProcess
 {
-	static class StatusReportingConstants
+	static class StatusReportingEventNames
 	{
-		public const string ProcessingFailedEvent = "Processing failed";
-		public const string ProcessingSucceededEvent = "Processing succeeded";
-		public const string ExceptionEvent = "Handled exceptions";
+		public const string ReadS3FileFailedEvent = "Reading crash records file from S3 failed";
+		public const string ReadS3RecordFailedEvent = "Reading crash record from S3 failed";
+		public const string DuplicateRejected = "Duplicates ignored";
+		public const string QueuedEvent = "Queued for processing";
 		public const string ProcessingStartedReceiverEvent = "Started processing (from CRR)";
 		public const string ProcessingStartedDataRouterEvent = "Started processing (from Data Router)";
-		public const string QueuedEvent = "Queued for processing";
+		public const string ProcessingSucceededEvent = "Processing succeeded";
+		public const string ProcessingFailedEvent = "Processing failed";
+		public const string ExceptionEvent = "Handled exceptions";	
 	}
 }
diff --git a/Engine/Source/Programs/CrashReporter/CrashReportWebSite/App_Start/UnityConfig.cs b/Engine/Source/Programs/CrashReporter/CrashReportWebSite/App_Start/UnityConfig.cs
--- a/Engine/Source/Programs/CrashReporter/CrashReportWebSite/App_Start/UnityConfig.cs
+++ b/Engine/Source/Programs/CrashReporter/CrashReportWebSite/App_Start/UnityConfig.cs
@@ -1,6 +1,8 @@
 using System;
 using Microsoft.Practices.Unity;
 using Microsoft.Practices.Unity.Configuration;
+using Tools.CrashReporter.CrashReportWebSite.DataModels;
+using Tools.CrashReporter.CrashReportWebSite.DataModels.Repositories;
 
 namespace Tools.CrashReporter.CrashReportWebSite.App_Start
 {
@@ -36,6 +38,8 @@ namespace Tools.CrashReporter.CrashReportWebSite.App_Start
             // container.LoadConfiguration();
 
             // TODO: Register your types here
+
+            container.RegisterType<IUnitOfWork, UnitOfWork>();
             // container.RegisterType<IProductRepository, ProductRepository>();
         }
     }
diff --git a/Engine/Source/Programs/CrashReporter/CrashReportWebSite/Controllers/BuggsController.cs b/Engine/Source/Programs/CrashReporter/CrashReportWebSite/Controllers/BuggsController.cs
--- a/Engine/Source/Programs/CrashReporter/CrashReportWebSite/Controllers/BuggsController.cs
+++ b/Engine/Source/Programs/CrashReporter/CrashReportWebSite/Controllers/BuggsController.cs
@@ -4,9 +4,13 @@ using System;
 using System.Collections.Generic;
 using System.Diagnostics;
 using System.Linq;
+using System.Linq.Expressions;
+using System.Web;
 using System.Web.Mvc;
-
-using Tools.CrashReporter.CrashReportWebSite.Models;
+using Microsoft.Practices.ObjectBuilder2;
+using Tools.CrashReporter.CrashReportWebSite.DataModels;
+using Tools.CrashReporter.CrashReportWebSite.DataModels.Repositories;
+using Tools.CrashReporter.CrashReportWebSite.ViewModels;
 
 namespace Tools.CrashReporter.CrashReportWebSite.Controllers
 {
@@ -15,28 +19,16 @@ namespace Tools.CrashReporter.CrashReportWebSite.Controllers
 	/// </summary>
 	public class BuggsController : Controller
 	{
+        private IUnitOfWork _unitOfWork;
+
 		/// <summary>
 		/// An empty constructor.
 		/// </summary>
-		public BuggsController()
+		public BuggsController(IUnitOfWork unitOfWork)
 		{
+		    _unitOfWork = unitOfWork;
 		}
 
-		/*
-		/// <summary>
-		/// The empty view of the buggs page.
-		/// </summary>
-		/// 
-		public ActionResult Index()
-		{
-			using (FAutoScopedLogTimer LogTimer = new FAutoScopedLogTimer( this.GetType().ToString(), bCreateNewLog: true ))
-			{
-				BuggsViewModel Results = new BuggsViewModel();
-				Results.GenerationTime = LogTimer.GetElapsedSeconds().ToString( "F2" );
-				return View( "Index", Results );
-			}
-		}*/
-
 		/// <summary>
 		/// The Index action.
 		/// </summary>
@@ -44,19 +36,17 @@ namespace Tools.CrashReporter.CrashReportWebSite.Controllers
 		/// <returns>The view to display a list of Buggs on the client.</returns>
 		public ActionResult Index( FormCollection BuggsForm )
 		{
-			using (FAutoScopedLogTimer LogTimer = new FAutoScopedLogTimer( this.GetType().ToString(), bCreateNewLog: true ))
+			using (var logTimer = new FAutoScopedLogTimer( this.GetType().ToString(), bCreateNewLog: true ))
 			{
-				BuggRepository Buggs = new BuggRepository();
-
-				FormHelper FormData = new FormHelper( Request, BuggsForm, "CrashesInTimeFrameGroup" );
-				BuggsViewModel Results = Buggs.GetResults( FormData );
-				foreach (var Bugg in Results.Results)
+				var formData = new FormHelper( Request, BuggsForm, "CrashesInTimeFrameGroup" );
+			    var results = GetResults( formData );
+				foreach (var bugg in results.Results)
 				{
 					// Populate function calls.
-					Bugg.GetFunctionCalls();
+				    //Bugg.FunctionCalls = new CallStackContainer(Bugg.Crashes.First()).GetFunctionCalls();
 				}
-				Results.GenerationTime = LogTimer.GetElapsedSeconds().ToString( "F2" );
-				return View( "Index", Results );
+				results.GenerationTime = logTimer.GetElapsedSeconds().ToString( "F2" );
+				return View( "Index", results );
 			}
 		}
 
@@ -74,15 +64,13 @@ namespace Tools.CrashReporter.CrashReportWebSite.Controllers
 				int BuggIDToBeAddedToJira = 0;
 				foreach( var Entry in BuggsForm )
 				{
-					if( Entry.ToString().Contains( Bugg.JiraSubmitName ) )
+					if( Entry.ToString().Contains( "CopyToJira-" ) )
 					{
-						int.TryParse( Entry.ToString().Substring( Bugg.JiraSubmitName.Length ), out BuggIDToBeAddedToJira );
+                        int.TryParse(Entry.ToString().Substring("CopyToJira-".Length), out BuggIDToBeAddedToJira);
 						break;
 					}
 				}
 
-				BuggRepository Buggs = new BuggRepository();
-
 				// Set the display properties based on the radio buttons
 				bool DisplayModuleNames = false;
 				if( BuggsForm["DisplayModuleNames"] == "true" )
@@ -117,25 +105,24 @@ namespace Tools.CrashReporter.CrashReportWebSite.Controllers
 
 				// Create a new view and populate with crashes
 				List<Crash> Crashes = null;
-		
 				BuggViewModel Model = new BuggViewModel();
-				Bugg NewBugg = Buggs.GetBugg( id );
-				if( NewBugg == null )
+				Bugg newBugg = _unitOfWork.BuggRepository.GetById( id );
+				if( newBugg == null )
 				{
-					return RedirectToAction( "" );
+					return RedirectToAction( "Index" );
 				}
 
-				Crashes = NewBugg.GetCrashes();
+				Crashes = newBugg.Crashes.OrderByDescending(data => data.TimeOfCrash).ToList();
 
 				using (FAutoScopedLogTimer GetCrashesTimer = new FAutoScopedLogTimer( "Bugg.PrepareBuggForJira" ))
 				{
 					if (Crashes.Count > 0)
 					{
-						NewBugg.PrepareBuggForJira( Crashes );
+						newBugg.PrepareBuggForJira( Crashes );
 
 						if (BuggIDToBeAddedToJira != 0)
 						{
-							NewBugg.CopyToJira();
+							newBugg.CopyToJira();
 						}
 					}
 				}
@@ -145,13 +132,13 @@ namespace Tools.CrashReporter.CrashReportWebSite.Controllers
 					var JC = JiraConnection.Get();
 					bool bValidJira = false;
 
-					// Verify valid JiraID, this may be still a TTP 
-					if( !string.IsNullOrEmpty( NewBugg.Jira ) )
+					// Verify valid JiraID, this may be still a TTP
+					if( !string.IsNullOrEmpty( newBugg.TTPID ) )
 					{
-						int TTPID = 0;
-						int.TryParse( NewBugg.Jira, out TTPID );
+						var jira = 0;
+                        int.TryParse(newBugg.TTPID, out jira);
 
-						if( TTPID == 0 )
+                        if (jira == 0)
 						{
 							//AddBuggJiraMapping( NewBugg, ref FoundJiras, ref JiraIDtoBugg );
 							bValidJira = true;
@@ -161,7 +148,7 @@ namespace Tools.CrashReporter.CrashReportWebSite.Controllers
 					if( JC.CanBeUsed() && bValidJira )
 					{
 						// Grab the data form JIRA.
-						string JiraSearchQuery = "key = " + NewBugg.Jira;
+						string JiraSearchQuery = "key = " + newBugg.TTPID;
 						var JiraResults = new Dictionary<string, Dictionary<string, object>>();
 
 						try
@@ -181,48 +168,50 @@ namespace Tools.CrashReporter.CrashReportWebSite.Controllers
 						}
 						catch (System.Exception)
 						{
-							NewBugg.JiraSummary = "JIRA MISMATCH";
-							NewBugg.JiraComponentsText = "JIRA MISMATCH";
-							NewBugg.JiraResolution = "JIRA MISMATCH";
-							NewBugg.JiraFixVersionsText = "JIRA MISMATCH";
-							NewBugg.JiraFixCL = "JIRA MISMATCH";
+                            newBugg.JiraSummary = "JIRA MISMATCH";
+                            newBugg.JiraComponentsText = "JIRA MISMATCH";
+                            newBugg.JiraResolution = "JIRA MISMATCH";
+                            newBugg.JiraFixVersionsText = "JIRA MISMATCH";
+                            newBugg.JiraFixCL = "JIRA MISMATCH";
 						}
 
 						// Jira Key, Summary, Components, Resolution, Fix version, Fix changelist
-						foreach( var Jira in JiraResults )
+						foreach( var jira in JiraResults )
 						{
-							string JiraID = Jira.Key;
+							string jiraId = jira.Key;
 
-							string Summary = (string)Jira.Value["summary"];
+							var summary = (string)jira.Value["summary"];
 
 							string ComponentsText = "";
-							System.Collections.ArrayList Components = (System.Collections.ArrayList)Jira.Value["components"];
+							System.Collections.ArrayList Components = (System.Collections.ArrayList)jira.Value["components"];
 							foreach( Dictionary<string, object> Component in Components )
 							{
 								ComponentsText += (string)Component["name"];
 								ComponentsText += " ";
 							}
 
-							Dictionary<string, object> ResolutionFields = (Dictionary<string, object>)Jira.Value["resolution"];
+							Dictionary<string, object> ResolutionFields = (Dictionary<string, object>)jira.Value["resolution"];
 							string Resolution = ResolutionFields != null ? (string)ResolutionFields["name"] : "";
 
 							string FixVersionsText = "";
-							System.Collections.ArrayList FixVersions = (System.Collections.ArrayList)Jira.Value["fixVersions"];
+							System.Collections.ArrayList FixVersions = (System.Collections.ArrayList)jira.Value["fixVersions"];
 							foreach( Dictionary<string, object> FixVersion in FixVersions )
 							{
 								FixVersionsText += (string)FixVersion["name"];
 								FixVersionsText += " ";
 							}
 
-							int FixCL = Jira.Value["customfield_11200"] != null ? (int)(decimal)Jira.Value["customfield_11200"] : 0;
+							int FixCL = jira.Value["customfield_11200"] != null ? (int)(decimal)jira.Value["customfield_11200"] : 0;
 
-							NewBugg.JiraSummary = Summary;
-							NewBugg.JiraComponentsText = ComponentsText;
-							NewBugg.JiraResolution = Resolution;
-							NewBugg.JiraFixVersionsText = FixVersionsText;
+                            //Conversion to ado.net entity framework
+
+                            newBugg.JiraSummary = summary;
+                            newBugg.JiraComponentsText = ComponentsText;
+                            newBugg.JiraResolution = Resolution;
+                            newBugg.JiraFixVersionsText = FixVersionsText;
 							if( FixCL != 0 )
 							{
-								NewBugg.JiraFixCL = FixCL.ToString();
+								newBugg.FixedChangeList = FixCL.ToString();
 							}
 							
 							break;
@@ -235,28 +224,29 @@ namespace Tools.CrashReporter.CrashReportWebSite.Controllers
 				{
 					if( !string.IsNullOrEmpty( BuggsForm["SetStatus"] ) )
 					{
-						NewBugg.Status = BuggsForm["SetStatus"];
-						Buggs.SetBuggStatus( NewBugg.Status, id );
+						newBugg.Status = BuggsForm["SetStatus"];
+                        newBugg.Crashes.ForEach(data =>data.Status = BuggsForm["SetStatus"]);
 					}
 
 					if( !string.IsNullOrEmpty( BuggsForm["SetFixedIn"] ) )
 					{
-						NewBugg.FixedChangeList = BuggsForm["SetFixedIn"];
-						Buggs.SetBuggFixedChangeList( NewBugg.FixedChangeList, id );
+						newBugg.FixedChangeList = BuggsForm["SetFixedIn"];
+                        newBugg.Crashes.ForEach(data => data.FixedChangeList = BuggsForm["SetFixedIn"]);
 					}
 
 					if( !string.IsNullOrEmpty( BuggsForm["SetTTP"] ) )
 					{
-						NewBugg.Jira = BuggsForm["SetTTP"];
-						Buggs.SetJIRAForBuggAndCrashes( NewBugg.Jira, id );
+						newBugg.TTPID = BuggsForm["SetTTP"];
+                        newBugg.Crashes.ForEach(data => data.Jira = BuggsForm["SetTTP"]);
 					}
 
-					// <STATUS>
+                    _unitOfWork.BuggRepository.Update(newBugg);
+                    _unitOfWork.Save();
 				}
 
 				// Set up the view model with the crash data
-				Model.Bugg = NewBugg;
 				Model.Crashes = Crashes;
+			    Model.Bugg = newBugg;
 
 				Crash NewCrash = Model.Crashes.FirstOrDefault();
 				if( NewCrash != null )
@@ -283,12 +273,376 @@ namespace Tools.CrashReporter.CrashReportWebSite.Controllers
 						Model.SourceContext = NewCrash.SourceContext;
 					}
 
-					Model.Bugg.LatestCrashSummary = NewCrash.Summary;
+					Model.LatestCrashSummary = NewCrash.Summary;
 				}
-
+                Model.LatestCrashSummary = NewCrash.Summary;
+                Model.Bugg.LatestCrashSummary = NewCrash.Summary;
 				Model.GenerationTime = LogTimer.GetElapsedSeconds().ToString( "F2" );
 				return View( "Show", Model );
 			}
 		}
+
+        /// <summary>
+        /// Retrieve all Buggs matching the search criteria.
+        /// </summary>
+        /// <param name="FormData">The incoming form of search criteria from the client.</param>
+        /// <returns>A view to display the filtered Buggs.</returns>
+        public BuggsViewModel GetResults(FormHelper FormData)
+        {
+            using (FAutoScopedLogTimer LogTimer = new FAutoScopedLogTimer(this.GetType().ToString()))
+            {
+                // Right now we take a Result IQueryable starting with ListAll() Buggs then we widdle away the result set by tacking on
+                // Linq queries. Essentially it's Results.ListAll().Where().Where().Where().Where().Where().Where()
+                // Could possibly create more optimized queries when we know exactly what we're querying
+                // The downside is that if we add more parameters each query may need to be updated.... Or we just add new case statements
+                // The other downside is that there's less code reuse, but that may be worth it.
+
+                var fromDate = FormData.DateFrom;
+                var toDate = FormData.DateTo.AddDays(1);
+
+                IQueryable<Bugg> results = null;
+                var skip = (FormData.Page - 1) * FormData.PageSize;
+                var take = FormData.PageSize;
+
+                var userGroupId = _unitOfWork.UserGroupRepository.First(data => data.Name == FormData.UserGroup).Id;
+
+                // Get every Bugg.
+                var resultsAll = _unitOfWork.BuggRepository.ListAll();
+
+                // Look at all Buggs that are still 'open' i.e. the last crash occurred in our date range.
+                results = FilterByDate(resultsAll, FormData.DateFrom, FormData.DateTo);
+
+                // Filter results by build version.
+                if (!string.IsNullOrEmpty(FormData.VersionName))
+                {
+                    results = FilterByBuildVersion(results, FormData.VersionName);
+                }
+
+                // Filter by BranchName
+                if (!string.IsNullOrEmpty(FormData.BranchName))
+                {
+                    results =
+                        results.Where(
+                            data => data.Crashes.Any(da => da.Branch.Equals(FormData.BranchName)));
+                }
+
+                // Filter by PlatformName
+                if (!string.IsNullOrEmpty(FormData.PlatformName))
+                {
+                    results =
+                        results.Where(
+                            data => data.Crashes.Any(da => da.PlatformName.Equals(FormData.PlatformName)));
+                }
+
+                // Run at the end
+                if (!string.IsNullOrEmpty(FormData.SearchQuery))
+                {
+                    results = FilterByCallstack(results, FormData.SearchQuery);
+                }
+
+                // Filter by Crash Type
+                if (FormData.CrashType != "All")
+                {
+                    switch (FormData.CrashType)
+                    {
+                        case "Crashes":
+                            results = results.Where(buggInstance => buggInstance.CrashType == 1);
+                            break;
+                        case "Assert":
+                            results = results.Where(buggInstance => buggInstance.CrashType == 2);
+                            break;
+                        case "Ensure":
+                            results = results.Where(buggInstance => buggInstance.CrashType == 3);
+                            break;
+                        case "CrashesAsserts":
+                            results = results.Where(buggInstance => buggInstance.CrashType == 1 || buggInstance.CrashType == 2);
+                            break;
+                    }
+                }
+
+                // Filter by user group if present
+                if(!string.IsNullOrEmpty(FormData.UserGroup))
+                    results = FilterByUserGroup(results, FormData.UserGroup);
+
+                if (!string.IsNullOrEmpty(FormData.JiraId))
+                {
+                    results = FilterByJira(results, FormData.JiraId);
+                }
+                
+                // Grab just the results we want to display on this page
+                var totalCountedRecords = results.Count();
+                
+                var crashesByUserGroupCounts = _unitOfWork.CrashRepository.ListAll().Where(data =>
+                   data.User.UserGroupId == userGroupId && 
+                   data.TimeOfCrash >= fromDate &&
+                   data.TimeOfCrash <= toDate &&
+                   data.PatternId != null)
+                   .GroupBy(data => data.PatternId)
+                   .Select(data => new { data.Key, Count = data.Count() })
+                   .OrderByDescending(data => data.Count)
+                   .ToDictionary(data => data.Key, data => data.Count);
+
+                var crashesInTimeFrameCounts = _unitOfWork.CrashRepository.ListAll().Where(data =>
+                    data.TimeOfCrash >= fromDate &&
+                    data.TimeOfCrash <= toDate
+                    && data.PatternId != null)
+                    .GroupBy(data => data.PatternId)
+                    .Select(data => new { data.Key, Count = data.Count() })
+                    .OrderByDescending(data => data.Count)
+                    .ToDictionary(data => data.Key, data => data.Count);
+
+                var affectedUsers = _unitOfWork.CrashRepository.ListAll().Where(data =>
+                   data.User.UserGroupId == userGroupId &&
+                   data.TimeOfCrash >= fromDate &&
+                   data.TimeOfCrash <= toDate &&
+                   data.PatternId != null).Select(data => new {PatternId = data.PatternId, ComputerName = data.ComputerName}).Distinct()
+                   .GroupBy(data => data.PatternId)
+                   .Select(data => new { data.Key, Count = data.Count() })
+                   .OrderByDescending(data => data.Count)
+                   .ToDictionary(data => data.Key, data => data.Count);
+                
+                results = GetSortedResults(results, FormData.SortTerm, FormData.SortTerm == "descending",
+                FormData.DateFrom, FormData.DateTo, FormData.UserGroup);
+                
+                var sortedResultsList = 
+                    results.ToList();
+
+                // Get UserGroup ResultCounts
+                var groups = results.SelectMany(data => data.UserGroups).GroupBy(data => data.Name).Select(data => new { Key = data.Key, Count = data.Count()}).ToDictionary(x => x.Key, y => y.Count);
+                var groupCounts = new SortedDictionary<string, int>(groups);
+                
+                foreach (var bugg in sortedResultsList)
+                {
+                    if (bugg.PatternId.HasValue && crashesByUserGroupCounts.ContainsKey(bugg.PatternId.Value))
+                    {
+                        bugg.CrashesInTimeFrameGroup = crashesByUserGroupCounts[bugg.PatternId.Value];
+                    }
+                    else
+                        bugg.CrashesInTimeFrameGroup = 0;
+
+                    if (bugg.PatternId.HasValue && crashesInTimeFrameCounts.ContainsKey(bugg.PatternId.Value))
+                    {
+                        bugg.CrashesInTimeFrameAll = crashesInTimeFrameCounts[bugg.PatternId.Value];
+                    }
+                    else
+                        bugg.CrashesInTimeFrameAll = 0;
+
+                    if (bugg.PatternId.HasValue && affectedUsers.ContainsKey(bugg.PatternId.Value))
+                    {
+                        bugg.NumberOfUniqueMachines = affectedUsers[bugg.PatternId.Value];
+                    }
+                    else
+                        bugg.NumberOfUniqueMachines = 0;
+                }
+
+                sortedResultsList = sortedResultsList.OrderByDescending(data => data.CrashesInTimeFrameGroup).Skip(skip)
+                        .Take(totalCountedRecords >= skip + take ? take : totalCountedRecords).ToList();
+
+                foreach (var bugg in sortedResultsList)
+                {
+                    var crash =
+                        _unitOfWork.CrashRepository.First(data => data.Buggs.Any(bg => bg.Id == bugg.Id));
+                    bugg.FunctionCalls = crash.GetCallStack().GetFunctionCalls();
+                }
+
+                return new BuggsViewModel(_unitOfWork.CrashRepository)
+                {
+                    Results = sortedResultsList,
+                    PagingInfo = new PagingInfo { CurrentPage = FormData.Page, PageSize = FormData.PageSize, TotalResults = totalCountedRecords },
+                    SortTerm = FormData.SortTerm,
+                    SortOrder = FormData.SortOrder,
+                    UserGroup = FormData.UserGroup,
+                    CrashType = FormData.CrashType,
+                    SearchQuery = FormData.SearchQuery,
+                    DateFrom = (long)(FormData.DateFrom - CrashesViewModel.Epoch).TotalMilliseconds,
+                    DateTo = (long)(FormData.DateTo - CrashesViewModel.Epoch).TotalMilliseconds,
+                    VersionName = FormData.VersionName,
+                    PlatformName = FormData.PlatformName,
+                    BranchName = FormData.BranchName,
+                    GroupCounts = groupCounts,
+                    Jira = FormData.JiraId,
+                };
+            }
+        }
+
+        /// <summary>
+        /// Filter a set of Buggs by a jira ticket.
+        /// </summary>
+        /// <param name="results">The unfiltered set of Buggs</param>
+        /// <param name="jira">The ticket by which to filter the list of buggs</param>
+        /// <returns></returns>
+        public IQueryable<Bugg> FilterByJira(IQueryable<Bugg> results, string jira)
+        {
+            return results.Where(bugg => bugg.TTPID == jira);
+        }
+
+        /// <summary>
+        /// Filter a set of Buggs to a date range.
+        /// </summary>
+        /// <param name="results">The unfiltered set of Buggs.</param>
+        /// <param name="dateFrom">The earliest date to filter by.</param>
+        /// <param name="dateTo">The latest date to filter by.</param>
+        /// <returns>The set of Buggs between the earliest and latest date.</returns>
+        public IQueryable<Bugg> FilterByDate(IQueryable<Bugg> results, DateTime dateFrom, DateTime dateTo)
+        {
+            dateTo = dateTo.AddDays(1);
+            var buggsInTimeFrame =
+                results.Where(bugg => (bugg.TimeOfFirstCrash >= dateFrom && bugg.TimeOfFirstCrash <= dateTo) ||
+                    (bugg.TimeOfLastCrash >= dateFrom && bugg.TimeOfLastCrash <= dateTo) || (bugg.TimeOfFirstCrash <= dateFrom && bugg.TimeOfLastCrash >= dateTo));
+
+            return buggsInTimeFrame;
+        }
+
+        /// <summary>
+        /// Filter a set of Buggs by build version.
+        /// </summary>
+        /// <param name="results">The unfiltered set of Buggs.</param>
+        /// <param name="buildVersion">The build version to filter by.</param>
+        /// <returns>The set of Buggs that matches specified build version</returns>
+        public IQueryable<Bugg> FilterByBuildVersion(IQueryable<Bugg> results, string buildVersion)
+        {
+            if (!string.IsNullOrEmpty(buildVersion))
+            {
+                results = results.Where(data => data.BuildVersion.Contains(buildVersion));
+            }
+            return results;
+        }
+
+        /// <summary>
+        /// Filter a set of Buggs by user group name.
+        /// </summary>
+        /// <param name="setOfBuggs">The unfiltered set of Buggs.</param>
+        /// <param name="groupName">The user group name to filter by.</param>
+        /// <returns>The set of Buggs by users in the requested user group.</returns>
+        public IQueryable<Bugg> FilterByUserGroup(IQueryable<Bugg> setOfBuggs, string groupName)
+        {
+            return setOfBuggs.Where(data => data.UserGroups.Any(ug => ug.Name == groupName));
+        }
+
+        /// <summary>
+        /// Sort the container of Buggs by the requested criteria.
+        /// </summary>
+        /// <param name="results">A container of unsorted Buggs.</param>
+        /// <param name="sortTerm">The term to sort by.</param>
+        /// <param name="bSortDescending">Whether to sort by descending or ascending.</param>
+        /// <param name="dateFrom">The date of the earliest Bugg to examine.</param>
+        /// <param name="dateTo">The date of the most recent Bugg to examine.</param>
+        /// <param name="groupName">The user group name to filter by.</param>
+        /// <returns>A sorted container of Buggs.</returns>
+        public IOrderedQueryable<Bugg> GetSortedResults(IQueryable<Bugg> results, string sortTerm, bool bSortDescending, DateTime dateFrom, DateTime dateTo, string groupName)
+        {
+            IOrderedQueryable<Bugg> orderedResults = null;
+
+            try
+            {
+                switch (sortTerm)
+                {
+                    case "Id":
+                        orderedResults = EnumerableOrderBy(results, buggCrashInstance => buggCrashInstance.Id, bSortDescending);
+                        break;
+
+                    case "BuildVersion":
+                        orderedResults = EnumerableOrderBy(results, buggCrashInstance => buggCrashInstance.BuildVersion, bSortDescending);
+                        break;
+
+                    case "LatestCrash":
+                        orderedResults = EnumerableOrderBy(results, buggCrashInstance => buggCrashInstance.TimeOfLastCrash, bSortDescending);
+                        break;
+
+                    case "FirstCrash":
+                        orderedResults = EnumerableOrderBy(results, buggCrashInstance => buggCrashInstance.TimeOfFirstCrash, bSortDescending);
+                        break;
+
+                    case "NumberOfCrashes":
+                        orderedResults = EnumerableOrderBy(results, buggCrashInstance => buggCrashInstance.NumberOfCrashes, bSortDescending);
+                        break;
+
+                    case "NumberOfUsers":
+                        orderedResults = EnumerableOrderBy(results, buggCrashInstance => buggCrashInstance.NumberOfUniqueMachines, bSortDescending);
+                        break;
+                    case "Pattern":
+                        orderedResults = EnumerableOrderBy(results, buggCrashInstance => buggCrashInstance.Pattern, bSortDescending);
+                        break;
+
+                    case "CrashType":
+                        orderedResults = EnumerableOrderBy(results, buggCrashInstance => buggCrashInstance.CrashType, bSortDescending);
+                        break;
+
+                    case "Status":
+                        orderedResults = EnumerableOrderBy(results, buggCrashInstance => buggCrashInstance.Status, bSortDescending);
+                        break;
+
+                    case "FixedChangeList":
+                        orderedResults = EnumerableOrderBy(results, buggCrashInstance => buggCrashInstance.FixedChangeList, bSortDescending);
+                        break;
+
+                    case "TTPID":
+                        orderedResults = EnumerableOrderBy(results, buggCrashInstance => buggCrashInstance.TTPID, bSortDescending);
+                        break;
+                    default:
+                        orderedResults = EnumerableOrderBy(results, buggCrashInstance => buggCrashInstance.Id, bSortDescending);
+                        break;
+
+                }
+                return orderedResults;
+            }
+            catch (Exception ex)
+            {
+                Debug.WriteLine("Exception in GetSortedResults: " + ex.ToString());
+            }
+            return orderedResults;
+        }
+
+        /// <summary>
+        /// 
+        /// </summary>
+        /// <typeparam name="TKey"></typeparam>
+        /// <param name="query"></param>
+        /// <param name="predicate"></param>
+        /// <param name="bDescending"></param>
+        /// <returns></returns>
+        public IOrderedQueryable<Bugg> EnumerableOrderBy<TKey>(IQueryable<Bugg> query, Expression<Func<Bugg, TKey>> predicate, bool bDescending)
+        {
+            return bDescending ? query.OrderByDescending(predicate) : query.OrderBy(predicate);
+        }
+
+        /// <summary>
+        /// Filter the list of Buggs by a callstack entry.
+        /// </summary>
+        public IQueryable<Bugg> FilterByCallstack(IQueryable<Bugg> results, string callstackEntry)
+        {
+            try
+            {
+                var queryString = HttpUtility.HtmlDecode(callstackEntry.ToString()) ?? "";
+
+                // Take out terms starting with a -
+                var terms = queryString.Split("-, ;+".ToCharArray(), StringSplitOptions.RemoveEmptyEntries);
+                var allFuncionCallIds = new HashSet<int>();
+                foreach (var term in terms)
+                {
+                    var functionCallIds = _unitOfWork.FunctionRepository.Get(functionCallInstance => functionCallInstance.Call.Contains(term)).Select(x => x.Id).ToList();
+                    foreach (var id in functionCallIds)
+                    {
+                        allFuncionCallIds.Add(id);
+                    }
+                }
+
+                // Search for all function ids. OR operation, not very efficient, but for searching for one function should be ok.
+                results = allFuncionCallIds.Aggregate(results, (current, id) => current.Where(x => x.Pattern.Contains(id + "+") || x.Pattern.Contains("+" + id)));
+            }
+            catch (Exception ex)
+            {
+                Debug.WriteLine("Exception in Search: " + ex.ToString());
+            }
+            return results;
+        }
+
+        protected override void Dispose(bool disposing)
+        {
+            if (disposing)
+            {
+                _unitOfWork.Dispose();
+            }
+        }
 	}
 }
diff --git a/Engine/Source/Programs/CrashReporter/CrashReportWebSite/Controllers/CSVController.cs b/Engine/Source/Programs/CrashReporter/CrashReportWebSite/Controllers/CSVController.cs
--- a/Engine/Source/Programs/CrashReporter/CrashReportWebSite/Controllers/CSVController.cs
+++ b/Engine/Source/Programs/CrashReporter/CrashReportWebSite/Controllers/CSVController.cs
@@ -10,8 +10,11 @@ using System.Net;
 using System.Text;
 using System.Web.Mvc;
 using System.Xml;
+using Tools.CrashReporter.CrashReportWebSite.DataModels;
+using Tools.CrashReporter.CrashReportWebSite.DataModels.Repositories;
 using Tools.CrashReporter.CrashReportWebSite.Models;
 using Tools.CrashReporter.CrashReportWebSite.Properties;
+using Tools.CrashReporter.CrashReportWebSite.ViewModels;
 using Tools.DotNETCommon;
 
 namespace Tools.CrashReporter.CrashReportWebSite.Controllers
@@ -94,7 +97,15 @@ namespace Tools.CrashReporter.CrashReportWebSite.Controllers
 		/// <summary>Fake id for all user groups</summary>
 		public static readonly int AllUserGroupId = -1;
 
-		//static readonly int MinNumberOfCrashes = 2;
+		static readonly int MinNumberOfCrashes = 2;
+
+        //Ugly instantiation of crash repository will replace with dependency injection BEFORE this gets anywhere near live.
+	    private IUnitOfWork _unitOfWork;
+
+	    public CSVController(IUnitOfWork unitOfWork)
+	    {
+	        _unitOfWork = unitOfWork;
+	    }
 
 		/// <summary>
 		/// Retrieve all Buggs matching the search criteria.
@@ -103,192 +114,175 @@ namespace Tools.CrashReporter.CrashReportWebSite.Controllers
 		/// <returns>A view to display the filtered Buggs.</returns>
 		public CSV_ViewModel GetResults( FormHelper FormData )
 		{
-			BuggRepository BuggsRepo = new BuggRepository();
-			CrashRepository CrashRepo = new CrashRepository();
-
 			using( FAutoScopedLogTimer LogTimer = new FAutoScopedLogTimer( this.GetType().ToString() ) )
 			{
-				string AnonymousGroup = "Anonymous";
-				//List<String> Users = FRepository.Get().GetUserNamesFromGroupName( AnonumousGroup );
-				int AnonymousGroupID = FRepository.Get( BuggsRepo ).FindOrAddGroup( AnonymousGroup );
-				HashSet<int> AnonumousIDs = FRepository.Get( BuggsRepo ).GetUserIdsFromUserGroup( AnonymousGroup );
-				int AnonymousID = AnonumousIDs.First();
-				HashSet<string> UserNamesForUserGroup = FRepository.Get( BuggsRepo ).GetUserNamesFromGroupName( AnonymousGroup );
-
-				// Enable to narrow results and improve debugging performance.
-				//FormData.DateFrom = FormData.DateTo.AddDays( -1 );
-				FormData.DateTo = FormData.DateTo.AddDays( 1 );
-
-				var FilteringQueryJoin = CrashRepo
-					.ListAll()
-					.Where( c => c.EpicAccountId != "" )
-					// Only crashes and asserts
-					.Where( c => c.CrashType == 1 || c.CrashType == 2 )
-					// Only anonymous user
-					.Where( c => c.UserNameId == AnonymousID )
-					// Filter be date
-					.Where( c => c.TimeOfCrash > FormData.DateFrom && c.TimeOfCrash < FormData.DateTo )
-					.Join
-					(
-						CrashRepo.Context.Buggs_Crashes,
-						c => c,
-						bc => bc.Crash,
-						( c, bc ) => new
-						{
-							GameName = c.GameName,
-							TimeOfCrash = c.TimeOfCrash.Value,
-							BuiltFromCL = c.BuiltFromCL,
-							PlatformName = c.PlatformName,
-							EngineMode = c.EngineMode,
-							MachineId = c.MachineId,
-							Module = c.Module,
-							BuildVersion = c.BuildVersion,
-							Jira = c.Jira,
-							Branch = c.Branch,
-							CrashType = c.CrashType.Value,
-							EpicId = c.EpicAccountId,
-							BuggId = bc.BuggId,
-						}
-					);
-
-				var FilteringQueryCrashes = CrashRepo
-					.ListAll()
-					.Where( c => c.EpicAccountId != "" )
-					// Only crashes and asserts
-					.Where( c => c.CrashType == 1 || c.CrashType == 2 )
-					// Only anonymous user
-					.Where( c => c.UserNameId == AnonymousID );
+			    var anonymousGroup = _unitOfWork.UserGroupRepository.First(data => data.Name == "Anonymous");
+                var anonymousGroupId = anonymousGroup.Id;
+			    var anonumousIDs = new HashSet<int>(anonymousGroup.Users.Select(data => data.Id));
+                var anonymousID = anonumousIDs.First();
+                var userNamesForUserGroup = new HashSet<string>(anonymousGroup.Users.Select(data => data.UserName));
+
+                // Enable to narrow results and improve debugging performance.
+                //FormData.DateFrom = FormData.DateTo.AddDays( -1 );
+                //FormData.DateTo = FormData.DateTo.AddDays( 1 );
+
+                var FilteringQueryJoin = _unitOfWork.CrashRepository
+                    .ListAll()
+                    .Where(c => c.EpicAccountId != "")
+                    // Only crashes and asserts
+                    .Where(c => c.CrashType == 1 || c.CrashType == 2)
+                    // Only anonymous user
+                    .Where(c => c.UserNameId == anonymousID)
+                    // Filter be date
+                    .Where(c => c.TimeOfCrash > FormData.DateFrom && c.TimeOfCrash < FormData.DateTo)
+                    .Select
+                    (
+                        c => new
+                        {
+                            GameName = c.GameName,
+                            TimeOfCrash = c.TimeOfCrash.Value,
+                            BuiltFromCL = c.ChangeListVersion,
+                            PlatformName = c.PlatformName,
+                            EngineMode = c.EngineMode,
+                            MachineId = c.ComputerName,
+                            Module = c.Module,
+                            BuildVersion = c.BuildVersion,
+                            Jira = c.Jira,
+                            Branch = c.Branch,
+                            CrashType = c.CrashType,
+                            EpicId = c.EpicAccountId,
+                            BuggId = c.Buggs.First().Id
+                        }
+                    );
+
+                var FilteringQueryCrashes = _unitOfWork.CrashRepository
+                    .ListAll()
+                    .Where(c => c.EpicAccountId != "")
+                    // Only crashes and asserts
+                    .Where(c => c.CrashType == 1 || c.CrashType == 2)
+                    // Only anonymous user
+                    .Where(c => c.UserNameId == anonymousID);
 
                 //Server timeout
-			    int TotalCrashes = CrashRepo.ListAll().Count();
+                int TotalCrashes = _unitOfWork.CrashRepository.ListAll().Count();
 
-				int TotalCrashesYearToDate = CrashRepo
-					.ListAll()
-					// Year to date
-					.Where( c => c.TimeOfCrash > new DateTime( DateTime.UtcNow.Year, 1, 1 ) )
-					.Count();
+                int TotalCrashesYearToDate = _unitOfWork.CrashRepository
+                    .ListAll().Count(c => c.TimeOfCrash > new DateTime(DateTime.UtcNow.Year, 1, 1));
 
-				var CrashesFilteredWithDateQuery = FilteringQueryCrashes
-					// Filter be date
-					.Where( c => c.TimeOfCrash > FormData.DateFrom && c.TimeOfCrash < FormData.DateTo );
+                var CrashesFilteredWithDateQuery = FilteringQueryCrashes
+                    // Filter be date
+                    .Where(c => c.TimeOfCrash > FormData.DateFrom && c.TimeOfCrash < FormData.DateTo);
 
-				int CrashesFilteredWithDate = CrashesFilteredWithDateQuery
-					.Count();
+                int CrashesFilteredWithDate = CrashesFilteredWithDateQuery
+                    .Count();
 
-				int CrashesYearToDateFiltered = FilteringQueryCrashes
-					// Year to date
-					.Where( c => c.TimeOfCrash > new DateTime( DateTime.UtcNow.Year, 1, 1 ) )
-					.Count();
+                int CrashesYearToDateFiltered = FilteringQueryCrashes.Count(c => c.TimeOfCrash > new DateTime( DateTime.UtcNow.Year, 1, 1 ));
 
                 //DB server timeout
-			    int AffectedUsersFiltered = FilteringQueryCrashes.Select( c => c.EpicAccountId ).Distinct().Count();
+                int AffectedUsersFiltered = FilteringQueryCrashes.Select( c => c.EpicAccountId ).Distinct().Count();
 
                 //DB server time out
-			    int UniqueCrashesFiltered = FilteringQueryCrashes.Select( c => c.Pattern ).Count();
-
-				int NumCrashes = FilteringQueryJoin.Count();
-
-				// Get all users
-				// SLOW
-				//var Users = FRepository.Get( BuggsRepo ).GetAnalyticsUsers().ToDictionary( X => X.EpicAccountId, Y => Y );
-
-				// Export data to the file.
-				string CSVPathname = Path.Combine( Settings.Default.CrashReporterCSV, DateTime.UtcNow.ToString( "yyyy-MM-dd.HH-mm-ss" ) );
-				CSVPathname += string
-					.Format( "__[{0}---{1}]__{2}", 
-					FormData.DateFrom.ToString( "yyyy-MM-dd" ),
-					FormData.DateTo.ToString( "yyyy-MM-dd" ), 
-					NumCrashes )
-					+ ".csv";
-
-				string ServerPath = Server.MapPath( CSVPathname );
-				var CSVFile = new StreamWriter( ServerPath, true, Encoding.UTF8 );
-
-				using (FAutoScopedLogTimer ExportToCSVTimer = new FAutoScopedLogTimer( "ExportToCSV" ))
-				{
-					var RowType = FilteringQueryJoin.FirstOrDefault().GetType();
-
-					string AllProperties = "";
-					foreach( var Property in RowType.GetProperties() )
-					{
-						AllProperties += Property.Name;
-						AllProperties += "; ";
-					}
-
-					// Write header
-					CSVFile.WriteLine( AllProperties );
-
-					foreach (var Row in FilteringQueryJoin)
-					{
-						var BVParts = Row.BuildVersion.Split( new char[] { '.' }, StringSplitOptions.RemoveEmptyEntries );
-						if (BVParts.Length > 2 && BVParts[0] != "0")
-						{
-							string CleanEngineVersion = string.Format( "{0}.{1}.{2}", BVParts[0], BVParts[1], BVParts[2] );
-
-							string[] RowProperties = new string[]
-							{
-								Row.GameName,
-								Row.TimeOfCrash.ToString(),
-								Row.BuiltFromCL,
-								Row.PlatformName,
-								Row.EngineMode,
-								Row.MachineId,
-								Row.Module,
-								CleanEngineVersion,
-								Row.Jira,
-								Row.Branch,
-								Row.CrashType == 1 ? "Crash" : "Assert",
-								Row.EpicId,
-								Row.BuggId.ToString()
-							};
-
-							string JoinedLine = string.Join( "; ", RowProperties );
-							JoinedLine += "; ";
-
-							CSVFile.WriteLine( JoinedLine );
-						}
-					}
-
-					CSVFile.Flush();
-					CSVFile.Close();
-					CSVFile = null;
-				}
-
-				List<FCSVRow> CSVRows = FilteringQueryJoin
-					.OrderByDescending( X => X.TimeOfCrash )
-					.Take( 32 )
-					.Select( c => new FCSVRow 
-					{
-						GameName = c.GameName,
-						TimeOfCrash = c.TimeOfCrash,
-						BuiltFromCL = c.BuiltFromCL,
-						PlatformName = c.PlatformName,
-						EngineMode = c.EngineMode,
-						MachineId = c.MachineId,
-						Module = c.Module,
-						BuildVersion = c.BuildVersion,
-						Jira = c.Jira,
-						Branch = c.Branch,
-						CrashType = c.CrashType,
-						EpicId = c.EpicId,
-						BuggId = c.BuggId,
-					} )
-					.ToList();
-				
-				return new CSV_ViewModel
-				{
-					CSVRows = CSVRows,
-					CSVPathname = CSVPathname,
-					DateFrom = (long)( FormData.DateFrom - CrashesViewModel.Epoch ).TotalMilliseconds,
-					DateTo = (long)( FormData.DateTo - CrashesViewModel.Epoch ).TotalMilliseconds,
-					DateTimeFrom = FormData.DateFrom,
-					DateTimeTo = FormData.DateTo,
-					AffectedUsersFiltered = AffectedUsersFiltered,
-					UniqueCrashesFiltered = UniqueCrashesFiltered,
-					CrashesFilteredWithDate = CrashesFilteredWithDate,
-					TotalCrashes = TotalCrashes,
-					TotalCrashesYearToDate = TotalCrashesYearToDate,
-				};
+                int UniqueCrashesFiltered = FilteringQueryCrashes.Select( c => c.Pattern ).Count();
+
+                int NumCrashes = FilteringQueryJoin.Count();
+                
+                // Export data to the file.
+                string CSVPathname = Path.Combine(Settings.Default.CrashReporterCSV, DateTime.UtcNow.ToString("yyyy-MM-dd.HH-mm-ss"));
+                CSVPathname += string
+                    .Format("__[{0}---{1}]__{2}",
+                    FormData.DateFrom.ToString("yyyy-MM-dd"),
+                    FormData.DateTo.ToString("yyyy-MM-dd"),
+                    NumCrashes)
+                    + ".csv";
+
+                string ServerPath = Server.MapPath(CSVPathname);
+                var CSVFile = new StreamWriter(ServerPath, true, Encoding.UTF8);
+
+                using (FAutoScopedLogTimer ExportToCSVTimer = new FAutoScopedLogTimer("ExportToCSV"))
+                {
+                    var RowType = FilteringQueryJoin.FirstOrDefault().GetType();
+
+                    string AllProperties = "";
+                    foreach (var Property in RowType.GetProperties())
+                    {
+                        AllProperties += Property.Name;
+                        AllProperties += "; ";
+                    }
+
+                    // Write header
+                    CSVFile.WriteLine(AllProperties);
+
+                    foreach (var Row in FilteringQueryJoin)
+                    {
+                        var BVParts = Row.BuildVersion.Split(new char[] { '.' }, StringSplitOptions.RemoveEmptyEntries);
+                        if (BVParts.Length > 2 && BVParts[0] != "0")
+                        {
+                            string CleanEngineVersion = string.Format("{0}.{1}.{2}", BVParts[0], BVParts[1], BVParts[2]);
+
+                            string[] RowProperties = new string[]
+                            {
+                                Row.GameName,
+                                Row.TimeOfCrash.ToString(),
+                                Row.BuiltFromCL,
+                                Row.PlatformName,
+                                Row.EngineMode,
+                                Row.MachineId,
+                                Row.Module,
+                                CleanEngineVersion,
+                                Row.Jira,
+                                Row.Branch,
+                                Row.CrashType == 1 ? "Crash" : "Assert",
+                                Row.EpicId,
+                                Row.BuggId.ToString()
+                            };
+
+                            string JoinedLine = string.Join("; ", RowProperties);
+                            JoinedLine += "; ";
+
+                            CSVFile.WriteLine(JoinedLine);
+                        }
+                    }
+
+                    CSVFile.Flush();
+                    CSVFile.Close();
+                    CSVFile = null;
+                }
+
+                List<FCSVRow> CSVRows = FilteringQueryJoin
+                    .OrderByDescending(X => X.TimeOfCrash)
+                    .Take(32)
+                    .Select(c => new FCSVRow
+                    {
+                        GameName = c.GameName,
+                        TimeOfCrash = c.TimeOfCrash,
+                        BuiltFromCL = c.BuiltFromCL,
+                        PlatformName = c.PlatformName,
+                        EngineMode = c.EngineMode,
+                        MachineId = c.MachineId,
+                        Module = c.Module,
+                        BuildVersion = c.BuildVersion,
+                        Jira = c.Jira,
+                        Branch = c.Branch,
+                        CrashType = c.CrashType,
+                        EpicId = c.EpicId,
+                        BuggId = c.BuggId,
+                    })
+                    .ToList();
+
+			    return new CSV_ViewModel()
+                {
+                    CSVRows = CSVRows,
+                    CSVPathname = CSVPathname,
+                    DateFrom = (long)(FormData.DateFrom - CrashesViewModel.Epoch).TotalMilliseconds,
+                    DateTo = (long)(FormData.DateTo - CrashesViewModel.Epoch).TotalMilliseconds,
+                    DateTimeFrom = FormData.DateFrom,
+                    DateTimeTo = FormData.DateTo,
+                    AffectedUsersFiltered = AffectedUsersFiltered,
+                    UniqueCrashesFiltered = UniqueCrashesFiltered,
+                    CrashesFilteredWithDate = CrashesFilteredWithDate,
+                    TotalCrashes = TotalCrashes,
+                    TotalCrashesYearToDate = TotalCrashesYearToDate,
+                };
 			}
 		}
 
@@ -317,12 +311,25 @@ namespace Tools.CrashReporter.CrashReportWebSite.Controllers
         /// <returns></returns>
 	    public ActionResult GenerateCSV(FormCollection Form)
         {
-            using (FAutoScopedLogTimer LogTimer = new FAutoScopedLogTimer(this.GetType().ToString(), bCreateNewLog: true))
+            using (var logTimer = new FAutoScopedLogTimer(this.GetType().ToString(), bCreateNewLog: true))
+            {
+                var formData = new FormHelper(Request, Form, "JustReport");
+                var results = GetResults(formData);
+                results.GenerationTime = logTimer.GetElapsedSeconds().ToString("F2");
+                return View("CSV", results);
+            }
+        }
+
+        /// <summary>
+        /// Holding on to entity contexts for too long causes desynchronisation errors
+        /// so dispose of the context when we're done with this controller
+        /// </summary>
+        /// <param name="disposing"></param>
+        protected override void Dispose(bool disposing)
+        {
+            if (disposing)
             {
-                FormHelper FormData = new FormHelper(Request, Form, "JustReport");
-                CSV_ViewModel Results = GetResults(FormData);
-                Results.GenerationTime = LogTimer.GetElapsedSeconds().ToString("F2");
-                return View("CSV", Results);
+                _unitOfWork.Dispose();
             }
         }
 	}
