Copying //UE4/Dev-Core to //UE4/Dev-Main (Source: //UE4/Dev-Core @ 3049602)

diff --git a/Engine/Source/Programs/CrashReporter/CrashReportWebSite/Controllers/CrashesController.cs b/Engine/Source/Programs/CrashReporter/CrashReportWebSite/Controllers/CrashesController.cs
--- a/Engine/Source/Programs/CrashReporter/CrashReportWebSite/Controllers/CrashesController.cs
+++ b/Engine/Source/Programs/CrashReporter/CrashReportWebSite/Controllers/CrashesController.cs
@@ -2,18 +2,22 @@
 
 using System;
 using System.Collections.Generic;
-using System.Data;
-using System.Diagnostics;
+using System.Data.Entity.Validation;
+using System.Globalization;
 using System.IO;
 using System.Linq;
-using System.Reflection;
-using System.Runtime.Remoting.Contexts;
 using System.Text;
+using System.Text.RegularExpressions;
+using System.Web;
 using System.Web.Mvc;
+using Tools.CrashReporter.CrashReportWebSite.DataModels;
+using Tools.CrashReporter.CrashReportWebSite.DataModels.Repositories;
+using Tools.CrashReporter.CrashReportWebSite.Properties;
+using Tools.CrashReporter.CrashReportWebSite.ViewModels;
 using Tools.DotNETCommon.XmlHandler;
 using Tools.CrashReporter.CrashReportCommon;
-using Tools.CrashReporter.CrashReportWebSite.Models;
 using System.Data.SqlClient;
+using CallStackContainer = Tools.CrashReporter.CrashReportWebSite.DataModels.CallStackContainer;
 
 namespace Tools.CrashReporter.CrashReportWebSite.Controllers
 {
@@ -23,11 +27,27 @@ namespace Tools.CrashReporter.CrashReportWebSite.Controllers
 	[HandleError]
 	public class CrashesController : Controller
 	{
+        //Ugly instantiation of crash repository will replace with dependency injection BEFORE this gets anywhere near live.
+	    private IUnitOfWork _unitOfWork;
+	    private SlackWriter _slackWriter;
+        
+        /// <summary> Special user name, currently used to mark crashes from UE4 releases. </summary>
+        const string UserNameAnonymous = "Anonymous";
+
 		/// <summary>
 		/// An empty constructor.
 		/// </summary>
-		public CrashesController()
+		public CrashesController(IUnitOfWork unitOfWork)
 		{
+		    _unitOfWork = unitOfWork;
+
+		    _slackWriter = new SlackWriter()
+		    {
+		        WebhookUrl = Settings.Default.SlackWebhookUrl,
+		        Channel = Settings.Default.SlackChannel,
+		        Username = Settings.Default.SlackUsername,
+		        IconEmoji = Settings.Default.SlackEmoji
+		    };
 		}
 
 		/// <summary>
@@ -37,49 +57,52 @@ namespace Tools.CrashReporter.CrashReportWebSite.Controllers
 		/// <returns>A view to display a list of crash reports.</returns>
 		public ActionResult Index( FormCollection CrashesForm )
 		{
-			using( FAutoScopedLogTimer LogTimer = new FAutoScopedLogTimer( this.GetType().ToString(), bCreateNewLog: true ) )
+            using( FAutoScopedLogTimer LogTimer = new FAutoScopedLogTimer( this.GetType().ToString(), bCreateNewLog: true ) )
 			{
-				CrashRepository Crashes = new CrashRepository();
-
 				// Handle any edits made in the Set form fields
 				foreach( var Entry in CrashesForm )
 				{
 					int Id = 0;
 					if( int.TryParse( Entry.ToString(), out Id ) )
 					{
-						Crash CurrentCrash = Crashes.GetCrash( Id );
-						if( CurrentCrash != null )
+                        Crash currentCrash = _unitOfWork.CrashRepository.GetById(Id);
+						if( currentCrash != null )
 						{
 							if( !string.IsNullOrEmpty( CrashesForm["SetStatus"] ) )
 							{
-								CurrentCrash.Status = CrashesForm["SetStatus"];
+								currentCrash.Status = CrashesForm["SetStatus"];
 							}
 
 							if( !string.IsNullOrEmpty( CrashesForm["SetFixedIn"] ) )
 							{
-								CurrentCrash.FixedChangeList = CrashesForm["SetFixedIn"];
+								currentCrash.FixedChangeList = CrashesForm["SetFixedIn"];
 							}
 
 							if( !string.IsNullOrEmpty( CrashesForm["SetTTP"] ) )
 							{
-								CurrentCrash.Jira = CrashesForm["SetTTP"];
+								currentCrash.Jira = CrashesForm["SetTTP"];
 							}
 						}
 					}
 
-					Crashes.SubmitChanges();
+				    _unitOfWork.Save();
 				}
 
 				// <STATUS>
 
 				// Parse the contents of the query string, and populate the form
-				FormHelper FormData = new FormHelper( Request, CrashesForm, "TimeOfCrash" );
-				CrashesViewModel Result = Crashes.GetResults( FormData );
+				var formData = new FormHelper( Request, CrashesForm, "TimeOfCrash" );
+				CrashesViewModel result = GetResults( formData );
+                result.BranchNames = _unitOfWork.CrashRepository.GetBranchesAsListItems();
+                result.VersionNames = _unitOfWork.CrashRepository.GetVersionsAsListItems();
+                result.PlatformNames = _unitOfWork.CrashRepository.GetPlatformsAsListItems();
+                result.EngineModes = _unitOfWork.CrashRepository.GetEngineModesAsListItems();
+
 
 				// Add the FromCollection to the CrashesViewModel since we don't need it for the get results function but we do want to post it back to the page.
-				Result.FormCollection = CrashesForm;
-				Result.GenerationTime = LogTimer.GetElapsedSeconds().ToString( "F2" );
-				return View( "Index", Result );
+				result.FormCollection = CrashesForm;
+				result.GenerationTime = LogTimer.GetElapsedSeconds().ToString( "F2" );
+				return View( "Index", result );
 			}
 		}
 
@@ -93,16 +116,14 @@ namespace Tools.CrashReporter.CrashReportWebSite.Controllers
 		{
 			using( FAutoScopedLogTimer LogTimer = new FAutoScopedLogTimer( this.GetType().ToString() + "(CrashId=" + id + ")", bCreateNewLog: true ) )
 			{
-				CrashRepository Crashes = new CrashRepository();
-
-				CallStackContainer CurrentCallStack = null;
+				CallStackContainer currentCallStack = null;
 
 				// Update the selected crash based on the form contents
-				Crash CurrentCrash = Crashes.GetCrash( id );
-
-				if( CurrentCrash == null )
+                var currentCrash = _unitOfWork.CrashRepository.GetById(id);
+                
+				if( currentCrash == null )
 				{
-					return RedirectToAction( "" );
+					return RedirectToAction( "Index" );
 				}
 
 				string FormValue;
@@ -110,44 +131,45 @@ namespace Tools.CrashReporter.CrashReportWebSite.Controllers
 				FormValue = CrashesForm["SetStatus"];
 				if( !string.IsNullOrEmpty( FormValue ) )
 				{
-					CurrentCrash.Status = FormValue;
+					currentCrash.Status = FormValue;
 				}
 
 				FormValue = CrashesForm["SetFixedIn"];
 				if( !string.IsNullOrEmpty( FormValue ) )
 				{
-					CurrentCrash.FixedChangeList = FormValue;
+					currentCrash.FixedChangeList = FormValue;
 				}
 
 				FormValue = CrashesForm["SetTTP"];
 				if( !string.IsNullOrEmpty( FormValue ) )
 				{
-					CurrentCrash.Jira = FormValue;
+					currentCrash.Jira = FormValue;
 				}
 
 				// Valid to set description to an empty string
 				FormValue = CrashesForm["Description"];
 				if( FormValue != null )
 				{
-					CurrentCrash.Description = FormValue;
+					currentCrash.Description = FormValue;
 				}
 
-				CurrentCallStack = new CallStackContainer( CurrentCrash );
+				currentCallStack = new CallStackContainer( currentCrash );
+			    currentCrash.Module = currentCallStack.GetModuleName();
 
 				//Set call stack properties
-				CurrentCallStack.bDisplayModuleNames = true;
-				CurrentCallStack.bDisplayFunctionNames = true;
-				CurrentCallStack.bDisplayFileNames = true;
-				CurrentCallStack.bDisplayFilePathNames = true;
-				CurrentCallStack.bDisplayUnformattedCallStack = false;
+				currentCallStack.bDisplayModuleNames = true;
+				currentCallStack.bDisplayFunctionNames = true;
+				currentCallStack.bDisplayFileNames = true;
+				currentCallStack.bDisplayFilePathNames = true;
+				currentCallStack.bDisplayUnformattedCallStack = false;
 
-				CurrentCrash.CallStackContainer = CurrentCrash.GetCallStack();
+				currentCrash.CallStackContainer = new CallStackContainer(currentCrash);
 
 				// Populate the crash with the correct user data
-				Crashes.PopulateUserInfo( CurrentCrash );
-				Crashes.SubmitChanges();
+				//_crashRepo.PopulateUserInfo( CurrentCrash );
+				//_crashRepo.SubmitChanges();
 
-				var Model = new CrashViewModel { Crash = CurrentCrash, CallStack = CurrentCallStack };
+				var Model = new CrashViewModel { Crash = currentCrash, CallStack = currentCallStack };
 				Model.GenerationTime = LogTimer.GetElapsedSeconds().ToString( "F2" );
 				return View( "Show", Model );
 			}
@@ -162,79 +184,841 @@ namespace Tools.CrashReporter.CrashReportWebSite.Controllers
 		{
 			using( FAutoScopedLogTimer LogTimer = new FAutoScopedLogTimer( this.GetType().ToString() + "(NewCrashId=" + id + ")") )
 			{
-				CrashRepository Crashes = new CrashRepository();
-
-				CrashReporterResult NewCrashResult = new CrashReporterResult();
-				NewCrashResult.ID = -1;
+				var newCrashResult = new CrashReporterResult();
+                CrashDescription newCrash;
+				newCrashResult.ID = -1;
 				string payloadString;
-                using (var reader = new StreamReader( Request.InputStream, Request.ContentEncoding ))
-				{
-                    payloadString = reader.ReadToEnd();
-                    if (string.IsNullOrEmpty(payloadString))
-                    {
-                        FLogger.Global.WriteEvent(string.Format("Add Crash Failed : Payload string empty"));
-                    }
-                }
+
+                //Read the request payload
+			    try
+			    {
+			        using (var reader = new StreamReader(Request.InputStream, Request.ContentEncoding))
+			        {
+			            payloadString = reader.ReadToEnd();
+			            if (string.IsNullOrEmpty(payloadString))
+			            {
+			                FLogger.Global.WriteEvent(string.Format("Add Crash Failed : Payload string empty"));
+			            }
+			        }
+			    }
+			    catch (Exception ex)
+			    {
+                    var messageBuilder = new StringBuilder();
+                    messageBuilder.AppendLine("Error Reading Crash Payload");
+                    messageBuilder.AppendLine("Exception was:");
+                    messageBuilder.AppendLine(ex.ToString());
+
+                    FLogger.Global.WriteException(messageBuilder.ToString());
+
+                    newCrashResult.Message = messageBuilder.ToString();
+                    newCrashResult.bSuccess = false;
+                    
+                    return Content(XmlHandler.ToXmlString<CrashReporterResult>(newCrashResult), "text/xml");
+			    }
+
+                // deserializes the payload string
+			    try
+			    {
+			        newCrash = XmlHandler.FromXmlString<CrashDescription>(payloadString);
+			    }
+			    catch (Exception ex)
+			    {
+                    var messageBuilder = new StringBuilder();
+                    messageBuilder.AppendLine("Error Reading CrashDescription XML");
+                    messageBuilder.AppendLine("Exception was: ");
+                    messageBuilder.AppendLine(ex.ToString());
+
+                    FLogger.Global.WriteException(messageBuilder.ToString());
+
+                    newCrashResult.Message = messageBuilder.ToString();
+                    newCrashResult.bSuccess = false;
+
+                    return Content(XmlHandler.ToXmlString<CrashReporterResult>(newCrashResult), "text/xml");
+			    }
                 
-                for (int Index = 0; Index < 3; Index++)
+                //Add crash to database
+			    try
+			    {
+			        var crash = CreateCrash(newCrash);
+			        newCrashResult.ID = crash.Id;
+			        newCrashResult.bSuccess = true;
+			    }
+			    catch (DbEntityValidationException dbentEx)
+			    {
+                    var messageBuilder = new StringBuilder();
+                    messageBuilder.AppendLine("Exception was:");
+                    messageBuilder.AppendLine(dbentEx.ToString());
+
+			        var innerEx = dbentEx.InnerException;
+			        while (innerEx != null)
+			        {
+			            messageBuilder.AppendLine("Inner Exception : " + innerEx.Message);
+			            innerEx = innerEx.InnerException;
+			        }
+
+			        if (dbentEx.EntityValidationErrors != null)
+			        {
+			            messageBuilder.AppendLine("Validation Errors : ");
+			            foreach (var valErr in dbentEx.EntityValidationErrors)
+			            {
+			                messageBuilder.AppendLine(valErr.ValidationErrors.Select(data => data.ErrorMessage).Aggregate((current, next) => current + "; /n" + next));
+			            }
+			        }
+
+                    messageBuilder.AppendLine("Received payload was:");
+                    messageBuilder.AppendLine(payloadString);
+
+                    FLogger.Global.WriteException(messageBuilder.ToString());
+
+                    newCrashResult.Message = messageBuilder.ToString();
+                    newCrashResult.bSuccess = false;
+			    }
+				catch (SqlException sqlExc)
 				{
-                    FLogger.Global.WriteEvent(string.Format("AddCrash : Attempt {0} of 5", Index + 1));
-					try
+					if (sqlExc.Number == -2)//If this is an sql timeout log the timeout and try again.
 					{
-                        var newCrash = XmlHandler.FromXmlString<CrashDescription>(payloadString);
-                        NewCrashResult.ID = Crashes.AddNewCrash(newCrash);
-                        NewCrashResult.bSuccess = true;
-					    break;
+						FLogger.Global.WriteEvent( string.Format( "AddCrash: Timeout" ) );
 					}
-					catch (SqlException sqlExc)
+					else
 					{
-						if (sqlExc.Number == -2)//If this is an sql timeout log the timeout and try again.
-						{
-							FLogger.Global.WriteEvent( string.Format( "AddCrash:Timeout, retrying {0} of 3", Index + 1 ) );
-						    if (Index == 2)
-						    {
-						        FLogger.Global.WriteEvent(string.Format(" AddCrash:Timeout, third consecutive timeout. Failing."));
-						        NewCrashResult.Message = "AddCrash timeout limit exceeded";
-						        NewCrashResult.bSuccess = false;
-						        break;
-						    }
-						}
-						else
-						{
-                            var messageBuilder = new StringBuilder();
-                            messageBuilder.AppendLine("Exception was:");
-                            messageBuilder.AppendLine(sqlExc.ToString());
-                            messageBuilder.AppendLine("Received payload was:");
-                            messageBuilder.AppendLine(payloadString);
+                        var messageBuilder = new StringBuilder();
+                        messageBuilder.AppendLine("Exception was:");
+                        messageBuilder.AppendLine(sqlExc.ToString());
+                        messageBuilder.AppendLine("Received payload was:");
+                        messageBuilder.AppendLine(payloadString);
 
-                            FLogger.Global.WriteException(messageBuilder.ToString());
+                        FLogger.Global.WriteException(messageBuilder.ToString());
 
-						    NewCrashResult.Message = messageBuilder.ToString();
-                            NewCrashResult.bSuccess = false;
-						    break;
-						}
+						newCrashResult.Message = messageBuilder.ToString();
+                        newCrashResult.bSuccess = false;
 					}
-					catch (Exception ex)
-					{
-						var messageBuilder = new StringBuilder();
-						messageBuilder.AppendLine("Exception was:");
-						messageBuilder.AppendLine(ex.ToString());
-						messageBuilder.AppendLine("Received payload was:");
-						messageBuilder.AppendLine(payloadString);
+				}
+				catch (Exception ex)
+				{
+					var messageBuilder = new StringBuilder();
+					messageBuilder.AppendLine("Exception was:");
+					messageBuilder.AppendLine(ex.ToString());
+					messageBuilder.AppendLine("Received payload was:");
+					messageBuilder.AppendLine(payloadString);
 
-                        FLogger.Global.WriteException(messageBuilder.ToString());
+                    FLogger.Global.WriteException(messageBuilder.ToString());
 
-						NewCrashResult.Message = messageBuilder.ToString();
-						NewCrashResult.bSuccess = false;
-                        break;
-					}
-					System.Threading.Thread.Sleep( 5000 * ( Index + 1 ) );
+					newCrashResult.Message = messageBuilder.ToString();
+					newCrashResult.bSuccess = false;
 				}
 
-				string ReturnResult = XmlHandler.ToXmlString<CrashReporterResult>( NewCrashResult );
+				string ReturnResult = XmlHandler.ToXmlString<CrashReporterResult>( newCrashResult );
                 
 				return Content( ReturnResult, "text/xml" );
 			}
 		}
-	}
+
+        /// <summary>
+        /// Gets a list of crashes filtered based on our form data.
+        /// </summary>
+        /// <param name="FormData"></param>
+        /// <returns></returns>
+        public CrashesViewModel GetResults(FormHelper FormData)
+        {
+            using (var logTimer = new FAutoScopedLogTimer(this.GetType().ToString()))
+            {
+                IEnumerable<Crash> results = null;
+                IQueryable<Crash> resultsQuery = null;
+
+                var skip = (FormData.Page - 1) * FormData.PageSize;
+                var take = FormData.PageSize;
+
+                resultsQuery = ConstructQueryForFiltering(FormData);
+
+                // Filter by data and get as enumerable.
+                resultsQuery = FilterByDate(resultsQuery, FormData.DateFrom, FormData.DateTo);
+
+                // Filter by BuggId
+                if (!string.IsNullOrEmpty(FormData.BuggId))
+                {
+                    var buggId = 0;
+                    var bValid = int.TryParse(FormData.BuggId, out buggId);
+
+                    if (bValid)
+                    {
+                        var newBugg = _unitOfWork.BuggRepository.GetById(buggId);
+
+                        if (newBugg != null)
+                        {
+                            resultsQuery = resultsQuery.Where(data => data.PatternId == newBugg.PatternId);
+                        }
+                    }
+                }
+
+                var countsQuery = resultsQuery.GroupBy(data => data.User.UserGroup).Select(data => new { Key = data.Key.Name, Count = data.Count()});
+                var groupCounts = countsQuery.OrderBy(data => data.Key).ToDictionary(data => data.Key, data => data.Count);
+
+                // Filter by user group if present
+                var userGroupId = !string.IsNullOrEmpty(FormData.UserGroup) ? _unitOfWork.UserGroupRepository.First(data => data.Name == FormData.UserGroup).Id : 1;
+                resultsQuery = resultsQuery.Where(data => data.User.UserGroupId == userGroupId);
+
+                var orderedQuery = GetSortedQuery(resultsQuery, FormData.SortTerm ?? "TimeOfCrash", FormData.SortOrder == "Descending");
+
+                // Grab just the results we want to display on this page
+                results = orderedQuery.Skip(skip).Take(take).ToList();
+
+                // Get the Count for pagination
+                var resultCount = 0;
+                using (var logTimer3 = new FScopedLogTimer("CrashRepository.Results.Users"))
+                {
+                    resultCount = orderedQuery.Count();
+                }
+
+                using (var logTimer3 = new FScopedLogTimer("CrashRepository.GetResults.GetCallstacks"))
+                {
+                    // Process call stack for display
+                    foreach (var crashInstance in results)
+                    {
+                        // Put call stacks into an list so we can access them line by line in the view
+                        crashInstance.CallStackContainer = new CallStackContainer(crashInstance);
+                    }
+                }
+
+                return new CrashesViewModel
+                {
+                    Results = results,
+                    PagingInfo = new PagingInfo { CurrentPage = FormData.Page, PageSize = FormData.PageSize, TotalResults = resultCount },
+                    SortOrder = FormData.SortOrder,
+                    SortTerm = FormData.SortTerm,
+                    UserGroup = FormData.UserGroup,
+                    CrashType = FormData.CrashType,
+                    SearchQuery = FormData.SearchQuery,
+                    UsernameQuery = FormData.UsernameQuery,
+                    EpicIdOrMachineQuery = FormData.EpicIdOrMachineQuery,
+                    MessageQuery = FormData.MessageQuery,
+                    BuiltFromCL = FormData.BuiltFromCL,
+                    BuggId = FormData.BuggId,
+                    JiraQuery = FormData.JiraQuery,
+                    DateFrom = (long)(FormData.DateFrom - CrashesViewModel.Epoch).TotalMilliseconds,
+                    DateTo = (long)(FormData.DateTo - CrashesViewModel.Epoch).TotalMilliseconds,
+                    BranchName = FormData.BranchName,
+                    VersionName = FormData.VersionName,
+                    PlatformName = FormData.PlatformName,
+                    GameName = FormData.GameName,
+                    GroupCounts = groupCounts
+                };
+            }
+        }
+
+	    /// <summary>
+	    /// Returns a lambda expression used to sort our crash data.
+	    /// </summary>
+	    /// <param name="resultsQuery">A query filter on the crashes entity.</param>
+	    /// <param name="sortTerm">Sort term identifying the field on which to sort.</param>
+	    /// <param name="sortDescending">bool indicating sort order</param>
+	    /// <returns></returns>
+	    public IOrderedQueryable<Crash> GetSortedQuery(IQueryable<Crash> resultsQuery, string sortTerm, bool sortDescending)
+	    {
+            switch (sortTerm)
+            {
+                case "Id":
+                    return sortDescending ? resultsQuery.OrderByDescending(data => data.Id) : resultsQuery.OrderBy(data => data.Id);
+                    break;
+                case "TimeOfCrash":
+                    return sortDescending ? resultsQuery.OrderByDescending(data => data.TimeOfCrash) : resultsQuery.OrderBy(crashInstance => crashInstance.TimeOfCrash);
+                    break;
+                case "UserName":
+                    return sortDescending ? resultsQuery.OrderByDescending(data => data.UserName) : resultsQuery.OrderBy(crashInstance => crashInstance.UserName);
+                    break;
+                case "RawCallStack":
+                    return sortDescending ? resultsQuery.OrderByDescending(data => data.RawCallStack) : resultsQuery.OrderBy(crashInstance => crashInstance.RawCallStack);
+                    break;
+                case "GameName":
+                    return sortDescending ? resultsQuery.OrderByDescending(data => data.GameName) : resultsQuery.OrderBy(crashInstance => crashInstance.GameName);
+                    break;
+                case "EngineMode":
+                    return sortDescending ? resultsQuery.OrderByDescending(data => data.EngineMode) : resultsQuery.OrderBy(crashInstance => crashInstance.EngineMode);
+                    break;
+                case "FixedChangeList":
+                    return sortDescending ? resultsQuery.OrderByDescending(data => data.FixedChangeList) : resultsQuery.OrderBy(crashInstance => crashInstance.FixedChangeList);
+                    break;
+                case "TTPID":
+                    return sortDescending ? resultsQuery.OrderByDescending(data => data.Jira) : resultsQuery.OrderBy(crashInstance => crashInstance.Jira);
+                    break;
+                case "Branch":
+                    return sortDescending ? resultsQuery.OrderByDescending(data => data.Branch) : resultsQuery.OrderBy(crashInstance => crashInstance.Branch);
+                    break;
+                case "ChangeListVersion":
+                    return sortDescending ? resultsQuery.OrderByDescending(data => data.ChangeListVersion) : resultsQuery.OrderBy(crashInstance => crashInstance.ChangeListVersion);
+                    break;
+                case "ComputerName":
+                    return sortDescending ? resultsQuery.OrderByDescending(data => data.ComputerName) : resultsQuery.OrderBy(crashInstance => crashInstance.ComputerName);
+                    break;
+                case "PlatformName":
+                    return sortDescending ? resultsQuery.OrderByDescending(data => data.PlatformName) : resultsQuery.OrderBy(crashInstance => crashInstance.PlatformName);
+                    break;
+                case "Status":
+                    return sortDescending ? resultsQuery.OrderByDescending(data => data.Status) : resultsQuery.OrderBy(crashInstance => crashInstance.Status);
+                    break;
+                case "Module":
+                    return sortDescending ? resultsQuery.OrderByDescending(data => data.Module) : resultsQuery.OrderBy(crashInstance => crashInstance.Module);
+                    break;
+                case "Summary":
+                    return sortDescending ? resultsQuery.OrderByDescending(data => data.Summary) : resultsQuery.OrderBy(data => data.Summary);
+            }
+	        return null;
+	    }
+
+        /// <summary>Constructs query for filtering.</summary>
+        private IQueryable<Crash> ConstructQueryForFiltering(FormHelper formData)
+        {
+            //Instead of returning a queryable and filtering it here we should construct a filtering expression and pass that to the repository.
+            //I don't like that data handling is taking place in the controller.
+            var results = _unitOfWork.CrashRepository.ListAll();
+
+            // Grab Results 
+            var queryString = HttpUtility.HtmlDecode(formData.SearchQuery);
+            if (!string.IsNullOrEmpty(queryString))
+            {
+                if (!string.IsNullOrEmpty(queryString))
+                {
+                    //We only use SearchQuery now for CallStack searching - if there's a SearchQuery value and a Username value, we need to get rid of the 
+                    //Username so that we can create a broader search range
+                    formData.UsernameQuery = "";
+                }
+
+                results = results.Where(data => data.RawCallStack.Contains(formData.SearchQuery));
+            }
+
+            // Filter by Crash Type
+            if (formData.CrashType != "All")
+            {
+                switch (formData.CrashType)
+                {
+                    case "Crashes":
+                        results = results.Where(crashInstance => crashInstance.CrashType == 1);
+                        break;
+                    case "Assert":
+                        results = results.Where(crashInstance => crashInstance.CrashType == 2);
+                        break;
+                    case "Ensure":
+                        results = results.Where(crashInstance => crashInstance.CrashType == 3);
+                        break;
+                    case "CrashesAsserts":
+                        results = results.Where(crashInstance => crashInstance.CrashType == 1 || crashInstance.CrashType == 2);
+                        break;
+                }
+            }
+
+            // JRX Restore EpicID/UserName searching
+            if (!string.IsNullOrEmpty(formData.UsernameQuery))
+            {
+                var decodedUsername = HttpUtility.HtmlDecode(formData.UsernameQuery).ToLower();
+                var user = _unitOfWork.UserRepository.First(data => data.UserName.Contains(decodedUsername));
+                if (user != null)
+                {
+                    var userId = user.Id; 
+                    results = (
+                        from crashDetail in results
+                        where crashDetail.UserNameId == userId
+                        select crashDetail
+                        );
+                }
+            }
+
+            // Start Filtering the results
+            if (!string.IsNullOrEmpty(formData.EpicIdOrMachineQuery))
+            {
+                var decodedEpicOrMachineId = HttpUtility.HtmlDecode(formData.EpicIdOrMachineQuery).ToLower();
+                results =
+                        (
+                            from crashDetail in results
+                            where crashDetail.EpicAccountId.Equals(decodedEpicOrMachineId) || crashDetail.ComputerName.Equals(decodedEpicOrMachineId)
+                            select crashDetail
+                        );
+            }
+
+            if (!string.IsNullOrEmpty(formData.JiraQuery))
+            {
+                var decodedJira = HttpUtility.HtmlDecode(formData.JiraQuery).ToLower();
+                results =
+                    (
+                        from crashDetail in results
+                        where crashDetail.Jira.Contains(decodedJira)
+                        select crashDetail
+                    );
+            }
+
+            // Filter by BranchName
+            if (!string.IsNullOrEmpty(formData.BranchName))
+            {
+                results =
+                    (
+                        from crashDetail in results
+                        where crashDetail.Branch.Equals(formData.BranchName)
+                        select crashDetail
+                    );
+            }
+
+            // Filter by VersionName
+            if (!string.IsNullOrEmpty(formData.VersionName))
+            {
+                results =
+                    (
+                        from crashDetail in results
+                        where crashDetail.BuildVersion.Equals(formData.VersionName)
+                        select crashDetail
+                    );
+            }
+
+            // Filter by VersionName
+            if (!string.IsNullOrEmpty(formData.EngineMode))
+            {
+                results =
+                    (
+                        from crashDetail in results
+                        where crashDetail.EngineMode.Equals(formData.EngineMode)
+                        select crashDetail
+                    );
+            }
+
+            // Filter by PlatformName
+            if (!string.IsNullOrEmpty(formData.PlatformName))
+            {
+                results =
+                    (
+                        from crashDetail in results
+                        where crashDetail.PlatformName.Contains(formData.PlatformName)
+                        select crashDetail
+                    );
+            }
+
+            // Filter by GameName
+            if (!string.IsNullOrEmpty(formData.GameName))
+            {
+                var DecodedGameName = HttpUtility.HtmlDecode(formData.GameName).ToLower();
+
+                if (DecodedGameName.StartsWith("-"))
+                {
+                    results =
+                    (
+                        from crashDetail in results
+                        where !crashDetail.GameName.Contains(DecodedGameName.Substring(1))
+                        select crashDetail
+                    );
+                }
+                else
+                {
+                    results =
+                    (
+                        from crashDetail in results
+                        where crashDetail.GameName.Contains(DecodedGameName)
+                        select crashDetail
+                    );
+                }
+            }
+
+            // Filter by MessageQuery
+            if (!string.IsNullOrEmpty(formData.MessageQuery))
+            {
+                results =
+                (
+                    from crashDetail in results
+                    where crashDetail.Summary.Contains(formData.MessageQuery) || crashDetail.Description.Contains(formData.MessageQuery + "%")
+                    select crashDetail
+                );
+            }
+
+            // Filter by BuiltFromCL
+            if (!string.IsNullOrEmpty(formData.BuiltFromCL))
+            {
+                var builtFromCl = 0;
+                var bValid = int.TryParse(formData.BuiltFromCL, out builtFromCl);
+
+                if (bValid)
+                {
+                    results =
+                        (
+                            from crashDetail in results
+                            where crashDetail.ChangeListVersion.Equals(formData.BuiltFromCL)
+                            select crashDetail
+                        );
+                }
+            }
+
+            return results;
+        }
+
+        /// <summary>
+        /// Filter a set of crashes to a date range.
+        /// </summary>
+        /// <param name="Results">The unfiltered set of crashes.</param>
+        /// <param name="DateFrom">The earliest date to filter by.</param>
+        /// <param name="DateTo">The latest date to filter by.</param>
+        /// <returns>The set of crashes between the earliest and latest date.</returns>
+        public IQueryable<Crash> FilterByDate(IQueryable<Crash> Results, DateTime DateFrom, DateTime DateTo)
+        {
+            using (FAutoScopedLogTimer LogTimer = new FAutoScopedLogTimer(this.GetType().ToString() + " SQL"))
+            {
+                DateTo = DateTo.AddDays(1);
+
+                IQueryable<Crash> CrashesInTimeFrame = Results
+                    .Where(MyCrash => MyCrash.TimeOfCrash >= DateFrom && MyCrash.TimeOfCrash <= DateTo);
+                return CrashesInTimeFrame;
+            }
+        }
+
+	    /// <summary>
+	    /// Create a new crash data model object and insert it into the database
+	    /// </summary>
+	    /// <param name="description"></param>
+	    /// <returns></returns>
+	    private Crash CreateCrash(CrashDescription description)
+	    {
+	        var newCrash = new Crash
+	        {
+	            Branch = description.BranchName,
+	            BaseDir = description.BaseDir,
+	            BuildVersion = description.BuildVersion,
+	            ChangeListVersion = description.BuiltFromCL.ToString(),
+	            CommandLine = description.CommandLine,
+	            EngineMode = description.EngineMode,
+	            ComputerName = description.MachineGuid
+	        };
+
+	        //if there's a valid username assign the associated UserNameId else use "anonymous".
+	        var userName = (!string.IsNullOrEmpty(description.UserName)) ? description.UserName : UserNameAnonymous;
+	        var user = _unitOfWork.UserRepository.GetByUserName(userName);
+
+	        if (user != null)
+	            newCrash.UserNameId = user.Id;
+	        else
+	        {
+	            newCrash.User = new User() {UserName = description.UserName, UserGroupId = 5};
+	        }
+
+	        //If there's a valid EpicAccountId assign that.
+	        if (!string.IsNullOrEmpty(description.EpicAccountId))
+	        {
+	            newCrash.EpicAccountId = description.EpicAccountId;
+	        }
+
+	        newCrash.Description = "";
+	        if (description.UserDescription != null)
+	        {
+	            newCrash.Description = string.Join(Environment.NewLine, description.UserDescription);
+	        }
+
+	        newCrash.EngineMode = description.EngineMode;
+	        newCrash.GameName = description.GameName;
+	        newCrash.LanguageExt = description.Language; // Converted by the crash process.
+	        newCrash.PlatformName = description.Platform;
+
+	        if (description.ErrorMessage != null)
+	        {
+	            newCrash.Summary = string.Join("\n", description.ErrorMessage);
+	        }
+
+	        if (description.CallStack != null)
+	        {
+	            newCrash.RawCallStack = string.Join("\n", description.CallStack);
+	        }
+
+	        if (description.SourceContext != null)
+	        {
+	            newCrash.SourceContext = string.Join("\n", description.SourceContext);
+	        }
+
+	        newCrash.TimeOfCrash = description.TimeofCrash;
+	        newCrash.Processed = description.bAllowToBeContacted;
+
+	        newCrash.Jira = "";
+	        newCrash.FixedChangeList = "";
+
+	        // Set the crash type
+	        newCrash.CrashType = 1;
+
+	        //if we have a crash type set the crash type
+	        if (!string.IsNullOrEmpty(description.CrashType))
+	        {
+	            switch (description.CrashType.ToLower())
+	            {
+	                case "crash":
+	                    newCrash.CrashType = 1;
+	                    break;
+	                case "assert":
+	                    newCrash.CrashType = 2;
+	                    break;
+	                case "ensure":
+	                    newCrash.CrashType = 3;
+	                    break;
+	                case "":
+	                case null:
+	                default:
+	                    newCrash.CrashType = 1;
+	                    break;
+	            }
+	        }
+	        else //else fall back to the old behavior and try to determine type from RawCallStack
+	        {
+	            if (newCrash.RawCallStack != null)
+	            {
+	                if (newCrash.RawCallStack.Contains("FDebug::AssertFailed"))
+	                {
+	                    newCrash.CrashType = 2;
+	                }
+	                else if (newCrash.RawCallStack.Contains("FDebug::Ensure"))
+	                {
+	                    newCrash.CrashType = 3;
+	                }
+	                else if (newCrash.RawCallStack.Contains("FDebug::OptionallyLogFormattedEnsureMessageReturningFalse"))
+	                {
+	                    newCrash.CrashType = 3;
+	                }
+	                else if (newCrash.RawCallStack.Contains("NewReportEnsure"))
+	                {
+	                    newCrash.CrashType = 3;
+	                }
+	            }
+	        }
+
+	        // As we're adding it, the status is always new
+	        newCrash.Status = "New";
+
+	        /*
+                Unused Crashes' fields.
+	 
+                Title				nchar(20)		
+                Selected			bit				
+                Version				int				
+                AutoReporterID		int				
+                Processed			bit	-> renamed to AllowToBeContacted			
+                HasDiagnosticsFile	bit	always true		
+                HasNewLogFile		bit				
+                HasMetaData			bit	always true
+            */
+
+	        // Set the unused fields to the default values.
+	        //NewCrash.Title = "";			removed from dbml
+	        //NewCrash.Selected = false;	removed from dbml
+	        //NewCrash.Version = 4;			removed from dbml
+	        //NewCrash.AutoReporterID = 0;	removed from dbml
+	        //NewCrash.HasNewLogFile = false;removed from dbml
+
+	        //NewCrash.HasDiagnosticsFile = true;
+	        //NewCrash.HasMetaData = true;
+            newCrash.UserActivityHint = description.UserActivityHint;
+            
+            BuildPattern(newCrash);
+
+            if(newCrash.CommandLine == null)
+                newCrash.CommandLine = "";
+
+            _unitOfWork.Dispose();
+            _unitOfWork = new UnitOfWork(new CrashReportEntities());
+	        var callStackRepository = _unitOfWork.CallstackRepository;
+
+	        try
+	        {
+	            var crashRepo = _unitOfWork.CrashRepository;
+                //if we don't have any callstack data then insert the crash and return
+	            if (string.IsNullOrEmpty(newCrash.Pattern))
+	            {
+	                crashRepo.Save(newCrash);
+                    _unitOfWork.Save();
+	                return newCrash;
+	            }
+
+                //If this isn't a new pattern then link it to our crash data model
+                if (callStackRepository.Any(data => data.Pattern == newCrash.Pattern))
+	            {
+                    var callstackPattern = callStackRepository.First(data => data.Pattern == newCrash.Pattern);
+	                newCrash.PatternId = callstackPattern.id;
+	            }
+	            else
+	            {
+	                //if this is a new callstack pattern then insert into data model and create a new bugg.
+	                var callstackPattern = new CallStackPattern { Pattern = newCrash.Pattern };
+	                callStackRepository.Save(callstackPattern);
+                    _unitOfWork.Save();
+	                newCrash.PatternId = callstackPattern.id;
+	            }
+
+                //Mask out the line number and File path from our error message.
+                var errorMessageString = description.ErrorMessage != null ? String.Join("", description.ErrorMessage) : "";
+                var fileMaskRegEx = new Regex(@"(\[File:).*?(])");
+                var lineMaskRegEx = new Regex(@"(\[Line:).*?(])");
+                var trimmedError = fileMaskRegEx.Replace(errorMessageString, "");
+                trimmedError = lineMaskRegEx.Replace(trimmedError, "");
+
+	            ErrorMessage errorMessage = null;
+
+	            if (_unitOfWork.ErrorMessageRepository.Any(data => data.Message.Contains(trimmedError)))
+	            {
+	                errorMessage =
+	                    _unitOfWork.ErrorMessageRepository.First(data => data.Message.Contains(trimmedError));
+	            }
+	            else
+	            {
+	                errorMessage = new ErrorMessage(){ Message = trimmedError };
+                    _unitOfWork.ErrorMessageRepository.Save(errorMessage);
+                    _unitOfWork.Save();
+	            }
+
+	            //Check for an existing bugg with this pattern and error message / no error message
+                if ( _unitOfWork.BuggRepository.Any(data => 
+                    (data.PatternId == newCrash.PatternId || data.Pattern == newCrash.Pattern) && 
+                    (data.ErrorMessageId == errorMessage.Id || data.ErrorMessageId == null) ))
+                {
+                    //if a bugg exists for this pattern update the bugg data
+                    var bugg = _unitOfWork.BuggRepository.First(data => data.PatternId == newCrash.PatternId) ??
+                                _unitOfWork.BuggRepository.First(data => data.Pattern == newCrash.Pattern);
+                    bugg.PatternId = newCrash.PatternId;
+                    bugg.ErrorMessageId = errorMessage.Id;
+
+	                //also update the bugg data while we're here
+	                bugg.TimeOfLastCrash = newCrash.TimeOfCrash;
+	                bugg.CrashType = newCrash.CrashType;
+	                if (String.Compare(newCrash.BuildVersion, bugg.BuildVersion,
+	                        StringComparison.Ordinal) != 1)
+	                    bugg.BuildVersion = newCrash.BuildVersion;
+
+                    _unitOfWork.Save();
+
+	                //if a bugg exists update this crash from the bugg
+	                //buggs are authoritative in this case
+	                newCrash.Buggs.Add(bugg);
+	                newCrash.Jira = bugg.TTPID;
+	                newCrash.FixedChangeList = bugg.FixedChangeList;
+	                newCrash.Status = bugg.Status;
+	                _unitOfWork.CrashRepository.Save(newCrash);
+	                _unitOfWork.Save();
+	            }
+	            else
+	            {
+	                //if there's no bugg for this pattern create a new bugg and insert into the data store.
+	                var bugg = new Bugg();
+	                bugg.TimeOfFirstCrash = newCrash.TimeOfCrash;
+	                bugg.TimeOfLastCrash = newCrash.TimeOfCrash;
+	                bugg.TTPID = newCrash.Jira;
+                    bugg.Pattern = newCrash.Pattern;
+	                bugg.PatternId = newCrash.PatternId;
+	                bugg.NumberOfCrashes = 1;
+	                bugg.NumberOfUsers = 1;
+	                bugg.NumberOfUniqueMachines = 1;
+	                bugg.BuildVersion = newCrash.BuildVersion;
+	                bugg.CrashType = newCrash.CrashType;
+	                bugg.Status = newCrash.Status;
+                    bugg.FixedChangeList = newCrash.FixedChangeList;
+	                bugg.ErrorMessageId = errorMessage.Id;
+                    newCrash.Buggs.Add(bugg);
+                    _unitOfWork.BuggRepository.Save(bugg);
+                    _unitOfWork.CrashRepository.Save(newCrash);
+                    _unitOfWork.Save();
+	            }
+	        }
+            catch (DbEntityValidationException dbentEx)
+            {
+                var messageBuilder = new StringBuilder();
+                messageBuilder.AppendLine("Db Entity Validation Exception Exception was:");
+                messageBuilder.AppendLine(dbentEx.ToString());
+
+                var innerEx = dbentEx.InnerException;
+                while (innerEx != null)
+                {
+                    messageBuilder.AppendLine("Inner Exception : " + innerEx.Message);
+                    innerEx = innerEx.InnerException;
+                }
+
+                if (dbentEx.EntityValidationErrors != null)
+                {
+                    messageBuilder.AppendLine("Validation Errors : ");
+                    foreach (var valErr in dbentEx.EntityValidationErrors)
+                    {
+                        messageBuilder.AppendLine(valErr.ValidationErrors.Select(data => data.ErrorMessage).Aggregate((current, next) => current + "; /n" + next));
+                    }
+                }
+
+                FLogger.Global.WriteException(messageBuilder.ToString());
+            }
+	        catch (Exception ex)
+	        {
+	            FLogger.Global.WriteException("Create Crash Exception : " +
+	                                          ex.Message.ToString(CultureInfo.InvariantCulture));
+                _slackWriter.Write("Create Crash Exception : " + ex.Message.ToString(CultureInfo.InvariantCulture));
+	            throw;
+	        }
+
+            return newCrash;
+	    }
+
+	    /// <summary>
+	    /// Create call stack pattern and either insert into database or match to existing.
+	    /// Update Associate Buggs.
+	    /// </summary>
+	    /// <param name="newCrash"></param>
+	    private void BuildPattern(Crash newCrash)
+	    {
+	        var callstack = new CallStackContainer(newCrash);
+	        newCrash.Module = callstack.GetModuleName();
+	        if (newCrash.PatternId == null)
+	        {
+	            var patternList = new List<string>();
+
+	            try
+	            {
+	                foreach (var entry in callstack.CallStackEntries.Take(CallStackContainer.MaxLinesToParse))
+	                {
+	                    FunctionCall currentFunctionCall;
+
+	                    var csEntry = entry;
+	                    if (_unitOfWork.FunctionRepository.Any(f => f.Call == csEntry.FunctionName))
+	                    {
+                            currentFunctionCall = _unitOfWork.FunctionRepository.First(f => f.Call == csEntry.FunctionName);
+	                    }
+	                    else
+	                    {
+	                        currentFunctionCall = new FunctionCall { Call = csEntry.FunctionName };
+	                        _unitOfWork.FunctionRepository.Save(currentFunctionCall);
+	                        _unitOfWork.Save();
+	                    }
+
+	                    patternList.Add(currentFunctionCall.Id.ToString());
+	                }
+
+	                newCrash.Pattern = string.Join("+", patternList);
+	            }
+	            catch (Exception ex)
+	            {
+                    var messageBuilder = new StringBuilder();
+                    FLogger.Global.WriteException("Build Pattern exception: " + ex.Message.ToString(CultureInfo.InvariantCulture));
+                    messageBuilder.AppendLine("Exception was:");
+                    messageBuilder.AppendLine(ex.ToString());
+                    while (ex.InnerException != null)
+                    {
+                        ex = ex.InnerException;
+                        messageBuilder.AppendLine(ex.ToString());
+                    }
+
+                    _slackWriter.Write("Build Pattern Exception : " + ex.Message.ToString(CultureInfo.InvariantCulture));
+                    throw;
+	            }
+	        }
+	    }
+
+	    protected override void Dispose(bool disposing)
+	    {
+            _unitOfWork.Dispose();
+	        base.Dispose(disposing);
+	    }
+
+	    #region Test Methods
+
+        /// <summary>
+        /// Test method that calls the add crash function and adds a test crash to the database
+        /// </summary>
+        public void TestAddCrash()
+	    {
+           var crashdescription = XmlHandler.FromXmlString<CrashDescription>(Settings.Default.TestXML);
+
+           CreateCrash(crashdescription);
+	    }
+
+	    #endregion
+    }
 }
\ No newline at end of file
diff --git a/Engine/Source/Programs/CrashReporter/CrashReportWebSite/Controllers/DashboardController.cs b/Engine/Source/Programs/CrashReporter/CrashReportWebSite/Controllers/DashboardController.cs
--- a/Engine/Source/Programs/CrashReporter/CrashReportWebSite/Controllers/DashboardController.cs
+++ b/Engine/Source/Programs/CrashReporter/CrashReportWebSite/Controllers/DashboardController.cs
@@ -5,7 +5,8 @@ using System.Collections.Generic;
 using System.Diagnostics;
 using System.Linq;
 using System.Web.Mvc;
-
+using Tools.CrashReporter.CrashReportWebSite.DataModels;
+using Tools.CrashReporter.CrashReportWebSite.DataModels.Repositories;
 using Tools.CrashReporter.CrashReportWebSite.Models;
 
 namespace Tools.CrashReporter.CrashReportWebSite.Controllers
@@ -52,8 +53,18 @@ namespace Tools.CrashReporter.CrashReportWebSite.Controllers
 		/// <summary>Fake id for all user groups</summary>
 		public static readonly int AllUserGroupId = -1;
 
-		CrashRepository _Crashes = new CrashRepository();
-		BuggRepository _Buggs = new BuggRepository();
+        //
+        CrashReportEntities _entities = new CrashReportEntities();
+	    private CrashRepository _Crashes;
+	    private BuggRepository _Buggs;
+
+	    public DashboardController()
+	    {
+	        //Hideous temporary code
+            //no! Unity! Dependency injection! Single Context wrapper!
+            _Crashes = new CrashRepository(_entities);
+            _Buggs = new BuggRepository(_entities);
+	    }
 
 		/// <summary>
 		/// Return a dictionary of crashes per group grouped by week.
@@ -67,7 +78,7 @@ namespace Tools.CrashReporter.CrashReportWebSite.Controllers
 			{
 				Dictionary<DateTime, int> Results = new Dictionary<DateTime, int>();
 
-				var UsersIds = FRepository.Get( _Crashes ).GetUserIdsFromUserGroupId( UserGroupId );
+			    var UsersIds = new HashSet<int>(_entities.Users.Distinct().Select(data => data.Id));
 
 				// Trim crashes to user group.
 				if( UserGroupId != DashboardController.AllUserGroupId )
@@ -143,7 +154,7 @@ namespace Tools.CrashReporter.CrashReportWebSite.Controllers
 			{
 				Dictionary<DateTime, int> Results = new Dictionary<DateTime, int>();
 
-				var UsersIds = FRepository.Get( _Crashes ).GetUserIdsFromUserGroupId( UserGroupId );
+				var UsersIds = _entities.UserGroups.First(data => data.Id == UserGroupId ).Users.Select(data => data.Id);
 
 				// Trim crashes to user group.
 				if( UserGroupId != DashboardController.AllUserGroupId )
@@ -221,7 +232,7 @@ namespace Tools.CrashReporter.CrashReportWebSite.Controllers
 				FAutoScopedLogTimer LogTimerSQL = new FAutoScopedLogTimer( "CrashesFilterByDate", "", "" );
 
 				// Get engine versions from the last 6 months.
-				var TempVersions = CrashRepository.GetVersions();
+				var TempVersions = _Crashes.GetVersions();
 				List<string> EngineUE4Versions = new List<string>();
 				foreach (var Version in TempVersions)
 				{
@@ -234,8 +245,10 @@ namespace Tools.CrashReporter.CrashReportWebSite.Controllers
 				// Only 4 latest version.
 				EngineUE4Versions = EngineUE4Versions.OrderByDescending( item => item ).Take( 5 ).ToList();
 
+			    var endDate = Today.AddDays(1);
+
 				IQueryable<Crash> CrashesInTimeFrame = _Crashes.ListAll()
-					.Where( MyCrash => MyCrash.TimeOfCrash >= AfewMonthsAgo && MyCrash.TimeOfCrash <= Today.AddDays( 1 ) );
+                    .Where(MyCrash => MyCrash.TimeOfCrash >= AfewMonthsAgo && MyCrash.TimeOfCrash <= endDate);
 				//IEnumerable<Crash> Crashes = FRepository.Get().Crashes.FilterByDate( FRepository.Get().Crashes.ListAll(), AfewMonthsAgo, Today );
 				var VMinimalCrashes = CrashesInTimeFrame.Select( Crash => new 
 				{ 
@@ -252,13 +265,13 @@ namespace Tools.CrashReporter.CrashReportWebSite.Controllers
 				}
 				LogTimerSQL.Dispose();
 
-				HashSet<int> AnonumousIDs = FRepository.Get( _Crashes ).GetUserIdsFromUserGroup( "Anonymous" );
+			    HashSet<int> AnonumousIDs = new HashSet<int>(_entities.UserGroups.First(data => data.Name == "Anonymous").Users.Select(data => data.Id));
 				int AnonymousID = AnonumousIDs.First();
 
-				int GeneralUserGroupId = FRepository.Get( _Crashes ).FindOrAddGroup( "General" );
-				int CoderUserGroupId = FRepository.Get( _Crashes ).FindOrAddGroup( "Coder" );
-				int EngineQAUserGroupId = FRepository.Get( _Crashes ).FindOrAddGroup( "EngineQA" );
-				int GameQAUserGroupId = FRepository.Get( _Crashes ).FindOrAddGroup( "GameQA" );
+			    int GeneralUserGroupId = 1;//FRepository.Get( _Crashes ).FindOrAddGroup( "General" );
+			    int CoderUserGroupId = 3;//FRepository.Get( _Crashes ).FindOrAddGroup( "Coder" );
+			    int EngineQAUserGroupId = 22;//FRepository.Get( _Crashes ).FindOrAddGroup( "EngineQA" );
+			    int GameQAUserGroupId = 21;//FRepository.Get( _Crashes ).FindOrAddGroup( "GameQA" );
 
 				// Weekly
 				Dictionary<DateTime, int> WeeklyGeneralResults = GetWeeklyCountsByGroup( MinimalCrashes, GeneralUserGroupId );
diff --git a/Engine/Source/Programs/CrashReporter/CrashReportWebSite/Controllers/HomeController.cs b/Engine/Source/Programs/CrashReporter/CrashReportWebSite/Controllers/HomeController.cs
--- a/Engine/Source/Programs/CrashReporter/CrashReportWebSite/Controllers/HomeController.cs
+++ b/Engine/Source/Programs/CrashReporter/CrashReportWebSite/Controllers/HomeController.cs
@@ -1,30 +1,48 @@
 ﻿// Copyright 1998-2016 Epic Games, Inc. All Rights Reserved.
-
-using System;
-using System.Diagnostics;
-using System.IO;
 using System.Web.Mvc;
-
-using Tools.DotNETCommon.XmlHandler;
-using Tools.CrashReporter.CrashReportCommon;
-using Tools.CrashReporter.CrashReportWebSite.Models;
+using Tools.CrashReporter.CrashReportWebSite.DataModels;
+using Tools.CrashReporter.CrashReportWebSite.DataModels.Repositories;
+using Tools.CrashReporter.CrashReportWebSite.ViewModels;
 
 namespace Tools.CrashReporter.CrashReportWebSite.Controllers
 {
 	/// <summary>Controls the start page.</summary>
 	public class HomeController : Controller
 	{
+	    private IUnitOfWork _unitOfWork;
+
+        /// <summary>
+        /// 
+        /// </summary>
+        /// <param name="unitOfWork"></param>
+	    public HomeController(IUnitOfWork unitOfWork)
+	    {
+	        _unitOfWork = unitOfWork;
+	    }
+
 		/// <summary>
 		/// The main view of the home page.
 		/// </summary>
 		public ActionResult Index()
 		{
-			using (FAutoScopedLogTimer LogTimer = new FAutoScopedLogTimer( this.GetType().ToString(), bCreateNewLog: true ))
+            using (var logTimer = new FAutoScopedLogTimer( this.GetType().ToString(), bCreateNewLog: true ))
 			{
-				CrashesViewModel Result = new CrashesViewModel();
-				Result.GenerationTime = LogTimer.GetElapsedSeconds().ToString( "F2" );
-				return View( "Index", Result );
+				var result = new CrashesViewModel();
+                result.BranchNames = _unitOfWork.CrashRepository.GetBranchesAsListItems();
+                result.VersionNames = _unitOfWork.CrashRepository.GetVersionsAsListItems();
+                result.PlatformNames = _unitOfWork.CrashRepository.GetPlatformsAsListItems();
+                result.EngineModes = _unitOfWork.CrashRepository.GetEngineModesAsListItems();
+				result.GenerationTime = logTimer.GetElapsedSeconds().ToString( "F2" );
+				return View( "Index", result );
 			}
 		}
+
+	    protected override void Dispose(bool disposing)
+	    {
+	        if (disposing)
+	        {
+	            _unitOfWork.Dispose();
+	        }
+	    }
 	}
 }
\ No newline at end of file
diff --git a/Engine/Source/Programs/CrashReporter/CrashReportWebSite/Controllers/ReportsController.cs b/Engine/Source/Programs/CrashReporter/CrashReportWebSite/Controllers/ReportsController.cs
--- a/Engine/Source/Programs/CrashReporter/CrashReportWebSite/Controllers/ReportsController.cs
+++ b/Engine/Source/Programs/CrashReporter/CrashReportWebSite/Controllers/ReportsController.cs
@@ -1,23 +1,19 @@
 ﻿// Copyright 1998-2016 Epic Games, Inc. All Rights Reserved.
 
 using System;
+using System.Collections;
 using System.Collections.Generic;
-using System.Diagnostics;
 using System.IO;
 using System.Linq;
-using System.Linq.Expressions;
 using System.Net;
 using System.Net.Mail;
-using System.Text;
 using System.Web;
-using System.Web.Helpers;
 using System.Web.Mvc;
 using System.Web.Routing;
-using System.Web.UI;
-using System.Xml;
-using Microsoft.Owin;
+using Tools.CrashReporter.CrashReportWebSite.DataModels;
+using Tools.CrashReporter.CrashReportWebSite.DataModels.Repositories;
 using Tools.CrashReporter.CrashReportWebSite.Models;
-using Tools.DotNETCommon;
+using Tools.CrashReporter.CrashReportWebSite.ViewModels;
 using Hangfire;
 using FormCollection = System.Web.Mvc.FormCollection;
 
@@ -40,6 +36,16 @@ namespace Tools.CrashReporter.CrashReportWebSite.Controllers
     {
         #region Public Methods
 
+        private IUnitOfWork _unitOfWork;
+
+        /// <summary>
+        /// 
+        /// </summary>
+        public ReportsController(IUnitOfWork unitOfWork)
+        {
+            _unitOfWork = unitOfWork;
+        }
+
         /// <summary>Fake id for all user groups</summary>
 		public static readonly int AllUserGroupId = -1;
         
@@ -58,9 +64,9 @@ namespace Tools.CrashReporter.CrashReportWebSite.Controllers
 				int BuggIDToBeAddedToJira = 0;
 				foreach( var Entry in ReportsForm )
 				{
-					if( Entry.ToString().Contains( Bugg.JiraSubmitName ) )
+                    if (Entry.ToString().Contains("CopyToJira-"))
 					{
-						int.TryParse( Entry.ToString().Substring( Bugg.JiraSubmitName.Length ), out BuggIDToBeAddedToJira );
+                        int.TryParse(Entry.ToString().Substring("CopyToJira-".Length), out BuggIDToBeAddedToJira);
 						break;
 					}
 				}
@@ -86,10 +92,10 @@ namespace Tools.CrashReporter.CrashReportWebSite.Controllers
                 // Handle 'CopyToJira' button
                 foreach (var Entry in ReportsForm)
                 {
-                    if (Entry.ToString().Contains(Bugg.JiraSubmitName))
+                    if (Entry.ToString().Contains("CopyToJira-"))
                     {
                         
-                        int.TryParse(Entry.ToString().Substring(Bugg.JiraSubmitName.Length), out buggIdToBeAddedToJira);
+                        int.TryParse(Entry.ToString().Substring("CopyToJira-".Length), out buggIdToBeAddedToJira);
                         break;
                     }
                 }
@@ -148,13 +154,59 @@ namespace Tools.CrashReporter.CrashReportWebSite.Controllers
             sMail.Send(message);
         }
 
+        /// <summary>
+        /// Test method - DELETE ME!
+        /// </summary>
+        public void GetJiraIssueSpec()
+        {
+            var jiraConnection = JiraConnection.Get();
+
+            if (!jiraConnection.CanBeUsed())
+                return;
+
+            var jiraResponse = jiraConnection.JiraRequest(
+                "/issue/createmeta?projectKeys=UE&issuetypeName=Bug&expand=projects.issuetypes.fields",
+                JiraConnection.JiraMethod.GET, null, HttpStatusCode.OK);
+
+            var projectSpec = JiraConnection.ParseJiraResponse(jiraResponse);
+
+            var projects = ((ArrayList)projectSpec["projects"]);
+
+            ArrayList issueTypes = null;
+
+            foreach (var project in projects)
+            {
+                var kvp = (KeyValuePair<string, object>)project;
+                if ( project == "issuetypes")
+                {
+                    issueTypes = kvp.Value as ArrayList;
+                }
+
+            }
+
+            if (issueTypes != null)
+            {
+                var bug = issueTypes[0];
+
+
+            }
+
+            //var issuetypes = projdic["issuetypes"];
+
+            //var issues = (Dictionary<string, object>) issuetypes;
+
+            //var bug = (Dictionary<string, object>)issues.First().Value;
+
+            //var fields = (Dictionary<string, object>)bug["fields"];
+        }
+
         #endregion
 
         #region Private Methods
         
         private void AddBuggJiraMapping(Bugg NewBugg, ref HashSet<string> FoundJiras, ref Dictionary<string, List<Bugg>> JiraIDtoBugg)
         {
-            string JiraID = NewBugg.Jira;
+            string JiraID = NewBugg.TTPID;
             FoundJiras.Add("key = " + JiraID);
 
             bool bAdd = !JiraIDtoBugg.ContainsKey(JiraID);
@@ -187,9 +239,6 @@ namespace Tools.CrashReporter.CrashReportWebSite.Controllers
 	    /// <returns></returns>
         private ReportsViewModel GetResults(string branchName, DateTime startDate, DateTime endDate, int BuggIDToBeAddedToJira)
 	    {
-            BuggRepository BuggsRepo = new BuggRepository();
-            CrashRepository CrashRepo = new CrashRepository();
-
 	        // It would be great to have a CSV export of this as well with buggs ID being the key I can then use to join them :)
 	        // 
 	        // Enumerate JIRA projects if needed.
@@ -202,12 +251,12 @@ namespace Tools.CrashReporter.CrashReportWebSite.Controllers
 	        {
 	            string AnonumousGroup = "Anonymous";
 	            //List<String> Users = FRepository.Get().GetUserNamesFromGroupName( AnonumousGroup );
-	            int AnonymousGroupID = FRepository.Get(BuggsRepo).FindOrAddGroup(AnonumousGroup);
-	            List<int> AnonumousIDs = FRepository.Get(BuggsRepo).GetUserIdsFromUserGroup(AnonumousGroup).ToList();
+	            int AnonymousGroupID = _unitOfWork.UserGroupRepository.First(data => data.Name == "Anonymous").Id;
+	            List<int> AnonumousIDs =
+                    _unitOfWork.UserGroupRepository.First(data => data.Name == "Anonymous").Users.Select(data => data.Id).ToList();
 	            int AnonymousID = AnonumousIDs.First();
 
-	            var CrashesQuery = CrashRepo
-	                .QueryByDate(CrashRepo.ListAll(), startDate, endDate)
+	            var CrashesQuery = _unitOfWork.CrashRepository.ListAll().Where(data => data.TimeOfCrash > startDate && data.TimeOfCrash <= endDate)
 	                // Only crashes and asserts
 	                .Where(Crash => Crash.CrashType == 1 || Crash.CrashType == 2)
 	                // Only anonymous user
@@ -224,24 +273,23 @@ namespace Tools.CrashReporter.CrashReportWebSite.Controllers
 	                        );
 	            }
 
-	            var Crashes = CrashesQuery.Select(Crash => new
+	            var crashes = CrashesQuery.Select(data => new
 	            {
-	                ID = Crash.Id,
-	                TimeOfCrash = Crash.TimeOfCrash.Value,
+                    ID = data.Id,
+                    TimeOfCrash = data.TimeOfCrash.Value,
 	                //UserID = Crash.UserNameId.Value, 
-	                BuildVersion = Crash.BuildVersion,
-	                JIRA = Crash.Jira,
-	                Platform = Crash.PlatformName,
-	                FixCL = Crash.FixedChangeList,
-	                BuiltFromCL = Crash.BuiltFromCL,
-	                Pattern = Crash.Pattern,
-	                MachineID = Crash.MachineId,
-	                Branch = Crash.Branch,
-	                Description = Crash.Description,
-	                RawCallStack = Crash.RawCallStack,
-	            })
-	                .ToList();
-	            int NumCrashes = Crashes.Count;
+                    BuildVersion = data.BuildVersion,
+                    JIRA = data.Jira,
+                    Platform = data.PlatformName,
+                    FixCL = data.FixedChangeList,
+                    BuiltFromCL = data.ChangeListVersion,
+                    Pattern = data.Pattern,
+                    MachineID = data.ComputerName,
+                    Branch = data.Branch,
+                    Description = data.Description,
+                    RawCallStack = data.RawCallStack,
+	            }).ToList();
+                var numCrashes = CrashesQuery.Count();
 
 	            /*
                 // Build patterns for crashes where patters is null.
@@ -259,7 +307,7 @@ namespace Tools.CrashReporter.CrashReportWebSite.Controllers
                 */
 
 	            // Total # of ALL (Anonymous) crashes in timeframe
-	            int TotalAnonymousCrashes = NumCrashes;
+	            int TotalAnonymousCrashes = numCrashes;
 
 	            // Total # of UNIQUE (Anonymous) crashes in timeframe
 	            HashSet<string> UniquePatterns = new HashSet<string>();
@@ -269,7 +317,7 @@ namespace Tools.CrashReporter.CrashReportWebSite.Controllers
 	            //List<int> CrashesWithoutPattern = new List<int>();
 	            //List<DateTime> CrashesWithoutPatternDT = new List<DateTime>();
 
-	            foreach (var Crash in Crashes)
+	            foreach (var Crash in crashes)
 	            {
 	                if (string.IsNullOrEmpty(Crash.Pattern))
 	                {
@@ -300,7 +348,7 @@ namespace Tools.CrashReporter.CrashReportWebSite.Controllers
 	            // Total # of AFFECTED USERS (Anonymous) in timeframe
 	            int TotalAffectedUsers = UniqueMachines.Count;
 
-	            var RealBuggs = BuggsRepo.Context.Buggs.Where(Bugg => PatternAndCount.Keys.Contains(Bugg.Pattern)).ToList();
+	            var RealBuggs = _unitOfWork.BuggRepository.ListAll().Where(Bugg => PatternAndCount.Keys.Contains(Bugg.Pattern));
 
 	            // Build search string.
 	            HashSet<string> FoundJiras = new HashSet<string>();
@@ -309,14 +357,14 @@ namespace Tools.CrashReporter.CrashReportWebSite.Controllers
 	            List<Bugg> Buggs = new List<Bugg>(NumTopRecords);
 	            foreach (var Top in PatternAndCount)
 	            {
-	                Bugg NewBugg = RealBuggs.Where(X => X.Pattern == Top.Key).FirstOrDefault();
+	                Bugg NewBugg = RealBuggs.FirstOrDefault(X => X.Pattern == Top.Key);
 	                if (NewBugg != null)
 	                {
 	                    using (
 	                        FAutoScopedLogTimer TopTimer =
 	                            new FAutoScopedLogTimer(string.Format("{0}:{1}", Buggs.Count + 1, NewBugg.Id)))
 	                    {
-	                        var CrashesForBugg = Crashes.Where(Crash => Crash.Pattern == Top.Key).ToList();
+	                        var CrashesForBugg = crashes.Where(Crash => Crash.Pattern == Top.Key).ToList();
 
 	                        // Convert anonymous to full type;
 	                        var FullCrashesForBugg = new List<Crash>(CrashesForBugg.Count);
@@ -324,15 +372,15 @@ namespace Tools.CrashReporter.CrashReportWebSite.Controllers
 	                        {
 	                            FullCrashesForBugg.Add(new Crash()
 	                            {
-	                                ID = Anon.ID,
+	                                Id = Anon.ID,
 	                                TimeOfCrash = Anon.TimeOfCrash,
 	                                BuildVersion = Anon.BuildVersion,
 	                                Jira = Anon.JIRA,
-	                                Platform = Anon.Platform,
-	                                FixCL = Anon.FixCL,
-	                                BuiltFromCL = Anon.BuiltFromCL,
+	                                PlatformName = Anon.Platform,
+	                                FixedChangeList = Anon.FixCL,
+	                                ChangeListVersion = Anon.BuiltFromCL,
 	                                Pattern = Anon.Pattern,
-	                                MachineId = Anon.MachineID,
+	                                ComputerName = Anon.MachineID,
 	                                Branch = Anon.Branch,
 	                                Description = Anon.Description,
 	                                RawCallStack = Anon.RawCallStack,
@@ -342,10 +390,10 @@ namespace Tools.CrashReporter.CrashReportWebSite.Controllers
 	                        NewBugg.PrepareBuggForJira(FullCrashesForBugg);
 
 	                        // Verify valid JiraID, this may be still a TTP 
-	                        if (!string.IsNullOrEmpty(NewBugg.Jira))
+	                        if (!string.IsNullOrEmpty(NewBugg.TTPID))
 	                        {
 	                            int TTPID = 0;
-	                            int.TryParse(NewBugg.Jira, out TTPID);
+	                            int.TryParse(NewBugg.TTPID, out TTPID);
 
 	                            if (TTPID == 0)
 	                            {
@@ -513,7 +561,7 @@ namespace Tools.CrashReporter.CrashReportWebSite.Controllers
 
 	                // If there are buggs, we need to update the summary to indicate an error.
 	                // Usually caused when bugg's project has changed.
-	                foreach (var Bugg in BuggsCopy.Where(b => !string.IsNullOrEmpty(b.Jira)))
+	                foreach (var Bugg in BuggsCopy.Where(b => !string.IsNullOrEmpty(b.TTPID)))
 	                {
 	                    Bugg.JiraSummary = "JIRA MISMATCH";
 	                    Bugg.JiraComponentsText = "JIRA MISMATCH";
@@ -529,7 +577,7 @@ namespace Tools.CrashReporter.CrashReportWebSite.Controllers
 	            {
 	                Buggs = Buggs,
                     BranchName = branchName,
-	                /*Crashes = Crashes,*/
+	                BranchNames = _unitOfWork.CrashRepository.GetBranchesAsListItems(),
 	                DateFrom = (long) (startDate - CrashesViewModel.Epoch).TotalMilliseconds,
 	                DateTo = (long) (endDate - CrashesViewModel.Epoch).TotalMilliseconds,
 	                TotalAffectedUsers = TotalAffectedUsers,
diff --git a/Engine/Source/Programs/CrashReporter/CrashReportWebSite/Controllers/UsersController.cs b/Engine/Source/Programs/CrashReporter/CrashReportWebSite/Controllers/UsersController.cs
--- a/Engine/Source/Programs/CrashReporter/CrashReportWebSite/Controllers/UsersController.cs
+++ b/Engine/Source/Programs/CrashReporter/CrashReportWebSite/Controllers/UsersController.cs
@@ -1,12 +1,10 @@
 ﻿// Copyright 1998-2016 Epic Games, Inc. All Rights Reserved.
 
-using System;
 using System.Collections.Generic;
-using System.Diagnostics;
 using System.Linq;
-using System.Web;
 using System.Web.Mvc;
-
+using Tools.CrashReporter.CrashReportWebSite.DataModels;
+using Tools.CrashReporter.CrashReportWebSite.DataModels.Repositories;
 using Tools.CrashReporter.CrashReportWebSite.Models;
 
 namespace Tools.CrashReporter.CrashReportWebSite.Controllers
@@ -16,6 +14,13 @@ namespace Tools.CrashReporter.CrashReportWebSite.Controllers
 	/// </summary>
 	public class UsersController : Controller
 	{
+	    private IUnitOfWork _unitOfWork;
+
+	    public UsersController(IUnitOfWork unitOfWork)
+	    {
+	        _unitOfWork = unitOfWork;
+	    }
+        
 		/// <summary>
 		/// Display the main user groups.
 		/// </summary>
@@ -24,21 +29,40 @@ namespace Tools.CrashReporter.CrashReportWebSite.Controllers
 		/// <returns>A view to show a list of users in the current user group.</returns>
 		public ActionResult Index( FormCollection UserForms, string UserGroup )
 		{
-			CrashRepository Crashes = new CrashRepository();
-
 			// Examine an incoming form for a new user group
 			foreach( var FormInstance in UserForms )
 			{
-				string UserName = FormInstance.ToString();
-				string NewUserGroup = UserForms[UserName];
-				FRepository.Get( Crashes ).SetUserGroup( UserName, NewUserGroup );
+				var UserName = FormInstance.ToString();
+				var NewUserGroup = UserForms[UserName];
 			}
 
 			UsersViewModel Model = new UsersViewModel();
 
+		    var userGroup = _unitOfWork.UserGroupRepository.First(data => data.Name == UserGroup);
+
 			Model.UserGroup = UserGroup;
-			Model.Users = FRepository.Get( Crashes ).GetUserNamesFromGroupName( UserGroup );
-			Model.GroupCounts = FRepository.Get( Crashes ).GetCountsByGroup();
+		    Model.Users =
+		        new HashSet<string>(
+		            _unitOfWork.UserRepository.ListAll()
+		                .Where(data => data.UserGroupId == userGroup.Id)
+		                .Select(data => data.UserName));
+
+		    var groupCounts =
+		        _unitOfWork.UserRepository.ListAll()
+		            .GroupBy(data => data.UserGroup)
+		            .Select(data => new {Key = data.Key, Count = data.Count()}).ToDictionary(groupCount => groupCount.Key.Name, groupCount => groupCount.Count);;
+            
+            // Add in all groups, even though there are no crashes associated
+			//IEnumerable<string> UserGroups = ( from UserGroupDetail in _entities.UserGroups select UserGroupDetail.Name );
+            //foreach( string UserGroupName in UserGroups )
+            //{
+            //    if( !groupCounts.Keys.Contains( UserGroupName ) )
+            //    {
+            //        groupCounts[UserGroupName] = 0;
+            //    }
+            //}
+
+            Model.GroupCounts = groupCounts;
 
 			return View( "Index", Model );
 		}
diff --git a/Engine/Source/Programs/CrashReporter/CrashReportWebSite/CrashReportWebSite.csproj b/Engine/Source/Programs/CrashReporter/CrashReportWebSite/CrashReportWebSite.csproj
--- a/Engine/Source/Programs/CrashReporter/CrashReportWebSite/CrashReportWebSite.csproj
+++ b/Engine/Source/Programs/CrashReporter/CrashReportWebSite/CrashReportWebSite.csproj
@@ -29,9 +29,9 @@
     <OldToolsVersion>4.0</OldToolsVersion>
     <UpgradeBackupLocation />
     <IISExpressSSLPort />
-    <IISExpressAnonymousAuthentication />
-    <IISExpressWindowsAuthentication />
-    <IISExpressUseClassicPipelineMode />
+    <IISExpressAnonymousAuthentication>enabled</IISExpressAnonymousAuthentication>
+    <IISExpressWindowsAuthentication>enabled</IISExpressWindowsAuthentication>
+    <IISExpressUseClassicPipelineMode>false</IISExpressUseClassicPipelineMode>
     <MvcProjectUpgradeChecked>true</MvcProjectUpgradeChecked>
   </PropertyGroup>
   <PropertyGroup Condition=" '$(Configuration)|$(Platform)' == 'Debug|AnyCPU' ">
@@ -54,6 +54,14 @@
       <HintPath>packages\Common.Logging.Core.3.3.1\lib\net40\Common.Logging.Core.dll</HintPath>
       <Private>True</Private>
     </Reference>
+    <Reference Include="EntityFramework, Version=6.0.0.0, Culture=neutral, PublicKeyToken=b77a5c561934e089, processorArchitecture=MSIL">
+      <HintPath>packages\EntityFramework.6.1.3\lib\net45\EntityFramework.dll</HintPath>
+      <Private>True</Private>
+    </Reference>
+    <Reference Include="EntityFramework.SqlServer, Version=6.0.0.0, Culture=neutral, PublicKeyToken=b77a5c561934e089, processorArchitecture=MSIL">
+      <HintPath>packages\EntityFramework.6.1.3\lib\net45\EntityFramework.SqlServer.dll</HintPath>
+      <Private>True</Private>
+    </Reference>
     <Reference Include="Hangfire.Core, Version=1.5.3.0, Culture=neutral, processorArchitecture=MSIL">
       <HintPath>packages\Hangfire.Core.1.5.3\lib\net45\Hangfire.Core.dll</HintPath>
       <Private>True</Private>
@@ -109,7 +117,6 @@
     <Reference Include="System" />
     <Reference Include="System.Data" />
     <Reference Include="System.Data.DataSetExtensions" />
-    <Reference Include="System.Data.Entity" />
     <Reference Include="System.Data.Linq">
       <Private>True</Private>
     </Reference>
@@ -117,6 +124,8 @@
     <Reference Include="System.ComponentModel.DataAnnotations">
       <RequiredTargetFramework>3.5</RequiredTargetFramework>
     </Reference>
+    <Reference Include="System.Runtime.Serialization" />
+    <Reference Include="System.Security" />
     <Reference Include="System.Web.ApplicationServices" />
     <Reference Include="System.Web.DynamicData" />
     <Reference Include="System.Web.Entity" />
@@ -177,27 +186,67 @@
     <Compile Include="Controllers\HomeController.cs" />
     <Compile Include="Controllers\ReportsController.cs" />
     <Compile Include="Controllers\UsersController.cs" />
-    <Compile Include="Models\BuggsViewModel.cs" />
-    <Compile Include="Models\CachedDataService.cs" />
-    <Compile Include="Models\BuggRepository.cs" />
-    <Compile Include="Models\BuggViewModel.cs" />
-    <Compile Include="Models\CSVViewModel.cs" />
-    <Compile Include="Models\DashboardViewModel.cs" />
-    <Compile Include="Models\CrashReport.cs">
-      <DependentUpon>CrashReport.dbml</DependentUpon>
+    <Compile Include="DataModels\Bugg.cs">
+      <DependentUpon>CrashReports.tt</DependentUpon>
+    </Compile>
+    <Compile Include="DataModels\BuggDerived.cs" />
+    <Compile Include="DataModels\Buggs_Users.cs">
+      <DependentUpon>CrashReports.tt</DependentUpon>
+    </Compile>
+    <Compile Include="DataModels\CallStackPattern.cs">
+      <DependentUpon>CrashReports.tt</DependentUpon>
+    </Compile>
+    <Compile Include="DataModels\Crash.cs">
+      <DependentUpon>CrashReports.tt</DependentUpon>
+    </Compile>
+    <Compile Include="DataModels\CrashReports.Context.cs">
+      <AutoGen>True</AutoGen>
+      <DesignTime>True</DesignTime>
+      <DependentUpon>CrashReports.Context.tt</DependentUpon>
+    </Compile>
+    <Compile Include="DataModels\CrashReports.cs">
+      <AutoGen>True</AutoGen>
+      <DesignTime>True</DesignTime>
+      <DependentUpon>CrashReports.tt</DependentUpon>
     </Compile>
-    <Compile Include="Models\CrashReport.designer.cs">
+    <Compile Include="DataModels\CrashReports.Designer.cs">
       <AutoGen>True</AutoGen>
       <DesignTime>True</DesignTime>
-      <DependentUpon>CrashReport.dbml</DependentUpon>
+      <DependentUpon>CrashReports.edmx</DependentUpon>
     </Compile>
-    <Compile Include="Models\CrashViewModel.cs" />
+    <Compile Include="DataModels\ErrorMessage.cs">
+      <DependentUpon>CrashReports.tt</DependentUpon>
+    </Compile>
+    <Compile Include="DataModels\FunctionCall.cs">
+      <DependentUpon>CrashReports.tt</DependentUpon>
+    </Compile>
+    <Compile Include="DataModels\Repositories\BuggRespository.cs" />
+    <Compile Include="DataModels\Repositories\CallStackRepository.cs" />
+    <Compile Include="DataModels\CrashDerived.cs" />
+    <Compile Include="DataModels\Repositories\CrashRepository.cs" />
+    <Compile Include="DataModels\Repositories\DataRepository.cs" />
+    <Compile Include="DataModels\Repositories\ErrorMessageRepository.cs" />
+    <Compile Include="DataModels\Repositories\FunctionRepository.cs" />
+    <Compile Include="DataModels\Repositories\IBuggRepository.cs" />
+    <Compile Include="DataModels\Repositories\ICallstackRepository.cs" />
+    <Compile Include="DataModels\Repositories\ICrashRepository.cs" />
+    <Compile Include="DataModels\Repositories\IDataRepository.cs" />
+    <Compile Include="DataModels\Repositories\IErrorMessageRepository.cs" />
+    <Compile Include="DataModels\Repositories\IFunctionRepository.cs" />
+    <Compile Include="DataModels\Repositories\IUserGroupRepository.cs" />
+    <Compile Include="DataModels\Repositories\IUserRepository.cs" />
+    <Compile Include="DataModels\Repositories\UnitOfWork.cs" />
+    <Compile Include="DataModels\Repositories\UserGroupRepository.cs" />
+    <Compile Include="DataModels\Repositories\UserRepository.cs" />
+    <Compile Include="DataModels\User.cs">
+      <DependentUpon>CrashReports.tt</DependentUpon>
+    </Compile>
+    <Compile Include="DataModels\UserGroup.cs">
+      <DependentUpon>CrashReports.tt</DependentUpon>
+    </Compile>
+    <Compile Include="Models\CachedDataService.cs" />
     <Compile Include="Models\EmailSubscriptionResponseModel.cs" />
-    <Compile Include="Models\FormHelper.cs" />
     <Compile Include="Models\PerformanceTimers.cs" />
-    <Compile Include="Models\ReportsViewModel.cs" />
-    <Compile Include="Models\Repository.cs" />
-    <Compile Include="Models\UsersViewModel.cs" />
     <Compile Include="Properties\Settings.Designer.cs">
       <AutoGen>True</AutoGen>
       <DesignTimeSharedInput>True</DesignTimeSharedInput>
@@ -210,10 +259,17 @@
       <DependentUpon>Global.asax</DependentUpon>
     </Compile>
     <Compile Include="Models\CrashDataFormater.cs" />
-    <Compile Include="Models\CrashesViewModel.cs" />
-    <Compile Include="Models\CrashRepository.cs" />
-    <Compile Include="Models\PagingInfo.cs" />
     <Compile Include="Properties\AssemblyInfo.cs" />
+    <Compile Include="ViewModels\BuggsViewModel.cs" />
+    <Compile Include="ViewModels\BuggViewModel.cs" />
+    <Compile Include="ViewModels\CrashesViewModel.cs" />
+    <Compile Include="ViewModels\CrashViewModel.cs" />
+    <Compile Include="ViewModels\CSVViewModel.cs" />
+    <Compile Include="ViewModels\DashboardViewModel.cs" />
+    <Compile Include="ViewModels\FormHelper.cs" />
+    <Compile Include="ViewModels\PagingInfo.cs" />
+    <Compile Include="ViewModels\ReportsViewModel.cs" />
+    <Compile Include="ViewModels\UsersViewModel.cs" />
     <Compile Include="Views\Helpers\UrlHelperExtension.cs" />
     <Compile Include="Views\Helpers\PagingHelper.cs" />
   </ItemGroup>
@@ -270,6 +326,16 @@
     <Content Include="Content\themes\base\jquery.ui.slider.css" />
     <Content Include="Content\themes\base\jquery.ui.tabs.css" />
     <Content Include="Content\themes\base\jquery.ui.theme.css" />
+    <Content Include="DataModels\CrashReports.Context.tt">
+      <Generator>TextTemplatingFileGenerator</Generator>
+      <DependentUpon>CrashReports.edmx</DependentUpon>
+      <LastGenOutput>CrashReports.Context.cs</LastGenOutput>
+    </Content>
+    <Content Include="DataModels\CrashReports.tt">
+      <Generator>TextTemplatingFileGenerator</Generator>
+      <DependentUpon>CrashReports.edmx</DependentUpon>
+      <LastGenOutput>CrashReports.cs</LastGenOutput>
+    </Content>
     <Content Include="Global.asax" />
     <Content Include="Properties\PublishProfiles\CRTest.pubxml.user" />
     <Content Include="Scripts\jquery-1.5.1-vsdoc.js" />
@@ -301,6 +367,13 @@
     <Content Include="Views\Crashes\ViewCrash.ascx" />
     <Content Include="Views\Home\Index.aspx" />
     <Content Include="Views\Reports\Index.aspx" />
+    <EntityDeploy Include="DataModels\CrashReports.edmx">
+      <Generator>EntityModelCodeGenerator</Generator>
+      <LastGenOutput>CrashReports.Designer.cs</LastGenOutput>
+    </EntityDeploy>
+    <Content Include="DataModels\CrashReports.edmx.diagram">
+      <DependentUpon>CrashReports.edmx</DependentUpon>
+    </Content>
     <None Include="Views\Reports\SignUpResponse.cshtml" />
     <Content Include="Views\Reports\Unsubscribe.aspx" />
     <Content Include="Views\Reports\ViewReports.ascx" />
@@ -330,13 +403,9 @@
   </ItemGroup>
   <ItemGroup>
     <Service Include="{3259AA49-8AA1-44D3-9025-A0B520596A8C}" />
+    <Service Include="{508349B6-6B84-4DF5-91F0-309BEEBAD82D}" />
   </ItemGroup>
   <ItemGroup>
-    <None Include="Models\CrashReport.dbml">
-      <Generator>MSLinqToSQLGenerator</Generator>
-      <LastGenOutput>CrashReport.designer.cs</LastGenOutput>
-      <SubType>Designer</SubType>
-    </None>
     <Content Include="packages.config" />
     <None Include="Properties\PublishProfiles\CrashReportDev.pubxml" />
     <None Include="Properties\PublishProfiles\CRTest.pubxml" />
@@ -349,11 +418,6 @@
       <DependentUpon>Web.config</DependentUpon>
     </None>
   </ItemGroup>
-  <ItemGroup>
-    <None Include="Models\CrashReport.dbml.layout">
-      <DependentUpon>CrashReport.dbml</DependentUpon>
-    </None>
-  </ItemGroup>
   <ItemGroup>
     <ProjectReference Include="..\..\DotNETCommon\DotNETUtilities\DotNETUtilities.csproj">
       <Project>{5d7d66e8-8c76-4af9-b3ec-2ef03421d730}</Project>
