Copying //UE4/Dev-Core to //UE4/Dev-Main (Source: //UE4/Dev-Core @ 3151653)

==========================
MAJOR FEATURES + CHANGES
==========================

Change 2975891 on 2016/05/12 by [REDACTED_PERSON]

	merged in new async stuff from dev-rendering.

Change 2976695 on 2016/05/13 by [REDACTED_PERSON]

	updated precache list

Change 2977030 on 2016/05/13 by [REDACTED_PERSON]

	Added time slicing to CreateAsyncPackagesFromQueue, radically reduced the frequency of "precache trimming" and changed a few things in the test rig and logging

Change 2977090 on 2016/05/13 by [REDACTED_PERSON]

	Fixed module manager threading and added cmd line param to force async loading thread.

Change 2977292 on 2016/05/13 by [REDACTED_PERSON]

	check for thread safety in looking at asset registry

Change 2977296 on 2016/05/13 by [REDACTED_PERSON]

	removed some super-expensive check()s from precacher

Change 2978368 on 2016/05/16 by [REDACTED_PERSON]

	Move several exposive bools inside of the basic tests inside of FLinkerLoad::Preload, saves a fraction of second.

Change 2978414 on 2016/05/16 by [REDACTED_PERSON]

	Added support and testing for unmounting pak files to the pak precacher.

Change 2978446 on 2016/05/16 by [REDACTED_PERSON]

	Allow linker listing in non-shipping builds

Change 2978550 on 2016/05/16 by [REDACTED_PERSON]

	Allowed some linker spew in non-shipping builds (instead of debug builds). Some tweak to help track down the music.uasset leak.

Change 2979952 on 2016/05/17 by [REDACTED_PERSON]

	Merging //UE4/Dev-Core @ 2979938 to Dev-UE-30519-LoadTimes

Change 2984927 on 2016/05/20 by [REDACTED_PERSON]

	fix a few bugs with an mcp repro

Change 2984951 on 2016/05/20 by [REDACTED_PERSON]

	fixed issues with USE_NEW_ASYNC_IO = 0

Change 2985296 on 2016/05/20 by [REDACTED_PERSON]

	Fixed several bugs with the MCP boot test

Change 2987956 on 2016/05/24 by [REDACTED_PERSON]

	Fixing leaked linkers created by blocking load requests during async loading.

Change 2987959 on 2016/05/24 by [REDACTED_PERSON]

	Enable load timings in block loading also (in addition to async loading).

Change 3017713 on 2016/06/17 by [REDACTED_PERSON]

	Removing GUseSeekFreeLoading.

Change 3017722 on 2016/06/17 by [REDACTED_PERSON]

	Renaming LOAD_SeekFree flag to LOAD_Async to better reflect its current purpose.

Change 3017833 on 2016/06/17 by [REDACTED_PERSON]

	Merging //UE4/Dev-Core to Dev-UE-30519-LoadTimes (//Tasks/Dev-Core/Dev-UE-30519-LoadTimes)

Change 3017840 on 2016/06/17 by [REDACTED_PERSON]

	Re-doing Dev-Core changes to Delegates 2/2

Change 3022872 on 2016/06/22 by [REDACTED_PERSON]

	reorder memory trim and deleting loaders

Change 3059218 on 2016/07/21 by [REDACTED_PERSON]

	Fixing compilation errors - adding missing load time tracker stats.

Change 3064508 on 2016/07/26 by [REDACTED_PERSON]

	Removing blocking loading path in cooked builds. LoadPackage will now use the async path.

Change 3066312 on 2016/07/27 by [REDACTED_PERSON]

	Event driven loader, first pass

Change 3066785 on 2016/07/27 by [REDACTED_PERSON]

	Removed check...searching forward for export fusion can release a node

Change 3068118 on 2016/07/28 by [REDACTED_PERSON]

	critical bug fixes for the event driven loader

Change 3068333 on 2016/07/28 by [REDACTED_PERSON]

	correctly handle the case where a file is rejected after loading the summary

Change 3069618 on 2016/07/28 by [REDACTED_PERSON]

	Merging //UE4/Dev-Core to Dev-UE-30519-LoadTimes (//Tasks/Dev-Core/Dev-UE-30519-LoadTimes)

Change 3069901 on 2016/07/29 by [REDACTED_PERSON]

	Fixing an hang when loading QA-Blueprints level

Change 3070171 on 2016/07/29 by [REDACTED_PERSON]

	fixed CDO cyclic dependencies

Change 3075288 on 2016/08/03 by [REDACTED_PERSON]

	misc fixes to the event driven loader

Change 3077332 on 2016/08/04 by [REDACTED_PERSON]

	Fixing checkSlow asserts caused by new loading code not being flagged as IsInAsyncLoadThread() and CreateSynchEvent deprecation warning.

Change 3078113 on 2016/08/04 by [REDACTED_PERSON]

	implemented "nicks rule" and undid some previous material and world hacks needed without it.

Change 3079480 on 2016/08/05 by [REDACTED_PERSON]

	fixes and tweaks on event driven loader

Change 3080135 on 2016/08/07 by [REDACTED_PERSON]

	misc fixes for event driven loader, now with reasonable memory

Change 3083722 on 2016/08/10 by [REDACTED_PERSON]

	Fixing hangs when async loading packages.

Change 3091747 on 2016/08/17 by [REDACTED_PERSON]

	Fix all hitches in streaming load that were regressions.

Change 3093258 on 2016/08/18 by [REDACTED_PERSON]

	Fix bug that caused an assert when packages fail to load for certain reasons (like loading an uncooked file).

Change 3095719 on 2016/08/20 by [REDACTED_PERSON]

	reenable async loading thread and cleanup and bug fixes

Change 3096350 on 2016/08/22 by [REDACTED_PERSON]

	tweak task priorities a bit to minimize precaching memory

Change 3096355 on 2016/08/22 by [REDACTED_PERSON]

	add support for precaching for "loose files" in the generic async layer.

Change 3098091 on 2016/08/23 by [REDACTED_PERSON]

	Split header into a separate file and disabled a bad optimization in the bulk data.

Change 3099783 on 2016/08/24 by [REDACTED_PERSON]

	rework dependency graph to be much, much faster. About half done.

Change 3100995 on 2016/08/25 by [REDACTED_PERSON]

	fixed bugs with streaming texture from .uexp and cook time check that should have been runtime only

Change 3101369 on 2016/08/25 by [REDACTED_PERSON]

	fixed bug with blueprints in the new loader.

Change 3102793 on 2016/08/26 by [REDACTED_PERSON]

	PS4 - fixed small block memcpy to actually be inline

Change 3103785 on 2016/08/27 by [REDACTED_PERSON]

	fixed case bug with pak order. devirtualized flinkerload::serialize, made sure -fileopenlog is not heavily skewed

Change 3104884 on 2016/08/29 by [REDACTED_PERSON]

	fixed a BP bug and tweaked the -fileopenlog behavior to do leaf assets DFS

Change 3105266 on 2016/08/29 by [REDACTED_PERSON]

	Editor build compilation fix

Change 3105774 on 2016/08/30 by [REDACTED_PERSON]

	add checks to locate cases where we try to use something that isn't loaded yet

Change 3107794 on 2016/08/31 by [REDACTED_PERSON]

	fixed abug with BP's not loading the parent CDO soon enough

Change 3114278 on 2016/09/06 by [REDACTED_PERSON]

	looping loads for paragon load test

Change 3114311 on 2016/09/06 by [REDACTED_PERSON]

	Fix linux compile

Change 3114350 on 2016/09/06 by [REDACTED_PERSON]

	Linux supports fast unaligned int reads

Change 3116169 on 2016/09/07 by [REDACTED_PERSON]

	Force enable separate bulk data cooking when using split cooked files, end-of-exp-file doesn't make sense with the new cook scheme and will crash at runtime

Change 3116538 on 2016/09/07 by [REDACTED_PERSON]

	add dependencies for CDO subobjects

Change 3116596 on 2016/09/07 by [REDACTED_PERSON]

	Change crash to warning when trying to load an import to a missing native class, can happen with editor only classes.

Change 3116855 on 2016/09/07 by [REDACTED_PERSON]

	Move cook dialog down a bit so I can cook without constant dialogs popping up

Change 3117452 on 2016/09/08 by [REDACTED_PERSON]

	Fixing hang when suspending async loading with the async loading thread enabled.

Change 3119255 on 2016/09/09 by [REDACTED_PERSON]

	Removing texture allocations from PackageFileSummary as they were not used by anything.

Change 3119303 on 2016/09/09 by [REDACTED_PERSON]

	Fixed font issue by making all all bulk data either inline or in a ubulk. Added support for compressed packages.

Change 3120324 on 2016/09/09 by [REDACTED_PERSON]

	Fix Cook warnings. Skip transient and client/server only objects when adding dependencies, and mark ShapeComponent BodySetups as properly transient.

Change 3121960 on 2016/09/12 by [REDACTED_PERSON]

	Add RandomizeLoadOrder CVar to randomize the package serial number it uses for sorting async loads

Change 3122635 on 2016/09/13 by [REDACTED_PERSON]

	reworked searching disk warning and minor change to the background tasks used for decompression

Change 3122743 on 2016/09/13 by [REDACTED_PERSON]

	added some checks around memory accounting

Change 3123395 on 2016/09/13 by [REDACTED_PERSON]

	Enable MallocBinned2 by default on cooked windows builds, similar to how PS4 works.
	Disabled thread pool cache clearing on windows, the threading function it was using is very slow on windows specifically

Change 3124748 on 2016/09/14 by [REDACTED_PERSON]

	Store template in import/export table and refer to it for each export to avoid calling GetArchetypeFromRequiredInfo. Minor fix for some NeedLoadForCLient etc stuff on landscape and CDOs. Fix texture streamer minmips stuff.

Change 3125153 on 2016/09/14 by [REDACTED_PERSON]

	don't put transient objects in the import map

Change 3126668 on 2016/09/15 by [REDACTED_PERSON]

	Fix critical bug with imports not waiting for the corresponding export to serialize. Fixed paragon test rig to run longer looping by flushing the renderer. Made random mode more random.

Change 3126755 on 2016/09/15 by [REDACTED_PERSON]

	ooops, test rig fix

Change 3127408 on 2016/09/15 by [REDACTED_PERSON]

	Back out changelist 3123395, restoring windows memory to 4.13 setup

Change 3127409 on 2016/09/15 by [REDACTED_PERSON]

	Remove Memory trim from FlushAsyncLoading, because it gets called much more often in new flow and is slow on some platforms

Change 3127948 on 2016/09/16 by [REDACTED_PERSON]

	Added a check() on any attempt to serialize a pointer to something that hasn't been created yet. This will help us find missing dependencies. There is an exception to this related to CDOs.

Change 3128094 on 2016/09/16 by [REDACTED_PERSON]

	Fixing exports referenced by weak object pointers not being added to the preload dependency list of of the exports that depend on them.

	+ Moved weak object pointer serialization to FArchive operator << to be able to override its behavior when cooking.

Change 3128148 on 2016/09/16 by [REDACTED_PERSON]

	Gil's mod to how we detect exports with missing dependencies

Change 3129052 on 2016/09/16 by [REDACTED_PERSON]

	Add Missing Serialize helpers for WeakObjectPtrs, fixes crash with replicating weak objects

Change 3129053 on 2016/09/16 by [REDACTED_PERSON]

	Fake integrate CL #3123581 from Dev-Framework, to correctly handle detecting components as editor only even when they have collision. Fixes crashes with blueprint editor only components that depend on native templates

Change 3129630 on 2016/09/17 by [REDACTED_PERSON]

	better logging for missing dependencies and properly ifdef'd the CDO primitive comp hack

Change 3130178 on 2016/09/19 by [REDACTED_PERSON]

	Use the correct macro (COOK_FOR_EVENT_DRIVEN_LOAD instead of USE_NEW_ASYNC_IO) for SavePackage changes from CL #3128094

Change 3130224 on 2016/09/19 by [REDACTED_PERSON]

	Compile error fix

Change 3130391 on 2016/09/19 by [REDACTED_PERSON]

	Add cook time fatal errors, and undid a previous change we don't seem to need relating to editor only CDOs

Change 3130484 on 2016/09/19 by [REDACTED_PERSON]

	fixed botched GetArchetypeFromRequiredInfo

Change 3131966 on 2016/09/20 by [REDACTED_PERSON]

	Making the new event driven loader disabled by default. It's now also configurable via project settings (under Streaming Settings -> Event Driven Loader Enabled).

	Enabled the event driven loader for a few internal projects.

Change 3132035 on 2016/09/20 by [REDACTED_PERSON]

	fix dynamic switch on new loader

Change 3132041 on 2016/09/20 by [REDACTED_PERSON]

	Fix for packages not being saved to disk when cooking with event driven loader disabled.

Change 3132195 on 2016/09/20 by [REDACTED_PERSON]

	Enabling the event driven loader for Zen

Change 3133870 on 2016/09/21 by [REDACTED_PERSON]

	Config files now enable the event driven loader with the correct cvar name

Change 3135812 on 2016/09/22 by [REDACTED_PERSON]

	fixed some bugs with GC during streaming

Change 3136102 on 2016/09/22 by [REDACTED_PERSON]

	Release GC lock when FlushingAsyncLoading when running GC.

Change 3136633 on 2016/09/22 by [REDACTED_PERSON]

	fix bug with linkers finsihing before other things linked their imports

Change 3138002 on 2016/09/23 by [REDACTED_PERSON]

	Added an assert that will prevent content cooked for the event driven loader to be loaded by game builds that have the EDL disabled.

Change 3138012 on 2016/09/23 by [REDACTED_PERSON]

	Improved the fix to prevent packages from finishing before external imports have linked. Async load object libraries.

Change 3138031 on 2016/09/23 by [REDACTED_PERSON]

	do not preload obj libs in editor

Change 3139176 on 2016/09/24 by [REDACTED_PERSON]

	fixed another bug with an attempt to call GetArchetypeFromRequiredInfo

Change 3139459 on 2016/09/26 by [REDACTED_PERSON]

	Merging //UE4/Release-4.13 to Dev-LoadTimes (//Tasks/UE4/Dev-LoadTimes)

Change 3139668 on 2016/09/26 by [REDACTED_PERSON]

	change some checks to errors on bad bulk data loads

Change 3141127 on 2016/09/27 by [REDACTED_PERSON]

	Preventing linkers from being detached too early when async loading.

Change 3141129 on 2016/09/27 by [REDACTED_PERSON]

	Releasing GC Lock before calling post GC callbacks to allow StaticFindObject use in these callbacks

Change 3142048 on 2016/09/27 by [REDACTED_PERSON]

	Changing async loading code to not close DelayedLinkerClosePackages linkers until the async package that triggered their creation has finished loading.

Change 3143132 on 2016/09/28 by [REDACTED_PERSON]

	fixed text render comp, which has some editor only issues. Fixes a runtime crash and adds a cooktime warning.

Change 3143198 on 2016/09/28 by [REDACTED_PERSON]

	fixed it so that bogus loads of bulk data are warned but do not crash

Change 3143287 on 2016/09/28 by [REDACTED_PERSON]

	UBT will now invalidate its makefiles if ini files are newer than the makefile (ini files may contains global build settings).

	+ Android toolchain will add hashed command line values to the action reposnse filenames to actually allow it to detect compiler command line changes when detecting actions to execute

Change 3143344 on 2016/09/28 by [REDACTED_PERSON]

	Make UAT pass the project filename to UBT when build non-code projects so that UBT can parse all ini files.

Change 3143865 on 2016/09/28 by [REDACTED_PERSON]

	iffy fix for the net load assert in paragon, plus a few checks and one bit of code removed that should never be hit in the EDL, but makes no sense

Change 3144683 on 2016/09/29 by [REDACTED_PERSON]

	Minor refactor of pak file non-filename stuff
	 - Don't check for file existing before running through the security delegate
	 - Default behaviour when using new IO is to reject uasset/umap/ubulk/uexp files immediately. Can be disabled by setting EXCLUDE_NONPAK_UE_EXTENSIONS to 0 in project .build.cs

Change 3144745 on 2016/09/29 by [REDACTED_PERSON]

	Orion non-pak file whitelisting is enabled for all cooked game only builds now, rather than just clients

Change 3144780 on 2016/09/29 by [REDACTED_PERSON]

	use poison proxy on non-test/shipping builds

Change 3144819 on 2016/09/29 by [REDACTED_PERSON]

	added a few asserts and added an improved fix for the net crash

Change 3145414 on 2016/09/29 by [REDACTED_PERSON]

	fixed android assert....not sure why I need that block of code.

Change 3146502 on 2016/09/30 by [REDACTED_PERSON]

	Fix for GPU hang from MarcusW

Change 3146774 on 2016/09/30 by [REDACTED_PERSON]

	Fixing a crash when constantly streaming levels in and out caused by keeping references to objects (levels) that were requested to be streamed out.

	- Removed FAsyncObjectsReferencer. References will now be owned by FAsyncPackage
	- UGCObjectReferencer is now more thread safe

Change 3148008 on 2016/10/01 by [REDACTED_PERSON]

	add additional error for attempting to create an object from a class that needs to be loaded

Change 3148009 on 2016/10/01 by [REDACTED_PERSON]

	fix very old threading bug whereby the ASL and GT would attempt to use the same static array

Change 3148222 on 2016/10/02 by [REDACTED_PERSON]

	Fix for an assert when an FGCObject is removed when purging UObjects

Change 3148229 on 2016/10/02 by [REDACTED_PERSON]

	disable assert that was crashing paragon ps4

Change 3148409 on 2016/10/03 by [REDACTED_PERSON]

	Allow another case for removing FGCObjects while in GC.

Change 3148416 on 2016/10/03 by [REDACTED_PERSON]

	Merging //UE4/Release-4.13 to Dev-LoadTimes (//Tasks/UE4/Dev-LoadTimes)

Change 3149566 on 2016/10/03 by [REDACTED_PERSON]

Change 3149913 on 2016/10/04 by [REDACTED_PERSON]

	better broadcast

Change 2889560 on 2016/03/02 by [REDACTED_PERSON]

	Packages for scheduled tasks.

Change 2889566 on 2016/03/02 by [REDACTED_PERSON]

	Remaining nuget packages for hangfire, unity and scheduled tasks.

Change 2980458 on 2016/05/17 by [REDACTED_PERSON]

	Attempt to fix crash report submission problems from CRP to CR website
	[UE-30257] - Crashreports are sometimes missing file attachments

	Passing crash GUID so that website can easily check for duplicates in future
	Increased request timeout for AddCrash to be longer than website database timeout
	Logging retries for future visibility
	CRP v.1.1.6

Change 3047870 on 2016/07/13 by [REDACTED_PERSON]

	Updated CRW to entity framework with repository models. #rb none

Change 3126265 on 2016/09/15 by [REDACTED_PERSON]

	Fix for TCString::Strspn.

Change 3126266 on 2016/09/15 by [REDACTED_PERSON]

	Alternative fix for GitHub 2698: Fix one bug : Parsing command "Enable True" is invalid.

Change 3126268 on 2016/09/15 by [REDACTED_PERSON]

	UWorld can no longer be extended by users.
	UHT now handles final class declarations.

Change 3126273 on 2016/09/15 by [REDACTED_PERSON]

	A further attempt to catch uninitialized pointers supplied to the GC.

Change 3130042 on 2016/09/19 by [REDACTED_PERSON]

	Super for USTRUCTs.

	Suggested here: https://udn.unrealengine.com/questions/310461/automatically-typedef-super-for-ustructs.html

Change 3131861 on 2016/09/20 by [REDACTED_PERSON]

	Reconciling work for view engine changes #rb none

Change 3131862 on 2016/09/20 by [REDACTED_PERSON]

	Removal of THasOperatorEquals and THasOperatorNotEquals from Platform.h, which should have happened as part of CL# 3045963.

Change 3131863 on 2016/09/20 by [REDACTED_PERSON]

	Adding packages #rb none

Change 3131869 on 2016/09/20 by [REDACTED_PERSON]

	Improved error message for enum classes with a missing base:

	Error: Missing base specifier for enum class 'EMyEnum' - did you mean ': uint8'?

Change 3132046 on 2016/09/20 by [REDACTED_PERSON]

	Fix for cvar thread access assert in FLandscapeComponentGrassData serialization function
	 - This function can be called from the async thread so access CVarGrassDiscardDataOnLoad with GetValueOnAnyThread() rather than GetValueOnGameThread()

Change 3133201 on 2016/09/20 by [REDACTED_PERSON]

	Reorganize WindowsPlatformMemory and MacPlatformMemory to work like LinuxPlatformMemory where there is an enum to select the allocator, and move some of it up to GenericPlatformMemory
	Add command line options to select malloc at runtime for Windows and Linux, I don't know how Mac options work
	Improve the performance of BroadcastSlow_OnlyUseForSpecialPurposes on windows, but there are cases where it occaisionally stalls for a few seconds waiting for the flush
	Add MallocBinned2 as an option for mac, linux, and windows, but default to off due to some threading issues

Change 3133722 on 2016/09/21 by [REDACTED_PERSON]

	Cooker forces a shader compilation flush when it detects that it has passed the max memory budget

Change 3133756 on 2016/09/21 by [REDACTED_PERSON]

	Refactor of TrimPrecedingAndTrailing to avoid a call to FString::Mid with a negative count, which is now illegal.

Change 3134182 on 2016/09/21 by [REDACTED_PERSON]

	GitHub #1986: Don't show warnings and erros in console twice with UCommandlet::LogToConsole == true

Change 3134306 on 2016/09/21 by [REDACTED_PERSON]

	Fix it so FMallocBinned2::Trim skips task threads on desktop platforms, they are too slow and don't allocate much memory
	Enable MallocBinned2 as default binned malloc on Windows
	Remove the -Run command line check as it was removed from the old version as well

Change 3135569 on 2016/09/22 by [REDACTED_PERSON]

	Don't create material resources if we are in a build that can never render
	 - Saves a few MB of memory

Change 3135652 on 2016/09/22 by [REDACTED_PERSON]

	New async-loading-thread-safe IsA implementation.

Change 3135692 on 2016/09/22 by [REDACTED_PERSON]

	Minor bug fixes to view pages #rb none

Change 3135990 on 2016/09/22 by [REDACTED_PERSON]

	Adding ENGINE_API to FStripDataFlags sp that it can be used outside of the Engine module.

Change 3136020 on 2016/09/22 by [REDACTED_PERSON]

	Display a meaningful error and shutdown if Core modules fail to load.

	https://udn.unrealengine.com/questions/312063/mac-unrealheadertool-failing-randomly.html

Change 3136107 on 2016/09/22 by [REDACTED_PERSON]

	Added S3 file upload to output stage of Crash Report Process (v.1.1.26)
	[UE-35991] - Crash Report Process to write crash files to S3

	Also adds OOM alerts to CRP.
	Also disk space alerts changed to 5% free space and repeat once every 30 minutes instead of 10 minutes.

Change 3137562 on 2016/09/23 by [REDACTED_PERSON]

	TUniquePtr<T[]> support.

Change 3138030 on 2016/09/23 by [REDACTED_PERSON]

	Virtual UProperty functions moved out of headers into .cpp files to ease iteration.

Change 3140381 on 2016/09/26 by [REDACTED_PERSON]

	Disabled uploads via CRRs while leaving services switched on to avoid crashes in some clients.
	[UETOOL-1005] - Turn off CrashReportReceivers

Change 3141150 on 2016/09/27 by [REDACTED_PERSON]

	Invoke support for TFunction.

Change 3141151 on 2016/09/27 by [REDACTED_PERSON]

	UBoolProperty now supports hashing and is therefore usable as a TSet element or TMap key.
	FText is now prevented from being a TSet element or TMap key.
	UTextProperty::GetCPPTypeForwardDeclaration implementation moved to the .cpp file.

Change 3141440 on 2016/09/27 by [REDACTED_PERSON]

	Removed legacy queues and unnecessary duplication checks from Crash Report Process (v1.2.0)
	[UE-36246] - CRP scalability: Simplify CRP inputs to DataRouter/S3 only

Change 3142999 on 2016/09/28 by [REDACTED_PERSON]

	Added dedicated PS4 crash queue to Crash Report Process (v1.2.1)

Change 3144831 on 2016/09/29 by [REDACTED_PERSON]

	InternalPrecache now flags the archive as in-error so that it can be checked by a caller, rather than popping up a dialog box and asserting.

Change 3145184 on 2016/09/29 by [REDACTED_PERSON]

	FScopedCreateImportCounter will now always store the current linker and restore the previous one when it exits.

Change 3148432 on 2016/10/03 by [REDACTED_PERSON]

	Thread safety fixes for the async log writer + made the async log writer flush its archive more often.

Change 3148661 on 2016/10/03 by [REDACTED_PERSON]

	Fixing merge of IsNonPakFilenameAllowed()
	 - Removed directory search stuff... we pass everything to the delegate now anyway

Change 3149669 on 2016/10/03 by [REDACTED_PERSON]

	Lower verbosity of warnings from deleting native properties. These cases do not cause any problems and are not fixable without resaving the content after it has started warning. I checked Jira history and neither of these warnings has ever found a real bug, but has caused a lot of content to be resaved unnecessarily.

Change 3149670 on 2016/10/03 by [REDACTED_PERSON]

	Merge CL #3149566 from Dev-LoadTimes

Change 3149835 on 2016/10/04 by [REDACTED_PERSON]

	Thread safety fix for SkyLightComponent
	 - Add to global update list from PostLoad rather than PostInitProperties so that it happens on the game thread, and not the async loading thread (if enabled)

Change 3149836 on 2016/10/04 by [REDACTED_PERSON]

	Thread safety fix for ReflectionCaptureComponent
	 - Add to global update list from PostLoad rather than PostInitProperties so that it happens on the game thread, and not the async loading thread (if enabled)

Change 3149959 on 2016/10/04 by [REDACTED_PERSON]

	Allow import packages to be missing if they're on the KnownMissingPackages list

Change 3150023 on 2016/10/04 by [REDACTED_PERSON]

	Updating jira strings. #rb none

Change 3150050 on 2016/10/04 by [REDACTED_PERSON]

	MakeShared now returns a TSharedRef (which is implicitly convertible to TSharedPtr) rather than a TSharedPtr (which is not implicitly convertible to TSharedRef), for ease of use and because MakeShared can't return a null pointer anyway.

Change 3150110 on 2016/10/04 by [REDACTED_PERSON]

	Allow UGCObjectReferencer::AddObjects to happen during BeginDestry and FinishDestroy. It's fine as long as we're not adding new objects during reachability analysis.

Change 3150120 on 2016/10/04 by [REDACTED_PERSON]

	fix task graph/binned2 broadcast for PS4

Change 3150195 on 2016/10/04 by [REDACTED_PERSON]

	Fixing WEX crash

Change 3150212 on 2016/10/04 by [REDACTED_PERSON]

	Increasing compiler memory limit to fix CIS errors

Change 3151583 on 2016/10/05 by [REDACTED_PERSON]

	Temporarily switching to the old IsA path

	#jria UE-36803

Change 3151642 on 2016/10/05 by [REDACTED_PERSON]

	Dependency fixes for GameFeedback modules.

Change 3151653 on 2016/10/05 by [REDACTED_PERSON]

	Maybe fix for crash on the Mac

diff --git a/Engine/Source/Developer/DesktopPlatform/Private/Windows/WindowsNativeFeedbackContext.cpp b/Engine/Source/Developer/DesktopPlatform/Private/Windows/WindowsNativeFeedbackContext.cpp
--- a/Engine/Source/Developer/DesktopPlatform/Private/Windows/WindowsNativeFeedbackContext.cpp
+++ b/Engine/Source/Developer/DesktopPlatform/Private/Windows/WindowsNativeFeedbackContext.cpp
@@ -59,7 +59,7 @@ void FWindowsNativeFeedbackContext::Serialize( const TCHAR* V, ELogVerbosity::Ty
 		}
 	}
 
-	if( GLogConsole && IsRunningCommandlet() )
+	if( GLogConsole && IsRunningCommandlet() && !GLog->IsRedirectingTo(GLogConsole) )
 	{
 		GLogConsole->Serialize( V, Verbosity, Category );
 	}
diff --git a/Engine/Source/Developer/HotReload/Private/HotReloadClassReinstancer.cpp b/Engine/Source/Developer/HotReload/Private/HotReloadClassReinstancer.cpp
--- a/Engine/Source/Developer/HotReload/Private/HotReloadClassReinstancer.cpp
+++ b/Engine/Source/Developer/HotReload/Private/HotReloadClassReinstancer.cpp
@@ -181,6 +181,11 @@ void FHotReloadClassReinstancer::SerializeCDOProperties(UObject* InObject, FHotR
 
 			return Ar;
 		}
+		FArchive& operator<<(FWeakObjectPtr& WeakObjectPtr) override
+		{
+			WeakObjectPtr.Serialize(*this);
+			return *this;
+		}
 		/** Archive name, for debugging */
 		virtual FString GetArchiveName() const override { return TEXT("FCDOWriter"); }
 	};
@@ -416,6 +421,11 @@ void FHotReloadClassReinstancer::UpdateDefaultProperties()
 
 			return Ar;
 		}
+		FArchive& operator<<(FWeakObjectPtr& WeakObjectPtr) override
+		{
+			WeakObjectPtr.Serialize(*this);
+			return *this;
+		}
 	};
 
 	// Collect default subobjects to update their properties too
diff --git a/Engine/Source/Editor/PropertyEditor/Private/PropertyNode.cpp b/Engine/Source/Editor/PropertyEditor/Private/PropertyNode.cpp
--- a/Engine/Source/Editor/PropertyEditor/Private/PropertyNode.cpp
+++ b/Engine/Source/Editor/PropertyEditor/Private/PropertyNode.cpp
@@ -2873,6 +2873,7 @@ void FPropertyNode::PropagatePropertyChange( UObject* ModifiedObject, const TCHA
 					virtual FArchive& operator<<(FLazyObjectPtr&        Val) override { this->Serialize(&Val, sizeof(Val)); return *this; }
 					virtual FArchive& operator<<(FAssetPtr&             Val) override { this->Serialize(&Val, sizeof(Val)); return *this; }
 					virtual FArchive& operator<<(FStringAssetReference& Val) override { this->Serialize(&Val, sizeof(Val)); return *this; }
+					virtual FArchive& operator<<(FWeakObjectPtr&		Val) override { this->Serialize(&Val, sizeof(Val)); return *this; }
 
 					void Serialize(void* Data, int64 Num)
 					{
@@ -2917,6 +2918,7 @@ void FPropertyNode::PropagatePropertyChange( UObject* ModifiedObject, const TCHA
 					virtual FArchive& operator<<(FLazyObjectPtr&        Val) override { this->Serialize(&Val, sizeof(Val)); return *this; }
 					virtual FArchive& operator<<(FAssetPtr&             Val) override { this->Serialize(&Val, sizeof(Val)); return *this; }
 					virtual FArchive& operator<<(FStringAssetReference& Val) override { this->Serialize(&Val, sizeof(Val)); return *this; }
+					virtual FArchive& operator<<(FWeakObjectPtr&		Val) override { this->Serialize(&Val, sizeof(Val)); return *this; }
 
 					void Serialize(void* Data, int64 Num)
 					{
diff --git a/Engine/Source/Editor/UnrealEd/Private/Commandlets/CookCommandlet.cpp b/Engine/Source/Editor/UnrealEd/Private/Commandlets/CookCommandlet.cpp
--- a/Engine/Source/Editor/UnrealEd/Private/Commandlets/CookCommandlet.cpp
+++ b/Engine/Source/Editor/UnrealEd/Private/Commandlets/CookCommandlet.cpp
@@ -1568,6 +1568,9 @@ bool UCookCommandlet::NewCook( const TArray<ITargetPlatform*>& Platforms, TArray
 				{
 					bShouldGC = true;
 					GCReason = TEXT("Exceeded Max Memory");
+
+					UE_LOG(LogCookCommandlet, Display, TEXT("Detected max mem exceeded - forcing shader compilation flush"));
+					GShaderCompilingManager->FinishAllCompilation();
 				}
 				else if ((TickResults & UCookOnTheFlyServer::COSR_RequiresGC) != 0) // cooker loaded some object which needs to be cleaned up before the cooker can proceed so force gc
 				{
diff --git a/Engine/Source/Editor/UnrealEd/Private/Commandlets/PackageUtilities.cpp b/Engine/Source/Editor/UnrealEd/Private/Commandlets/PackageUtilities.cpp
--- a/Engine/Source/Editor/UnrealEd/Private/Commandlets/PackageUtilities.cpp
+++ b/Engine/Source/Editor/UnrealEd/Private/Commandlets/PackageUtilities.cpp
@@ -1095,6 +1095,21 @@ void FPkgInfoReporter_Log::GeneratePackageReport( FLinkerLoad* InLinker/*=NULL*/
 					}
 				}
 
+				// find the name of this object's parent...for UClasses, this will be the parent class
+				// for UFunctions, this will be the SuperFunction, if it exists, etc.
+				FString TemplateName;
+				if (!Export.TemplateIndex.IsNull())
+				{
+					if ((InfoFlags&PKGINFO_Paths) != 0)
+					{
+						TemplateName = *Linker->GetPathName(Export.TemplateIndex);
+					}
+					else
+					{
+						TemplateName = Linker->ImpExp(Export.TemplateIndex).ObjectName.ToString();
+					}
+				}
+
 				// find the name of this object's Outer.  For UClasses, this will generally be the
 				// top-level package itself.  For properties, a UClass, etc.
 				FString OuterName;
@@ -1111,7 +1126,8 @@ void FPkgInfoReporter_Log::GeneratePackageReport( FLinkerLoad* InLinker/*=NULL*/
 				}
 
 				UE_LOG(LogPackageUtilities, Warning, TEXT("\t\t         Class: '%s' (%i)"), *ClassName.ToString(), ClassIndex.ForDebugging() );
-				UE_LOG(LogPackageUtilities, Warning, TEXT("\t\t        Parent: '%s' (%d)"), *ParentName, Export.SuperIndex.ForDebugging() );
+				UE_LOG(LogPackageUtilities, Warning, TEXT("\t\t        Parent: '%s' (%d)"), *ParentName, Export.SuperIndex.ForDebugging());
+				UE_LOG(LogPackageUtilities, Warning, TEXT("\t\t      Template: '%s' (%d)"), *TemplateName, Export.TemplateIndex.ForDebugging());
 				UE_LOG(LogPackageUtilities, Warning, TEXT("\t\t         Outer: '%s' (%d)"), *OuterName, Export.OuterIndex.ForDebugging() );
 				UE_LOG(LogPackageUtilities, Warning, TEXT("\t\t      Pkg Guid: %s"), *Export.PackageGuid.ToString());
 				UE_LOG(LogPackageUtilities, Warning, TEXT("\t\t   ObjectFlags: 0x%08X"), (uint32)Export.ObjectFlags );
diff --git a/Engine/Source/Editor/UnrealEd/Private/Kismet2/KismetReinstanceUtilities.cpp b/Engine/Source/Editor/UnrealEd/Private/Kismet2/KismetReinstanceUtilities.cpp
--- a/Engine/Source/Editor/UnrealEd/Private/Kismet2/KismetReinstanceUtilities.cpp
+++ b/Engine/Source/Editor/UnrealEd/Private/Kismet2/KismetReinstanceUtilities.cpp
@@ -265,6 +265,10 @@ FBlueprintCompileReinstancer::FBlueprintCompileReinstancer(UClass* InClassToRein
 						ReparentChild(ChildClass);
 					}
 				}
+	
+			#if UCLASS_FAST_ISA_IMPL == UCLASS_ISA_CLASSARRAY
+				ChildClass->ReinitializeBaseChainArray();
+			#endif
 			}
 		}
 
diff --git a/Engine/Source/Editor/UnrealEd/Public/AssetViewerSettings.h b/Engine/Source/Editor/UnrealEd/Public/AssetViewerSettings.h
--- a/Engine/Source/Editor/UnrealEd/Public/AssetViewerSettings.h
+++ b/Engine/Source/Editor/UnrealEd/Public/AssetViewerSettings.h
@@ -29,6 +29,9 @@ struct FPreviewSceneProfile
 		RotationSpeed = 2.0f;
 		// Set up default cube map texture from editor/engine textures
 		EnvironmentCubeMap = nullptr;
+#if WITH_EDITORONLY_DATA
+		EnvironmentCubeMap = LoadObject<UTextureCube>(NULL, TEXT("/Engine/EditorMaterials/AssetViewer/EpicQuadPanorama_CC+EV1.EpicQuadPanorama_CC+EV1"));		
+#endif
 		bPostProcessingEnabled = true;
 		DirectionalLightRotation = FRotator(-40.f, -67.5f, 0.f);
 	}
diff --git a/Engine/Source/Programs/AutomationTool/Scripts/BuildProjectCommand.Automation.cs b/Engine/Source/Programs/AutomationTool/Scripts/BuildProjectCommand.Automation.cs
--- a/Engine/Source/Programs/AutomationTool/Scripts/BuildProjectCommand.Automation.cs
+++ b/Engine/Source/Programs/AutomationTool/Scripts/BuildProjectCommand.Automation.cs
@@ -108,29 +108,40 @@ public partial class Project : CommandUtils
 		// Setup cooked targets
 		if (Params.HasClientCookedTargets && (TargetMask & ProjectBuildTargets.ClientCooked) == ProjectBuildTargets.ClientCooked)
 		{
-            List<UnrealTargetPlatform> UniquePlatformTypes = Params.ClientTargetPlatforms.ConvertAll(x => x.Type).Distinct().ToList();
+			List<UnrealTargetPlatform> UniquePlatformTypes = Params.ClientTargetPlatforms.ConvertAll(x => x.Type).Distinct().ToList();
 
-            foreach (var BuildConfig in Params.ClientConfigsToBuild)
+			foreach (var BuildConfig in Params.ClientConfigsToBuild)
 			{
-                foreach (var ClientPlatformType in UniquePlatformTypes)
-				{
-                    string ScriptPluginArgs = GetBlueprintPluginPathArgument(Params, true, ClientPlatformType);
-                    CrashReportPlatforms.Add(ClientPlatformType);
-					Agenda.AddTargets(Params.ClientCookedTargets.ToArray(), ClientPlatformType, BuildConfig, Params.CodeBasedUprojectPath, InAddArgs: ScriptPluginArgs + " -remoteini=\"" + Params.RawProjectPath.Directory.FullName + "\"");
+				foreach (var ClientPlatformType in UniquePlatformTypes)
+				{					
+					string AdditionalArgs = GetBlueprintPluginPathArgument(Params, true, ClientPlatformType);
+					AdditionalArgs += " -remoteini=\"" + Params.RawProjectPath.Directory.FullName + "\"";
+					if (Params.IsCodeBasedProject == false)
+					{
+						AdditionalArgs += " -project=\"" + Path.GetFullPath(Params.RawProjectPath.FullName) + "\"";
+					}
+					CrashReportPlatforms.Add(ClientPlatformType);
+					Agenda.AddTargets(Params.ClientCookedTargets.ToArray(), ClientPlatformType, BuildConfig, Params.CodeBasedUprojectPath, InAddArgs: AdditionalArgs);
 				}
 			}
 		}
 		if (Params.HasServerCookedTargets && (TargetMask & ProjectBuildTargets.ServerCooked) == ProjectBuildTargets.ServerCooked)
 		{
-            List<UnrealTargetPlatform> UniquePlatformTypes = Params.ServerTargetPlatforms.ConvertAll(x => x.Type).Distinct().ToList();
+			List<UnrealTargetPlatform> UniquePlatformTypes = Params.ServerTargetPlatforms.ConvertAll(x => x.Type).Distinct().ToList();
 
-            foreach (var BuildConfig in Params.ServerConfigsToBuild)
+			foreach (var BuildConfig in Params.ServerConfigsToBuild)
 			{
 				foreach (var ServerPlatformType in UniquePlatformTypes)
 				{
-                    string ScriptPluginArgs = GetBlueprintPluginPathArgument(Params, false, ServerPlatformType);
-                    CrashReportPlatforms.Add(ServerPlatformType);
-					Agenda.AddTargets(Params.ServerCookedTargets.ToArray(), ServerPlatformType, BuildConfig, Params.CodeBasedUprojectPath, InAddArgs: ScriptPluginArgs + " -remoteini=\"" + Params.RawProjectPath.Directory.FullName + "\"");
+					string AdditionalArgs = GetBlueprintPluginPathArgument(Params, false, ServerPlatformType);
+					AdditionalArgs += " -remoteini=\"" + Params.RawProjectPath.Directory.FullName + "\"";
+					if (Params.IsCodeBasedProject == false)
+					{
+						AdditionalArgs += " -project=\"" + Path.GetFullPath(Params.RawProjectPath.FullName) + "\"";
+					}
+
+					CrashReportPlatforms.Add(ServerPlatformType);
+					Agenda.AddTargets(Params.ServerCookedTargets.ToArray(), ServerPlatformType, BuildConfig, Params.CodeBasedUprojectPath, InAddArgs: AdditionalArgs);
 				}
 			}
 		}
diff --git a/Engine/Source/Programs/CrashReporter/CrashReportCommon/AmazonClient.cs b/Engine/Source/Programs/CrashReporter/CrashReportCommon/AmazonClient.cs
new file mode 100644
--- /dev/null
+++ b/Engine/Source/Programs/CrashReporter/CrashReportCommon/AmazonClient.cs
@@ -0,0 +1,205 @@
+﻿using System;
+using System.Collections.Generic;
+using System.IO;
+using System.Linq;
+using System.Net;
+using System.Text;
+using System.Threading.Tasks;
+using Amazon.Route53.Model;
+using Amazon.Runtime;
+using Amazon.S3;
+using Amazon.S3.Model;
+using Amazon.SQS;
+using Amazon.SQS.Model;
+
+namespace Tools.CrashReporter.CrashReportCommon
+{
+	/// <summary>
+	/// Helper class for accessing AWS
+	/// </summary>
+	public class AmazonClient : IDisposable
+	{
+		/// <summary>
+		/// Is the AmazonClient initialized and ready to use with S3?
+		/// </summary>
+		public bool IsS3Valid
+		{
+			get { return S3Client != null; }
+		}
+
+		/// <summary>
+		/// Is the AmazonClient initialized and ready to use with SQS?
+		/// </summary>
+		public bool IsSQSValid
+		{
+			get { return SqsClient != null; }
+		}
+
+		/// <summary>
+		/// Create a new AmazonClient
+		/// </summary>
+		/// <param name="Credentials">Amazon Credentials for the connection to AWS</param>
+		/// <param name="SqsConfig">SQS connection details</param>
+		/// <param name="S3Config">S3 connection details</param>
+		/// <param name="OutError">Out error, receives the error string if an error occurs</param>
+		public AmazonClient(AWSCredentials Credentials, AmazonSQSConfig SqsConfig, AmazonS3Config S3Config, out string OutError)
+		{
+			OutError = string.Empty;
+
+			if (SqsConfig != null)
+			{
+				try
+				{
+					SqsClient = new AmazonSQSClient(Credentials, SqsConfig);
+				}
+				catch (Exception Ex)
+				{
+					OutError += Ex.Message;
+				}
+			}
+
+			if (S3Config != null)
+			{
+				try
+				{
+					S3Client = new AmazonS3Client(Credentials, S3Config);
+				}
+				catch (Exception Ex)
+				{
+					OutError += Ex.Message;
+				}
+			}
+		}
+
+		/// <summary>
+		/// Dispose of this object
+		/// </summary>
+		public void Dispose()
+		{
+			Dispose(true);
+			GC.SuppressFinalize(this);
+		}
+
+		/// <summary>
+		/// Dispose of this object
+		/// </summary>
+		/// <param name="disposing">Dispose of managed resources?</param>
+		protected virtual void Dispose(bool disposing)
+		{
+			if (disposing)
+			{
+				if (SqsClient != null) SqsClient.Dispose();
+				if (S3Client != null) S3Client.Dispose();
+			}
+		}
+
+		/// <summary>
+		/// Get the number of items in an SQS queue
+		/// </summary>
+		/// <param name="QueueUrl">The SQS queue URL</param>
+		/// <returns>Number of messages in the queue. Returns zero in the event of an error.</returns>
+		public int GetSQSQueueCount(string QueueUrl)
+		{
+			if (!IsSQSValid) return 0;
+
+			GetQueueAttributesRequest AttribRequest = new GetQueueAttributesRequest
+			{
+				QueueUrl = QueueUrl,
+				AttributeNames = new List<string>
+				{
+					"ApproximateNumberOfMessages"
+				}
+			};
+
+			GetQueueAttributesResponse AttribResponse = SqsClient.GetQueueAttributes(AttribRequest);
+			return AttribResponse.ApproximateNumberOfMessages;
+		}
+
+		/// <summary>
+		/// Gets a message from an SQS queue and if successful, deletes the message from the queue.
+		/// Safe method for multiple concurrent consumers of a single SQS because the initial message read makes the message unavailable to other callers.
+		/// </summary>
+		/// <param name="QueueUrl">The SQS queue URL</param>
+		/// <param name="bForceNoDelete">Debug option. Skips the message delete when testing reading messages from a live queue</param>
+		/// <returns>A dequeued SQS message or null if it fails.</returns>
+		public Message SQSDequeueMessage(string QueueUrl, bool bForceNoDelete)
+		{
+			if (!IsSQSValid) return null;
+
+			ReceiveMessageRequest ReceiveRequest = new ReceiveMessageRequest
+			{
+				QueueUrl = QueueUrl,
+				MaxNumberOfMessages = 1
+			};
+
+			ReceiveMessageResponse ReceiveResponse = SqsClient.ReceiveMessage(ReceiveRequest);
+
+			if (ReceiveResponse.Messages.Count == 1)
+			{
+				Message Message = ReceiveResponse.Messages[0];
+
+				if (Message != null && (bForceNoDelete || DeleteRecordSQS(QueueUrl, Message)))
+				{
+					return Message;
+				}
+			}
+
+			return null;
+		}
+
+		/// <summary>
+		/// Retrieves an object from S3. (Caller can use the response to read the response stream)
+		/// </summary>
+		/// <param name="BucketName">S3 bucket name</param>
+		/// <param name="Key">S3 object key</param>
+		/// <returns>S3 object response for the requested object or null if it fails.</returns>
+		public GetObjectResponse GetS3Object(string BucketName, string Key)
+		{
+			if (!IsS3Valid) return null;
+
+			GetObjectRequest ObjectRequest = new GetObjectRequest
+			{
+				BucketName = BucketName,
+				Key = Key
+			};
+
+			return S3Client.GetObject(ObjectRequest);
+		}
+
+		/// <summary>
+		/// Uploads a file to S3 synchronously (blocks until complete).
+		/// </summary>
+		/// <param name="BucketName">S3 bucket name</param>
+		/// <param name="Key">S3 object key</param>
+		/// <param name="Filepath">Full path of the file to upload</param>
+		/// <returns>S3 object response for the uploaded object or null if it fails.</returns>
+		public PutObjectResponse PutS3ObjectFromFile(string BucketName, string Key, string Filepath)
+		{
+			if (!IsS3Valid) return null;
+
+			PutObjectRequest ObjectRequest = new PutObjectRequest
+			{
+				BucketName = BucketName,
+				Key = Key,
+				FilePath = Filepath
+			};
+
+			return S3Client.PutObject(ObjectRequest);
+		}
+
+		private bool DeleteRecordSQS(string QueueUrl, Message InMessage)
+		{
+			DeleteMessageRequest DeleteRequest = new DeleteMessageRequest
+			{
+				QueueUrl = QueueUrl,
+				ReceiptHandle = InMessage.ReceiptHandle
+			};
+
+			var DeleteResponse = SqsClient.DeleteMessage(DeleteRequest);
+			return DeleteResponse.HttpStatusCode == HttpStatusCode.OK;
+		}
+
+		private readonly AmazonSQSClient SqsClient;
+		private readonly AmazonS3Client S3Client;
+	}
+}
diff --git a/Engine/Source/Programs/CrashReporter/CrashReportCommon/CrashReportCommon.csproj b/Engine/Source/Programs/CrashReporter/CrashReportCommon/CrashReportCommon.csproj
--- a/Engine/Source/Programs/CrashReporter/CrashReportCommon/CrashReportCommon.csproj
+++ b/Engine/Source/Programs/CrashReporter/CrashReportCommon/CrashReportCommon.csproj
@@ -49,6 +49,10 @@
     <CodeAnalysisRuleSet>MinimumRecommendedRules.ruleset</CodeAnalysisRuleSet>
   </PropertyGroup>
   <ItemGroup>
+    <Reference Include="AWSSDK, Version=2.3.50.1, Culture=neutral, PublicKeyToken=9f476d3089b52be3, processorArchitecture=MSIL">
+      <SpecificVersion>False</SpecificVersion>
+      <HintPath>..\..\..\..\Binaries\ThirdParty\AWSSDK\AWSSDK.dll</HintPath>
+    </Reference>
     <Reference Include="System" />
     <Reference Include="System.Core" />
     <Reference Include="System.Net" />
@@ -61,6 +65,7 @@
     <Reference Include="System.Xml" />
   </ItemGroup>
   <ItemGroup>
+    <Compile Include="AmazonClient.cs" />
     <Compile Include="Common.cs" />
     <Compile Include="CrashContext.cs" />
     <Compile Include="CrashReporterConstants.cs" />
diff --git a/Engine/Source/Programs/CrashReporter/CrashReportCommon/JiraConnection.cs b/Engine/Source/Programs/CrashReporter/CrashReportCommon/JiraConnection.cs
--- a/Engine/Source/Programs/CrashReporter/CrashReportCommon/JiraConnection.cs
+++ b/Engine/Source/Programs/CrashReporter/CrashReportCommon/JiraConnection.cs
@@ -9,6 +9,7 @@ using System.Net;
 using System.Text;
 using System.Threading; 
 using System.Threading.Tasks;
+using System.Web;
 using System.Web.Script.Serialization;
 
 /// <summary>
@@ -77,7 +78,7 @@ public class JiraConnection
 	}
 
     /// <summary>
-    /// Maps the name to the version
+    /// Maps the name to the version.
     /// </summary>
     Dictionary<string, Dictionary<string, object>> NameToBranch = new Dictionary<string, Dictionary<string, object>>();
     /// <summary>
@@ -109,16 +110,16 @@ public class JiraConnection
 	}
 
 	/// <summary>
-	/// Retrieves all components from https://jira.ol.epicgames.net/rest/api/2/project/UE/components
+	/// Retrieves all components from https://jira.it.epicgames.net/rest/api/2/project/UE/components
 	/// </summary>
 	private void InitializeComponentArray()
 	{
 		/*
-			"self": "https://jira.ol.epicgames.net/rest/api/2/component/..",
+			"self": "https://jira.it.epicgames.net/rest/api/2/component/..",
 			"id": "14019",
 			"name": "...",
 			"lead": {
-			"self": "https://jira.ol.epicgames.net/rest/api/2/user?username=...,
+			"self": "https://jira.it.epicgames.net/rest/api/2/user?username=...,
 			"key": "...",
 			"name": "...",
 			"avatarUrls": {...},
@@ -144,12 +145,12 @@ public class JiraConnection
 	}
 
 	/// <summary>
-	/// Retrieves all build versions from https://jira.ol.epicgames.net/rest/api/2/project/UE/versions
+	/// Retrieves all build versions from https://jira.it.epicgames.net/rest/api/2/project/UE/versions
 	/// </summary>
 	private void InitializeVersionArray()
 	{
 		/*
-			"self": "https://jira.ol.epicgames.net/rest/api/2/version/...",
+			"self": "https://jira.it.epicgames.net/rest/api/2/version/...",
 			"id": "11700",
 			"description": "Scheduled engine release",
 			"name": "4.8",
@@ -262,15 +263,15 @@ public class JiraConnection
 	/// <summary>
 	/// Creates a Jira ticket with the given fields
 	/// </summary>
-	public string AddJiraTicket( Dictionary<string, object> Fields )
+	public string AddJiraTicket( Dictionary<string, object> fields )
 	{
 		// Create the request object
-		Dictionary<string, object> Request = new Dictionary<string, object> { { "fields", Fields } };
-		HttpWebResponse Response = JiraRequest( "/issue", JiraMethod.POST, Request, HttpStatusCode.Created );
+		var request = new Dictionary<string, object> { { "fields", fields } };
+		var response = JiraRequest( "/issue", JiraMethod.POST, request, HttpStatusCode.Created );
 
 		// Return the new issue id
-		Dictionary<string, object> ResponseObject = ParseJiraResponse( Response );
-		return (string)ResponseObject["key"];
+		var responseObject = ParseJiraResponse( response );
+		return (string)responseObject["key"];
 	}
 
 	/// <summary>
@@ -290,40 +291,40 @@ public class JiraConnection
 	/// <param name="FieldsToGet">Fields that will be grabbed from the results</param>
 	public Dictionary<string,Dictionary<string, object>> SearchJiraTickets( string SearchQuery, string[] FieldsToGet )
 	{
-		Dictionary<string, Dictionary<string, object>> Result = new Dictionary<string, Dictionary<string, object>>();
+		var result = new Dictionary<string, Dictionary<string, object>>();
 
 		// Do the query in multiple steps, to avoid timeouts when receiving large amounts of data
-		const int MaxResultsPerQuery = 100;
-		for( int StartAt = 0; ; StartAt += MaxResultsPerQuery )
+		const int maxResultsPerQuery = 100;
+		for( var startAt = 0; ; startAt += maxResultsPerQuery )
 		{
 			// Create the query
-			Dictionary<string, object> Request = new Dictionary<string, object>();
-			Request.Add( "jql", SearchQuery );
-			Request.Add( "startAt", StartAt.ToString() );
-			Request.Add( "maxResults", MaxResultsPerQuery.ToString() );
-			Request.Add( "fields", FieldsToGet );
-			HttpWebResponse Response = JiraRequest( "/search", JiraMethod.POST, Request, HttpStatusCode.OK );
+			var request = new Dictionary<string, object>();
+			request.Add( "jql", SearchQuery );
+			request.Add( "startAt", startAt.ToString() );
+			request.Add( "maxResults", maxResultsPerQuery.ToString() );
+			request.Add( "fields", FieldsToGet );
+			var response = JiraRequest( "/search", JiraMethod.POST, request, HttpStatusCode.OK );
 
 			// Parse the results
-			Dictionary<string, object> Results = ParseJiraResponse( Response );
-			System.Collections.ArrayList Issues = (System.Collections.ArrayList)Results["issues"];
+			var results = ParseJiraResponse( response );
+			var issues = (System.Collections.ArrayList)results["issues"];
 
 			// Process this chunk of results
-			foreach( Dictionary<string, object> Issue in Issues )
+			foreach( Dictionary<string, object> Issue in issues )
 			{
-				string Key = (string)Issue["key"];
-				Dictionary<string, object> Fields = (Dictionary<string, object>)Issue["fields"];
-				Result.Add( Key, Fields );
+				var key = (string)Issue["key"];
+				var fields = (Dictionary<string, object>)Issue["fields"];
+				result.Add( key, fields );
 			}
 
 			// Quit once we got less than the max results
-			if( Issues.Count < MaxResultsPerQuery )
+			if( issues.Count < maxResultsPerQuery )
 			{
 				break;
 			}
 		}
 
-		return Result;
+		return result;
 	}
 
 	/// <summary>
@@ -331,31 +332,48 @@ public class JiraConnection
 	/// </summary>
 	HttpWebResponse JiraRequest( string RequestUrl, JiraMethod Method, object Input )
 	{
-		// Get the username and password
-		string UserName, Password;
-		if( Credentials.TryGetValue( "jira.username", out UserName ) && Credentials.TryGetValue( "jira.password", out Password ) )
-		{
-			// Send the request
-			HttpWebRequest Request = (HttpWebRequest)WebRequest.Create( "https://jira.ol.epicgames.net/rest/api/2" + RequestUrl );
-			Request.Method = Enum.GetName( typeof( JiraMethod ), Method );
-			Request.Headers.Add( "Authorization", "Basic " + Convert.ToBase64String( UTF8Encoding.UTF8.GetBytes( UserName + ":" + Password ) ) );
-			Request.ContentType = "application/json";
-			if( Input != null )
-			{
-				using( StreamWriter RequestWriter = new StreamWriter( Request.GetRequestStream(), Encoding.ASCII ) )
-				{
-					JavaScriptSerializer Serializer = new JavaScriptSerializer();
-					RequestWriter.Write( Serializer.Serialize( Input ) );
-				}
-			}
-
-			// Check the response was ok
-			return (HttpWebResponse)Request.GetResponse();
-		}
-		else
-		{
-			return null;
-		}
+	    try
+	    {
+
+	        // Get the username and password
+	        string UserName, Password;
+	        if (Credentials.TryGetValue("jira.username", out UserName) &&
+	            Credentials.TryGetValue("jira.password", out Password))
+	        {
+	            // Send the request
+	            var request = (HttpWebRequest) WebRequest.Create("https://jira.it.epicgames.net/rest/api/2" + RequestUrl);
+	            request.Method = Enum.GetName(typeof (JiraMethod), Method);
+	            request.Headers.Add("Authorization",
+	                "Basic " + Convert.ToBase64String(UTF8Encoding.UTF8.GetBytes(UserName + ":" + Password)));
+	            request.ContentType = "application/json";
+	            if (Input != null)
+	            {
+	                using (var requestWriter = new StreamWriter(request.GetRequestStream(), Encoding.ASCII))
+	                {
+	                    var serializer = new JavaScriptSerializer();
+	                    requestWriter.Write(serializer.Serialize(Input));
+	                }
+	            }
+	            // Check the response was ok
+	            return (HttpWebResponse) request.GetResponse();
+	        }
+	        else
+	        {
+	            return null;
+	        }
+	    }
+	    catch (WebException ex)
+	    {
+            using( StreamReader ResponseReader = new StreamReader( ex.Response.GetResponseStream() ) )
+		    {
+			    string ResponseText = ResponseReader.ReadToEnd();
+		    }
+            return null;
+	    }
+	    catch (Exception)
+	    {
+            return null;
+	    }
 	}
 
 	/// <summary>
@@ -371,7 +389,7 @@ public class JiraConnection
 		return Response;
 	}
 
-	/// <summary>
+    /// <summary>
 	/// Parses a Json object from a web response
 	/// </summary>
 	public static Dictionary<string, object> ParseJiraResponse( HttpWebResponse Response )
diff --git a/Engine/Source/Programs/CrashReporter/CrashReportProcess/Config.cs b/Engine/Source/Programs/CrashReporter/CrashReportProcess/Config.cs
--- a/Engine/Source/Programs/CrashReporter/CrashReportProcess/Config.cs
+++ b/Engine/Source/Programs/CrashReporter/CrashReportProcess/Config.cs
@@ -82,6 +82,12 @@ namespace Tools.CrashReporter.CrashReportProcess
 		[XmlElement]
 		public string DataRouterLandingZone { get; set; }
 
+		/// <summary>
+		/// Folder where new crash reports are queued from the PS4 crash service.
+		/// </summary>
+		[XmlElement]
+		public string PS4LandingZone { get; set; }
+
 		/// <summary>
 		/// Folder where failed reports are moved.
 		/// </summary>
@@ -252,34 +258,70 @@ namespace Tools.CrashReporter.CrashReportProcess
 		public int MaxMemoryQueueSize { get; set; }
 
 		/// <summary>
-		/// AWSSDK service Url for S3 client
+		/// AWSSDK AWS credentials filepath containing the keys used to access SQS and S3
 		/// </summary>
 		[XmlElement]
-		public string AWSS3ServiceURL { get; set; }
+		public string AWSCredentialsFilepath { get; set; }
 
 		/// <summary>
-		/// AWSSDK service Url for SQS client
+		/// AWSSDK Profile name used in the AWS credentials file for reading crashes from DataRouter
 		/// </summary>
 		[XmlElement]
-		public string AWSSQSServiceURL { get; set; }
+		public string AWSProfileInputName { get; set; }
 
 		/// <summary>
-		/// AWSSDK queue Url for SQS client
+		/// AWSSDK service Url for S3 client reading crashes from DataRouter
 		/// </summary>
 		[XmlElement]
-		public string AWSSQSQueueUrl { get; set; }
+		public string AWSS3ServiceInputURL { get; set; }
 
 		/// <summary>
-		/// AWSSDK Profile name used in the AWS credentials file
+		/// AWSSDK service Url for SQS client reading crashes from DataRouter
 		/// </summary>
 		[XmlElement]
-		public string AWSProfileName { get; set; }
+		public string AWSSQSServiceInputURL { get; set; }
 
 		/// <summary>
-		/// AWSSDK AWS credentials filepath containing the keys used to access SQS and S3
+		/// AWSSDK queue Url for SQS client reading crashes from DataRouter
 		/// </summary>
 		[XmlElement]
-		public string AWSCredentialsFilepath { get; set; }
+		public string AWSSQSQueueInputUrl { get; set; }
+
+		/// <summary>
+		/// AWSSDK Profile name used in the AWS credentials file for writing crashes to S3
+		/// </summary>
+		[XmlElement]
+		public string AWSProfileOutputName { get; set; }
+
+		/// <summary>
+		/// AWSSDK service Url for S3 client for writing crashes to S3
+		/// </summary>
+		[XmlElement]
+		public string AWSS3ServiceOutputURL { get; set; }
+
+		/// <summary>
+		/// Should we output a copy of the crash report files to disk? (ProcessedReports, ProcessedVideos)
+		/// </summary>
+		[XmlElement]
+		public bool CrashFilesToDisk { get; set; }
+
+		/// <summary>
+		/// Should we output a copy of the crash report files to S3?
+		/// </summary>
+		[XmlElement]
+		public bool CrashFilesToAWS { get; set; }
+
+		/// <summary>
+		/// AWSSDK AWS S3 bucket used for output of crash reporter files (optional)
+		/// </summary>
+		[XmlElement]
+		public string AWSS3OutputBucket { get; set; }
+
+		/// <summary>
+		/// AWSSDK AWS S3 path/key prefix used for output of crash reporter files (suffix will be crash id and file name) (optional)
+		/// </summary>
+		[XmlElement]
+		public string AWSS3OutputKeyPrefix { get; set; }
 
 		/// <summary>
 		/// Buffer size used to decompress zlib archives taken from S3
@@ -288,7 +330,7 @@ namespace Tools.CrashReporter.CrashReportProcess
 		public int MaxUncompressedS3RecordSize { get; set; }
 
 		/// <summary>
-		/// Index file used to store all processed crash names and times. Stops duplicates.
+		/// Index file used to store all processed crash names and times. Stops duplicates. Leave blank to disable.
 		/// </summary>
 		[XmlElement]
 		public string ProcessedReportsIndexPath { get; set; }
@@ -366,13 +408,31 @@ namespace Tools.CrashReporter.CrashReportProcess
 			LoadedConfig.ProcessedReports = Path.Combine(LoadedConfig.DebugTestingFolder, "ProcessedReports");
 			LoadedConfig.ProcessedVideos = Path.Combine(LoadedConfig.DebugTestingFolder, "ProcessedVideos");
 			LoadedConfig.DepotRoot = Path.Combine(LoadedConfig.DebugTestingFolder, "DepotRoot");
-			LoadedConfig.InternalLandingZone = Path.Combine(LoadedConfig.DebugTestingFolder, "InternalLandingZone");
-			LoadedConfig.ExternalLandingZone = Path.Combine(LoadedConfig.DebugTestingFolder, "ExternalLandingZone");
-			LoadedConfig.DataRouterLandingZone = Path.Combine(LoadedConfig.DebugTestingFolder, "DataRouterLandingZone");
+			if (!string.IsNullOrWhiteSpace(LoadedConfig.InternalLandingZone))
+			{
+				LoadedConfig.InternalLandingZone = Path.Combine(LoadedConfig.DebugTestingFolder, "InternalLandingZone");
+			}
+			if (!string.IsNullOrWhiteSpace(LoadedConfig.ExternalLandingZone))
+			{
+				LoadedConfig.ExternalLandingZone = Path.Combine(LoadedConfig.DebugTestingFolder, "ExternalLandingZone");
+			}
+			if (!string.IsNullOrWhiteSpace(LoadedConfig.DataRouterLandingZone))
+			{
+				LoadedConfig.DataRouterLandingZone = Path.Combine(LoadedConfig.DebugTestingFolder, "DataRouterLandingZone");
+			}
+			if (!string.IsNullOrWhiteSpace(LoadedConfig.PS4LandingZone))
+			{
+				LoadedConfig.PS4LandingZone = Path.Combine(LoadedConfig.DebugTestingFolder, "PS4LandingZone");
+			}
 			LoadedConfig.InvalidReportsDirectory = Path.Combine(LoadedConfig.DebugTestingFolder, "InvalidReportsDirectory");
 			LoadedConfig.VersionString += " debugbuild";
 			LoadedConfig.AWSCredentialsFilepath = Path.Combine(LoadedConfig.DebugTestingFolder, "AWS", "credentials.ini");
-			LoadedConfig.ProcessedReportsIndexPath = Path.Combine(LoadedConfig.DebugTestingFolder, "ProcessedReports.ini");
+			if (!string.IsNullOrWhiteSpace(LoadedConfig.ProcessedReportsIndexPath))
+			{
+				LoadedConfig.ProcessedReportsIndexPath = Path.Combine(LoadedConfig.DebugTestingFolder, "ProcessedReports.ini");
+			}
+			LoadedConfig.CrashReportWebSite = string.Empty;
+			LoadedConfig.AWSS3OutputKeyPrefix = LoadedConfig.AWSS3OutputKeyPrefix.Replace("prod", "test");
 
 #if SLACKTESTING
 			LoadedConfig.SlackUsername = "CrashReportProcess_TESTING_IgnoreMe";
diff --git a/Engine/Source/Programs/CrashReporter/CrashReportProcess/Config/CrashReportProcess.config b/Engine/Source/Programs/CrashReporter/CrashReportProcess/Config/CrashReportProcess.config
--- a/Engine/Source/Programs/CrashReporter/CrashReportProcess/Config/CrashReportProcess.config
+++ b/Engine/Source/Programs/CrashReporter/CrashReportProcess/Config/CrashReportProcess.config
@@ -11,6 +11,7 @@
 	<InternalLandingZone>MyInternalCrashLandingZoneFolderPath</InternalLandingZone>
 	<ExternalLandingZone>MyExternallCrashLandingZoneFolderPath</ExternalLandingZone>
 	<DataRouterLandingZone>MyDataRouterCrashLandingZoneFolderPath</DataRouterLandingZone>
+	<PS4LandingZone>MyPS4CrashLandingZoneFolderPath</PS4LandingZone>
 	<InvalidReportsDirectory>MyInvalidCrashReportsFolderPath</InvalidReportsDirectory>
 	<QueueLowerLimitForDiscard>5000</QueueLowerLimitForDiscard>
 	<QueueUpperLimitForDiscard>20000</QueueUpperLimitForDiscard>
@@ -39,11 +40,17 @@
 	<ConsecutiveFailedWebAddLimit>10</ConsecutiveFailedWebAddLimit>
 	<MinDesiredMemoryQueueSize>50</MinDesiredMemoryQueueSize>
 	<MaxMemoryQueueSize>250</MaxMemoryQueueSize>
-	<AWSS3ServiceURL>https://s3.amazonaws.com</AWSS3ServiceURL>
-	<AWSSQSServiceURL>https://sqs.us-west-2.amazonaws.com</AWSSQSServiceURL>
-	<AWSSQSQueueUrl>https://sqs.us-west-2.amazonaws.com/queue.amazonaws.com/MyAccountNumber/MyQueueName</AWSSQSQueueUrl>
-	<AWSProfileName>default</AWSProfileName>
 	<AWSCredentialsFilepath>MyAWSCredsFilepath</AWSCredentialsFilepath>
+	<AWSProfileInputName>defaultin</AWSProfileInputName>
+	<AWSS3ServiceInputURL>https://s3.amazonaws.com</AWSS3ServiceInputURL>
+	<AWSSQSServiceInputURL>https://sqs.us-west-2.amazonaws.com</AWSSQSServiceInputURL>
+	<AWSSQSQueueInputUrl>https://sqs.us-west-2.amazonaws.com/queue.amazonaws.com/MyAccountNumber/MyQueueName</AWSSQSQueueInputUrl>
+	<AWSProfileOutputName>defaultout</AWSProfileOutputName>
+	<AWSS3ServiceOutputURL>https://s3.amazonaws.com</AWSS3ServiceOutputURL>
+	<CrashFilesToDisk>true</CrashFilesToDisk>
+	<CrashFilesToAWS>false</CrashFilesToAWS>
+	<AWSS3OutputBucket>crashreporterfiles</AWSS3OutputBucket>
+	<AWSS3OutputKeyPrefix>websitefiles</AWSS3OutputKeyPrefix>
 	<MaxUncompressedS3RecordSize>33554432</MaxUncompressedS3RecordSize>
 	<ProcessedReportsIndexPath>MyProcessedReportsIndexFilePath</ProcessedReportsIndexPath>
 	<ReportsIndexRetentionDays>120</ReportsIndexRetentionDays>
diff --git a/Engine/Source/Programs/CrashReporter/CrashReportProcess/CrashReportProcessService.cs b/Engine/Source/Programs/CrashReporter/CrashReportProcess/CrashReportProcessService.cs
--- a/Engine/Source/Programs/CrashReporter/CrashReportProcess/CrashReportProcessService.cs
+++ b/Engine/Source/Programs/CrashReporter/CrashReportProcess/CrashReportProcessService.cs
@@ -11,6 +11,9 @@ using System.Reflection;
 using System.ServiceProcess;
 using System.Text;
 using System.Threading.Tasks;
+using Amazon.Runtime;
+using Amazon.S3;
+using Amazon.SQS;
 using Tools.CrashReporter.CrashReportCommon;
 using Tools.CrashReporter.CrashReportProcess.Properties;
 
@@ -32,12 +35,21 @@ namespace Tools.CrashReporter.CrashReportProcess
 
 		private static SlackWriter Slack = null;
 
+		public static AmazonClient DataRouterAWS { get; private set; }
+
+		public static AmazonClient OutputAWS { get; private set; }
+
 		public static StatusReporting StatusReporter { get; private set; }
 
 		public static Symbolicator Symbolicator { get; private set; }
 
+		public static ReportIndex ReportIndex { get; private set; }
+
 		/// <summary>Folder in which to store log files</summary>
-		static private string LogFolder = null;
+		private static string LogFolder = null;
+
+		private bool bDisposing = false;
+		private bool bDisposed = false;
 
 		/// <summary>Folder in which to store symbolication log files</summary>
 		static public string SymbolicatorLogFolder { get; private set; }
@@ -89,13 +101,25 @@ namespace Tools.CrashReporter.CrashReportProcess
 		/// <summary>
 		/// Write an exception message to the event log.
 		/// </summary>
-		static public void WriteException( string Message )
+		static public void WriteException( string Message, Exception Ex )
 		{
-			if( Message != null && Message.Length > 2 )
+			if (Ex != null)
 			{
-				Log.Print( "[EXCEPT] " + Message );
-				StatusReporter.IncrementCount(StatusReportingEventNames.ExceptionEvent);
+				if (Ex is OutOfMemoryException)
+				{
+					StatusReporter.Alert("OOM", "Out of memory.", 2);
+				}
 			}
+
+			if (Message != null && Message.Length > 2)
+			{
+				Log.Print("[EXCEPT] " + Message);
+			}
+			else
+			{
+				Log.Print("[EXCEPT] NO MESSAGE");
+			}
+			StatusReporter.IncrementCount(StatusReportingEventNames.ExceptionEvent);
 		}
 
 		/// <summary>
@@ -144,6 +168,44 @@ namespace Tools.CrashReporter.CrashReportProcess
 
 			StatusReporter = new StatusReporting();
 
+			ReportIndex = new ReportIndex
+			{
+				IsEnabled = !string.IsNullOrWhiteSpace(Config.Default.ProcessedReportsIndexPath),
+                Filepath = Config.Default.ProcessedReportsIndexPath,
+				Retention = TimeSpan.FromDays(Config.Default.ReportsIndexRetentionDays)
+			};
+
+			ReportIndex.ReadFromFile();
+
+			WriteEvent("Initializing AWS");
+			string AWSError;
+			AWSCredentials AWSCredentialsForDataRouter = new StoredProfileAWSCredentials(Config.Default.AWSProfileInputName, Config.Default.AWSCredentialsFilepath);
+			AmazonSQSConfig SqsConfigForDataRouter = new AmazonSQSConfig
+			{
+				ServiceURL = Config.Default.AWSSQSServiceInputURL
+			};
+			AmazonS3Config S3ConfigForDataRouter = new AmazonS3Config
+			{
+				ServiceURL = Config.Default.AWSS3ServiceInputURL
+			};
+			DataRouterAWS = new AmazonClient(AWSCredentialsForDataRouter, SqsConfigForDataRouter, S3ConfigForDataRouter, out AWSError);
+			if (!DataRouterAWS.IsSQSValid || !DataRouterAWS.IsS3Valid)
+			{
+				WriteFailure("AWS failed to initialize profile for DataRouter access. Error:" + AWSError);
+				StatusReporter.Alert("AWSFailInput", "AWS failed to initialize profile for DataRouter access", 0);
+			}
+			AWSCredentials AWSCredentialsForOutput = new StoredProfileAWSCredentials(Config.Default.AWSProfileOutputName, Config.Default.AWSCredentialsFilepath);
+			AmazonS3Config S3ConfigForOutput = new AmazonS3Config
+			{
+				ServiceURL = Config.Default.AWSS3ServiceOutputURL
+			};
+			OutputAWS = new AmazonClient(AWSCredentialsForOutput, null, S3ConfigForOutput, out AWSError);
+			if (!OutputAWS.IsS3Valid)
+			{
+				WriteFailure("AWS failed to initialize profile for output S3 bucket access. Error:" + AWSError);
+				StatusReporter.Alert("AWSFailOutput", "AWS failed to initialize profile for output S3 bucket access", 0);
+			}
+
 			// Add directory watchers
 			WriteEvent("Creating ReportWatcher");
 			Watcher = new ReportWatcher();
@@ -167,6 +229,7 @@ namespace Tools.CrashReporter.CrashReportProcess
 			FoldersToMonitor.Add(Config.Default.DepotRoot, "P4 Workspace");
 			FoldersToMonitor.Add(Config.Default.InternalLandingZone, "CRR Landing Zone");
 			FoldersToMonitor.Add(Config.Default.DataRouterLandingZone, "Data Router Landing Zone");
+			FoldersToMonitor.Add(Config.Default.PS4LandingZone, "PS4 Landing Zone");
 			FoldersToMonitor.Add(Config.Default.InvalidReportsDirectory, "Invalid Reports");
 			FoldersToMonitor.Add(Assembly.GetExecutingAssembly().Location, "CRP Binaries and Logs");
 			FoldersToMonitor.Add(Config.Default.MDDPDBCachePath, "MDD PDB Cache");
@@ -193,6 +256,13 @@ namespace Tools.CrashReporter.CrashReportProcess
 		{
 			StatusReporter.OnPreStopping();
 
+			OnDispose();
+		}
+
+		private void OnDispose()
+		{
+			bDisposing = true;
+
 			// Clean up the directory watcher and crash processor threads
 			foreach (var Processor in Processors)
 			{
@@ -207,6 +277,15 @@ namespace Tools.CrashReporter.CrashReportProcess
 			Watcher.Dispose();
 			Watcher = null;
 
+			ReportIndex.WriteToFile();
+			ReportIndex = null;
+
+			OutputAWS.Dispose();
+			OutputAWS = null;
+
+			DataRouterAWS.Dispose();
+			DataRouterAWS = null;
+
 			StatusReporter.Dispose();
 			StatusReporter = null;
 
@@ -216,6 +295,8 @@ namespace Tools.CrashReporter.CrashReportProcess
 			// Flush the log to disk
 			Log.Dispose();
 			Log = null;
+
+			bDisposed = true;
 		}
 
 		/// <summary>
diff --git a/Engine/Source/Programs/CrashReporter/CrashReportProcess/DataRouterReportQueue.cs b/Engine/Source/Programs/CrashReporter/CrashReportProcess/DataRouterReportQueue.cs
--- a/Engine/Source/Programs/CrashReporter/CrashReportProcess/DataRouterReportQueue.cs
+++ b/Engine/Source/Programs/CrashReporter/CrashReportProcess/DataRouterReportQueue.cs
@@ -30,29 +30,6 @@ namespace Tools.CrashReporter.CrashReportProcess
 		public DataRouterReportQueue(string InQueueName, string InLandingZoneTempPath, int InDecimateWaitingCountStart, int InDecimateWaitingCountEnd)
 			: base(InQueueName, InLandingZoneTempPath, InDecimateWaitingCountStart, InDecimateWaitingCountEnd)
 		{
-			AWSCredentials Credentials = new StoredProfileAWSCredentials(Config.Default.AWSProfileName, Config.Default.AWSCredentialsFilepath);
-
-			AmazonSQSConfig SqsConfig = new AmazonSQSConfig
-			{
-				ServiceURL = Config.Default.AWSSQSServiceURL
-			};
-
-			SqsClient = new AmazonSQSClient(Credentials, SqsConfig);
-
-			AmazonS3Config S3Config = new AmazonS3Config
-			{
-				ServiceURL = Config.Default.AWSS3ServiceURL
-			};
-
-			S3Client = new AmazonS3Client(Credentials, S3Config);
-		}
-
-		public override void Dispose()
-		{
-			SqsClient.Dispose();
-			S3Client.Dispose();
-
-			base.Dispose();
 		}
 
 		protected override int GetTotalWaitingCount()
@@ -105,15 +82,9 @@ namespace Tools.CrashReporter.CrashReportProcess
 					string S3Key = RecordPair[1];
 					string ReadableRequestString = "Bucket=" + S3BucketName + " Key=" + S3Key;
 
-					var ObjectRequest = new GetObjectRequest
-					{
-						BucketName = S3BucketName,
-						Key = S3Key
-					};
-
 					using (Stream ProtocolBufferStream = new MemoryStream())
 					{
-						using (GetObjectResponse ObjectResponse = S3Client.GetObject(ObjectRequest))
+						using (GetObjectResponse ObjectResponse = CrashReporterProcessServicer.DataRouterAWS.GetS3Object(S3BucketName, S3Key))
 						{
 							using (Stream ResponseStream = ObjectResponse.ResponseStream)
 							{
@@ -129,11 +100,11 @@ namespace Tools.CrashReporter.CrashReportProcess
 						NewCrashCount += UnpackRecordsFromDelimitedProtocolBuffers(ProtocolBufferStream, LandingZone, ReadableRequestString);
 					}
 				}
-				catch (Exception ex)
+				catch (Exception Ex)
 				{
 					CrashReporterProcessServicer.StatusReporter.IncrementCount(StatusReportingEventNames.ReadS3FileFailedEvent);
 					CrashReporterProcessServicer.WriteException("TryGetNewS3Crashes: failure during processing SQS record " + SQSRecord +
-					                                            "\n" + ex);
+					                                            "\n" + Ex, Ex);
 				}
 			}
 		}
@@ -219,7 +190,7 @@ namespace Tools.CrashReporter.CrashReportProcess
 					// Early check for duplicate processed report
 					lock (ReportIndexLock)
 					{
-						if (ReportIndex.ContainsReport(CrashFolderName))
+						if (CrashReporterProcessServicer.ReportIndex.ContainsReport(CrashFolderName))
 						{
 							// Crash report not accepted by index
 							CrashReporterProcessServicer.WriteEvent(string.Format(
@@ -271,7 +242,7 @@ namespace Tools.CrashReporter.CrashReportProcess
 					// Early check for duplicate processed report
 					lock (ReportIndexLock)
 					{
-						if (ReportIndex.ContainsReport(CrashFolderName))
+						if (CrashReporterProcessServicer.ReportIndex.ContainsReport(CrashFolderName))
 						{
 							// Crash report not accepted by index
 							CrashReporterProcessServicer.WriteEvent(string.Format(
@@ -335,21 +306,11 @@ namespace Tools.CrashReporter.CrashReportProcess
 		{
 			try
 			{
-				var AttribRequest = new GetQueueAttributesRequest
-				{
-					QueueUrl = Config.Default.AWSSQSQueueUrl,
-					AttributeNames = new List<string>
-					{
-						"ApproximateNumberOfMessages"
-					}
-				};
-
-				var AttribResponse = SqsClient.GetQueueAttributes(AttribRequest);
-				return AttribResponse.ApproximateNumberOfMessages;
+				return CrashReporterProcessServicer.DataRouterAWS.GetSQSQueueCount(Config.Default.AWSSQSQueueInputUrl);
 			}
 			catch (Exception Ex)
 			{
-				CrashReporterProcessServicer.WriteException("GetSQSCount: " + Ex.ToString());
+				CrashReporterProcessServicer.WriteException("GetSQSCount: " + Ex, Ex);
 			}
 			return 0;
 		}
@@ -360,57 +321,26 @@ namespace Tools.CrashReporter.CrashReportProcess
 
 			try
 			{
-				var ReceiveRequest = new ReceiveMessageRequest
-				{
-					QueueUrl = Config.Default.AWSSQSQueueUrl,
-					MaxNumberOfMessages = 1
-				};
-
-				var ReceiveResponse = SqsClient.ReceiveMessage(ReceiveRequest);
+#if DEBUG
+				bool bNoDeleteFromSQS = true;
+#else
+				bool bNoDeleteFromSQS = false;
+#endif
 
-				if (ReceiveResponse.Messages.Count == 1)
+				Message DequeuedMessage = CrashReporterProcessServicer.DataRouterAWS.SQSDequeueMessage(Config.Default.AWSSQSQueueInputUrl, bNoDeleteFromSQS);
+				if (DequeuedMessage != null)
 				{
-					var Message = ReceiveResponse.Messages[0];
-
-					if (Message != null && TryDeleteRecordSQS(Message))
-					{
-						OutRecordString = Message.Body;
-						return true;
-					}
+					OutRecordString = DequeuedMessage.Body;
+					return true;
 				}
 			}
 			catch (Exception Ex)
 			{
-				CrashReporterProcessServicer.WriteException("DequeueRecordSQS: " + Ex.ToString());
+				CrashReporterProcessServicer.WriteException("DequeueRecordSQS: " + Ex, Ex);
 			}
 			return false;
 		}
 
-		private bool TryDeleteRecordSQS(Message InMessage)
-		{
-#if DEBUG
-			// Debug doesn't empty the queue - it's just reads records
-			return true;
-#else
-			try
-			{
-				var DeleteRequest = new DeleteMessageRequest
-				{
-					QueueUrl = Config.Default.AWSSQSQueueUrl,
-					ReceiptHandle = InMessage.ReceiptHandle
-				};
-
-				var DeleteResponse = SqsClient.DeleteMessage(DeleteRequest);
-				return DeleteResponse.HttpStatusCode == HttpStatusCode.OK;
-			}
-			catch (Exception Ex)
-			{
-				CrashReporterProcessServicer.WriteException("TryDeleteRecordSQS: " + Ex.ToString());
-			}
-			return false;
-#endif
-        }
-
         private class CrashHeader
 		{
 			public const int FixedSize = 4 + 260 + 4 + 260 + 4 + 4;		// 2 x 260 ansi char strings with length + 2 x int32
@@ -430,8 +360,5 @@ namespace Tools.CrashReporter.CrashReportProcess
 				return OutCrashHeader;
 			}
 		}
-
-		private readonly AmazonSQSClient SqsClient;
-		private readonly AmazonS3Client S3Client;
 	}
 }
diff --git a/Engine/Source/Programs/CrashReporter/CrashReportProcess/ReceiverReportQueue.cs b/Engine/Source/Programs/CrashReporter/CrashReportProcess/ReceiverReportQueue.cs
--- a/Engine/Source/Programs/CrashReporter/CrashReportProcess/ReceiverReportQueue.cs
+++ b/Engine/Source/Programs/CrashReporter/CrashReportProcess/ReceiverReportQueue.cs
@@ -12,21 +12,25 @@ namespace Tools.CrashReporter.CrashReportProcess
 		{ 
 			get
 			{
-				return StatusReportingEventNames.ProcessingStartedReceiverEvent;
+				return ProcessingStartedEventName;
 			}
 		}
 
 		/// <summary>
 		/// Constructor taking the landing zone
 		/// </summary>
-		public ReceiverReportQueue(string InQueueName, string LandingZonePath, int InDecimateWaitingCountStart, int InDecimateWaitingCountEnd)
+		public ReceiverReportQueue(string InQueueName, string LandingZonePath, string InProcessingStartedEventName, int InDecimateWaitingCountStart,
+		                           int InDecimateWaitingCountEnd)
 			: base(InQueueName, LandingZonePath, InDecimateWaitingCountStart, InDecimateWaitingCountEnd)
 		{
+			ProcessingStartedEventName = InProcessingStartedEventName;
 		}
 
 		protected override int GetTotalWaitingCount()
 		{
 			return LastQueueSizeOnDisk;
 		}
+
+		private readonly string ProcessingStartedEventName;
 	}
 }
\ No newline at end of file
diff --git a/Engine/Source/Programs/CrashReporter/CrashReportProcess/ReportIndex.cs b/Engine/Source/Programs/CrashReporter/CrashReportProcess/ReportIndex.cs
--- a/Engine/Source/Programs/CrashReporter/CrashReportProcess/ReportIndex.cs
+++ b/Engine/Source/Programs/CrashReporter/CrashReportProcess/ReportIndex.cs
@@ -7,10 +7,12 @@ using System.Threading.Tasks;
 
 namespace Tools.CrashReporter.CrashReportProcess
 {
-	static class ReportIndex
+	class ReportIndex
 	{
-		public static void WriteToFile()
+		public void WriteToFile()
 		{
+			if (!IsEnabled) return;
+
 			try
 			{
 				if (File.Exists(Filepath))
@@ -34,15 +36,17 @@ namespace Tools.CrashReporter.CrashReportProcess
 
 				LastFlush = DateTime.UtcNow;
 			}
-			catch (Exception ex)
+			catch (Exception Ex)
 			{
-				CrashReporterProcessServicer.WriteException(string.Format("Failed to write ReportIndex to {0}. Exception was: {1}", Filepath, ex));
+				CrashReporterProcessServicer.WriteException(string.Format("Failed to write ReportIndex to {0}. Exception was: {1}", Filepath, Ex), Ex);
 				CrashReporterProcessServicer.StatusReporter.Alert("ReportIndex.WriteToFile", string.Format("Failed to write ReportIndex to {0}", Filepath), Config.Default.SlackAlertRepeatMinimumMinutes);
 			}
 		}
 
-		public static void ReadFromFile()
+		public void ReadFromFile()
 		{
+			if (!IsEnabled) return;
+
 			try
 			{
 				Index = new Dictionary<string, DateTime>();
@@ -83,15 +87,17 @@ namespace Tools.CrashReporter.CrashReportProcess
 
 				LastFlush = DateTime.UtcNow;
 			}
-			catch (Exception ex)
+			catch (Exception Ex)
 			{
-				CrashReporterProcessServicer.WriteException(string.Format("Failed to read ReportIndex from {0}. Exception was: {1}", Filepath, ex));
+				CrashReporterProcessServicer.WriteException(string.Format("Failed to read ReportIndex from {0}. Exception was: {1}", Filepath, Ex), Ex);
 				CrashReporterProcessServicer.StatusReporter.Alert("ReportIndex.ReadFromFile", string.Format("Failed to read ReportIndex from {0}", Filepath), Config.Default.SlackAlertRepeatMinimumMinutes);
 			}
 		}
 
-		public static bool TryAddNewReport(string ReportKey)
+		public bool TryAddNewReport(string ReportKey)
 		{
+			if (!IsEnabled) return true;
+
 			if (ContainsReport(ReportKey))
 			{
 				return false;
@@ -107,13 +113,17 @@ namespace Tools.CrashReporter.CrashReportProcess
 			return true;
 		}
 
-		public static bool ContainsReport(string ReportKey)
+		public bool ContainsReport(string ReportKey)
 		{
+			if (!IsEnabled) return false;
+
 			return Index.ContainsKey(ReportKey);
 		}
 
-		public static bool TryRemoveReport(string ReportKey)
+		public bool TryRemoveReport(string ReportKey)
 		{
+			if (!IsEnabled) return true;
+
 			return Index.Remove(ReportKey);
 		}
 
@@ -149,10 +159,11 @@ namespace Tools.CrashReporter.CrashReportProcess
 			return true;
 		}
 
-		public static string Filepath { get; set; }
-		public static TimeSpan Retention { get; set; }
+		public bool IsEnabled { get; set; }
+		public string Filepath { get; set; }
+		public TimeSpan Retention { get; set; }
 
-		private static Dictionary<string, DateTime> Index = new Dictionary<string, DateTime>();
-		private static DateTime LastFlush = DateTime.MinValue;
+		private Dictionary<string, DateTime> Index = new Dictionary<string, DateTime>();
+		private DateTime LastFlush = DateTime.MinValue;
 	}
 }
diff --git a/Engine/Source/Programs/CrashReporter/CrashReportProcess/ReportProcessor.cs b/Engine/Source/Programs/CrashReporter/CrashReportProcess/ReportProcessor.cs
--- a/Engine/Source/Programs/CrashReporter/CrashReportProcess/ReportProcessor.cs
+++ b/Engine/Source/Programs/CrashReporter/CrashReportProcess/ReportProcessor.cs
@@ -10,6 +10,7 @@ using System.Net;
 using System.Text;
 using System.Threading;
 using System.Threading.Tasks;
+using Amazon.S3.Model;
 using Tools.CrashReporter.CrashReportCommon;
 using Tools.DotNETCommon.LaunchProcess;
 using Tools.DotNETCommon.SimpleWebRequest;
@@ -30,7 +31,7 @@ namespace Tools.CrashReporter.CrashReportProcess
 		const int SyncTimeoutSeconds = 2 * 60;
 
 		/// <summary> Number of all processed reports. </summary>
-		public static int ProcessedReports = 1;
+		public static int ProcessedReportCount = 1;
 
 		/// <summary> Global timer used to measure web added reports per day. </summary>
 		public static Stopwatch Timer = Stopwatch.StartNew();
@@ -143,7 +144,7 @@ namespace Tools.CrashReporter.CrashReportProcess
 					}
 					catch (Exception Ex)
 					{
-						CrashReporterProcessServicer.WriteException(string.Format("PROC-{0} ", ProcessorIndex) + "ProcessNewReports: " + Ex.ToString());
+						CrashReporterProcessServicer.WriteException(string.Format("PROC-{0} ", ProcessorIndex) + "ProcessNewReports: " + Ex, Ex);
 					}
 
 					TickStatic(Watcher);
@@ -229,7 +230,7 @@ namespace Tools.CrashReporter.CrashReportProcess
 				{
 					if( bWriteException )
 					{
-						CrashReporterProcessServicer.WriteException( "CleanReport: " + Ex.ToString() );
+						CrashReporterProcessServicer.WriteException( "CleanReport: " + Ex, Ex);
 						bWriteException = false;
 					}
 
@@ -283,9 +284,9 @@ namespace Tools.CrashReporter.CrashReportProcess
 				CrashReporterProcessServicer.WriteEvent( string.Format( "Moved to {0}", DestinationDirectory ) );
 				UpdateProcessedReports();
 			}
-			catch( System.Exception Ex )
+			catch( Exception Ex )
 			{
-				CrashReporterProcessServicer.WriteException( "MoveReportToInvalid: " + Ex.ToString() );
+				CrashReporterProcessServicer.WriteException("MoveReportToInvalid: " + Ex, Ex);
 			}	
 		}
 
@@ -295,7 +296,7 @@ namespace Tools.CrashReporter.CrashReportProcess
 		/// </summary>
 		private static void UpdateProcessedReports()
 		{
-			Interlocked.Increment( ref ProcessedReports );
+			Interlocked.Increment( ref ProcessedReportCount );
 		}
 
 		/// <summary> 
@@ -315,7 +316,7 @@ namespace Tools.CrashReporter.CrashReportProcess
 			}
 			catch( Exception Ex )
 			{
-				CrashReporterProcessServicer.WriteException( "CleanRepository: " + Ex.ToString() );
+				CrashReporterProcessServicer.WriteException("CleanRepository: " + Ex, Ex);
 			}
 		}
 
@@ -409,7 +410,7 @@ namespace Tools.CrashReporter.CrashReportProcess
 			}
 			catch (Exception Ex)
 			{
-				CrashReporterProcessServicer.WriteException(string.Format("PROC-{0} ", ProcessorIndex) + "CreateCrash: " + Ex.ToString());
+				CrashReporterProcessServicer.WriteException(string.Format("PROC-{0} ", ProcessorIndex) + "CreateCrash: " + Ex, Ex);
 			}
 
 			return XmlPayload;
@@ -478,7 +479,7 @@ namespace Tools.CrashReporter.CrashReportProcess
 			}
 			catch( Exception Ex )
 			{
-				CrashReporterProcessServicer.WriteException(string.Format("PROC-{0} ", ProcessorIndex) + "UploadCrash: " + Ex.ToString());
+				CrashReporterProcessServicer.WriteException(string.Format("PROC-{0} ", ProcessorIndex) + "UploadCrash: " + Ex, Ex);
 			}
 
 			return NewID;
@@ -531,54 +532,84 @@ namespace Tools.CrashReporter.CrashReportProcess
 					ConsecutiveFailedUploads = 0;
 				}
 
-				string IDThenUnderscore = string.Format("{0}_", ReportID);
+				bool bToS3 = Config.Default.CrashFilesToAWS && CrashReporterProcessServicer.OutputAWS.IsS3Valid;
+				bool bToDisk = Config.Default.CrashFilesToDisk;
 
 				// Use the row id to name and move the files the way the web site requires
+				string IDThenUnderscore = string.Format("{0}_", ReportID);
+				int ReportIDSegment = (ReportID/10000)*10000;
+				string S3IDPrefix = string.Format("/{0}/{1}/{1}_", ReportIDSegment, ReportID);
 				string DestinationFolder = Path.Combine(Config.Default.ProcessedReports, IDThenUnderscore);
-				CrashReporterProcessServicer.StatusReporter.AlertOnLowDisk(DestinationFolder, Config.Default.DiskSpaceAlertPercent);
+				Stopwatch WriteToS3Timer = new Stopwatch();
+				Stopwatch WriteToDiskTimer = new Stopwatch();
+
+				if (bToDisk)
+				{
+					CrashReporterProcessServicer.StatusReporter.AlertOnLowDisk(DestinationFolder, Config.Default.DiskSpaceAlertPercent);
+				}
 
-				// Move report files to crash reporter file store
+				// Save log file
 				if (LogFileName != null)
 				{
 					LogFileName = Path.Combine(DirInfo.FullName, LogFileName);
 					FileInfo LogInfo = new FileInfo(LogFileName);
-					if(LogInfo.Exists)
+					if (LogInfo.Exists)
 					{
-						LogInfo.MoveTo(DestinationFolder + "Launch.log");
+						if (bToS3)
+						{
+							WriteToS3Timer.Start();
+							UploadFileToS3(LogInfo, S3IDPrefix + "Launch.log");
+							WriteToS3Timer.Stop();
+						}
+						if (bToDisk)
+						{
+							WriteToDiskTimer.Start();
+							LogInfo.MoveTo(DestinationFolder + "Launch.log");
+							WriteToDiskTimer.Stop();
+						}
 					}
 				}
 
-				string CrashContextRuntimeName = Path.Combine( DirInfo.FullName, FGenericCrashContext.CrashContextRuntimeXMLName );
-				FileInfo CrashContextInfo = new FileInfo( CrashContextRuntimeName );
+				// Save crash context file
+				string CrashContextRuntimeName = Path.Combine(DirInfo.FullName, FGenericCrashContext.CrashContextRuntimeXMLName);
+				FileInfo CrashContextInfo = new FileInfo(CrashContextRuntimeName);
 				if (CrashContextInfo.Exists)
 				{
-					CrashContextInfo.MoveTo( DestinationFolder + FGenericCrashContext.CrashContextRuntimeXMLName );
+					if (bToS3)
+					{
+						WriteToS3Timer.Start();
+						UploadFileToS3(CrashContextInfo, S3IDPrefix + FGenericCrashContext.CrashContextRuntimeXMLName);
+						WriteToS3Timer.Stop();
+					}
+					if (bToDisk)
+					{
+						WriteToDiskTimer.Start();
+						CrashContextInfo.MoveTo(DestinationFolder + FGenericCrashContext.CrashContextRuntimeXMLName);
+						WriteToDiskTimer.Stop();
+					}
 				}
 
-// 				WERMetaDataName = Path.Combine(DirInfo.FullName, WERMetaDataName);
-// 				FileInfo MetaInfo = new FileInfo(WERMetaDataName);
-// 				if (MetaInfo.Exists)
-// 				{
-// 					MetaInfo.MoveTo(DestinationFolder + "WERMeta.xml");
-// 				}
-
-				if( DumpFileName != null )
+				if (DumpFileName != null)
 				{
-					DumpFileName = Path.Combine( DirInfo.FullName, DumpFileName );
-					FileInfo DumpInfo = new FileInfo( DumpFileName );
-					if (DumpInfo.Exists && NewContext.PrimaryCrashProperties.CrashDumpMode != 1 /* ECrashDumpMode.FullDump = 1*/ )
+					DumpFileName = Path.Combine(DirInfo.FullName, DumpFileName);
+					FileInfo DumpInfo = new FileInfo(DumpFileName);
+					if (DumpInfo.Exists && NewContext.PrimaryCrashProperties.CrashDumpMode != 1 /* ECrashDumpMode.FullDump = 1*/)
 					{
-						DumpInfo.MoveTo( DestinationFolder + "MiniDump.dmp" );
+						if (bToS3)
+						{
+							WriteToS3Timer.Start();
+							UploadFileToS3(DumpInfo, S3IDPrefix + "MiniDump.dmp");
+							WriteToS3Timer.Stop();
+						}
+						if (bToDisk)
+						{
+							WriteToDiskTimer.Start();
+							DumpInfo.MoveTo(DestinationFolder + "MiniDump.dmp");
+							WriteToDiskTimer.Stop();
+						}
 					}
 				}
 
-// 				DiagnosticsFileName = Path.Combine(DirInfo.FullName, DiagnosticsFileName);
-// 				FileInfo DiagnosticsInfo = new FileInfo(DiagnosticsFileName);
-// 				if (DiagnosticsInfo.Exists)
-// 				{
-// 					DiagnosticsInfo.MoveTo(DestinationFolder + CrashReporterConstants.DiagnosticsFileName);
-// 				}
-
 				// Move the video (if it exists) to an alternate store
 				if (VideoFileName != null)
 				{
@@ -588,34 +619,75 @@ namespace Tools.CrashReporter.CrashReportProcess
 					FileInfo VideoInfo = new FileInfo(VideoFileName);
 					if (VideoInfo.Exists)
 					{
-						CrashReporterProcessServicer.StatusReporter.AlertOnLowDisk(DestinationFolder, Config.Default.DiskSpaceAlertPercent);
-						VideoInfo.MoveTo(DestinationFolder + CrashReporterConstants.VideoFileName);
+						if (bToS3)
+						{
+							WriteToS3Timer.Start();
+							UploadFileToS3(VideoInfo, S3IDPrefix + CrashReporterConstants.VideoFileName);
+							WriteToS3Timer.Stop();
+						}
+						if (bToDisk)
+						{
+							CrashReporterProcessServicer.StatusReporter.AlertOnLowDisk(DestinationFolder, Config.Default.DiskSpaceAlertPercent);
+
+							WriteToDiskTimer.Start();
+							VideoInfo.MoveTo(DestinationFolder + CrashReporterConstants.VideoFileName);
+							WriteToDiskTimer.Stop();
+						}
 					}
 				}
 
-				CrashReporterProcessServicer.WriteEvent(string.Format("PROC-{0} ", ProcessorIndex) + "# WebAdded: ReportID   =" + string.Format("{0,7}", ReportID) + " Path=" + NewContext.CrashDirectory);
+				string TimeTakenString = string.Empty;
+				if (bToS3)
+				{
+					TimeTakenString = string.Format("S3UploadTime={0:F1} ", WriteToS3Timer.Elapsed.TotalSeconds);
+				}
+				if (bToDisk)
+				{
+					TimeTakenString += string.Format("DiskMoveTime={0:F1} ", WriteToDiskTimer.Elapsed.TotalSeconds);
+				}
+
+				CrashReporterProcessServicer.WriteEvent(string.Format("PROC-{0} AddReport: ReportID={1,8} {2}Path={3}", ProcessorIndex, ReportID, TimeTakenString, NewContext.CrashDirectory));
 
 				UpdateProcessedReports();
 				WebAddCounter.AddEvent();
 				CrashReporterProcessServicer.StatusReporter.IncrementCount(StatusReportingEventNames.ProcessingSucceededEvent);
-				double Ratio = (double)WebAddCounter.TotalEvents / (double)ProcessedReports * 100;
+				double Ratio = (double)WebAddCounter.TotalEvents / (double)ProcessedReportCount * 100;
 
 				double AddedPerDay = (double)WebAddCounter.TotalEvents / Timer.Elapsed.TotalDays;
 
 				CrashReporterProcessServicer.WriteEvent(string.Format("PROC-{0} ", ProcessorIndex) + 
 					string.Format(
-						"Ratio={0,2} Processed={1,7} WebAdded={2,7} AddReportTime={3} AddedPerDay={4} AddedPerMinute={5:N1}", (int) Ratio,
-						ProcessedReports, WebAddCounter.TotalEvents, AddReportTime.Elapsed.TotalSeconds.ToString("0.00"), (int)AddedPerDay, WebAddCounter.EventsPerSecond * 60));
+						"AddReport: Ratio={0,2} Processed={1,7} WebAdded={2,7} AddReportTime={3} AddedPerDay={4} AddedPerMinute={5:N1}", (int) Ratio,
+						ProcessedReportCount, WebAddCounter.TotalEvents, AddReportTime.Elapsed.TotalSeconds.ToString("0.00"), (int)AddedPerDay, WebAddCounter.EventsPerSecond * 60));
 				return true;
 			}
 			catch( Exception Ex )
 			{
-				CrashReporterProcessServicer.WriteException(string.Format("PROC-{0} ", ProcessorIndex) + "AddReport: " + DirInfo.Name + "\n\n" + Ex.ToString());
+				CrashReporterProcessServicer.WriteException(string.Format("PROC-{0} ", ProcessorIndex) + "AddReport: " + DirInfo.Name + "\n\n" + Ex, Ex);
 			}
 
 			return false;
 		}
 
+		private void UploadFileToS3(FileInfo FileInfo, string DestFilename)
+		{
+			try
+			{
+				PutObjectResponse Response = CrashReporterProcessServicer.OutputAWS.PutS3ObjectFromFile(Config.Default.AWSS3OutputBucket,
+																										Config.Default.AWSS3OutputKeyPrefix + DestFilename,
+																										FileInfo.FullName);
+
+				if (Response == null || Response.HttpStatusCode != HttpStatusCode.OK)
+				{
+					throw new CrashReporterException(string.Format("Failed to upload {0} to {1}", FileInfo.FullName, DestFilename));
+				}
+			}
+			catch (Exception Ex)
+			{
+				CrashReporterProcessServicer.WriteException("UploadFileToS3: " + Ex, Ex);
+			}
+		}
+
 		/// <summary>
 		/// Sync the MinidumpDiagnostics binary and the engine config files to #head.
 		/// </summary>
@@ -789,7 +861,7 @@ namespace Tools.CrashReporter.CrashReportProcess
 			}
 			catch( Exception Ex )
 			{
-				CrashReporterProcessServicer.WriteException(string.Format("PROC-{0} ", ProcessorIndex) + "ProcessReport: " + NewContext.CrashDirectory + "\n\n: " + Ex.ToString());
+				CrashReporterProcessServicer.WriteException(string.Format("PROC-{0} ", ProcessorIndex) + "ProcessReport: " + NewContext.CrashDirectory + "\n\n: " + Ex, Ex);
 			}
 
 			return 0.0;
diff --git a/Engine/Source/Programs/CrashReporter/CrashReportProcess/ReportQueue.cs b/Engine/Source/Programs/CrashReporter/CrashReportProcess/ReportQueue.cs
--- a/Engine/Source/Programs/CrashReporter/CrashReportProcess/ReportQueue.cs
+++ b/Engine/Source/Programs/CrashReporter/CrashReportProcess/ReportQueue.cs
@@ -113,7 +113,7 @@ namespace Tools.CrashReporter.CrashReportProcess
 							}
 							catch (Exception Ex)
 							{
-								CrashReporterProcessServicer.WriteException("CheckForNewReportsInner: " + Ex.ToString());
+								CrashReporterProcessServicer.WriteException("CheckForNewReportsInner: " + Ex, Ex);
 								ReportProcessor.MoveReportToInvalid(NewReportPath, "NEWRECORD_FAIL_" + DateTime.Now.Ticks);
 							}
 						}
@@ -133,7 +133,7 @@ namespace Tools.CrashReporter.CrashReportProcess
 			}
 			catch (Exception Ex)
 			{
-				CrashReporterProcessServicer.WriteException("CheckForNewReportsOuter: " + Ex.ToString());
+				CrashReporterProcessServicer.WriteException("CheckForNewReportsOuter: " + Ex, Ex);
 			}
 
 			return GetTotalWaitingCount();
@@ -167,7 +167,7 @@ namespace Tools.CrashReporter.CrashReportProcess
 					try
 					{
 						string ReportName = Path.GetFileName(NewCrashContexts.Peek().CrashDirectory);
-						ReportIndex.TryRemoveReport(ReportName);
+						CrashReporterProcessServicer.ReportIndex.TryRemoveReport(ReportName);
 					}
 					finally
 					{
@@ -495,9 +495,10 @@ namespace Tools.CrashReporter.CrashReportProcess
 			// This WILL happen when running multiple, non-mutually exclusive crash sources (e.g. Receiver + Data Router)
 			// This can be safely discontinued when all crashes come from the same SQS
 			// This DOES NOT scale to multiple CRP instances
+			// ReportIndex can be disabled by leaving the path empty in config
 			lock (ReportIndexLock)
 			{
-				if (!ReportIndex.TryAddNewReport(ReportName))
+				if (!CrashReporterProcessServicer.ReportIndex.TryAddNewReport(ReportName))
 				{
 					// Crash report not accepted by index
 					CrashReporterProcessServicer.WriteEvent(string.Format(
diff --git a/Engine/Source/Programs/CrashReporter/CrashReportProcess/ReportWatcher.cs b/Engine/Source/Programs/CrashReporter/CrashReportProcess/ReportWatcher.cs
--- a/Engine/Source/Programs/CrashReporter/CrashReportProcess/ReportWatcher.cs
+++ b/Engine/Source/Programs/CrashReporter/CrashReportProcess/ReportWatcher.cs
@@ -45,7 +45,6 @@ namespace Tools.CrashReporter.CrashReportProcess
 			{
 				Queue.Dispose();
 			}
-			ReportIndex.WriteToFile();
 
 			CancelSource.Dispose();
 		}
@@ -76,7 +75,7 @@ namespace Tools.CrashReporter.CrashReportProcess
 			{
 				if( System.IO.Directory.Exists( Settings.InternalLandingZone ) )
 				{
-					ReportQueues.Add(new ReceiverReportQueue("Epic Crashes", Settings.InternalLandingZone, Config.Default.QueueLowerLimitForDiscard, Config.Default.QueueUpperLimitForDiscard));
+					ReportQueues.Add(new ReceiverReportQueue("Epic Crashes", Settings.InternalLandingZone, StatusReportingEventNames.ProcessingStartedReceiverEvent, Config.Default.QueueLowerLimitForDiscard, Config.Default.QueueUpperLimitForDiscard));
 					CrashReporterProcessServicer.WriteEvent(string.Format("\t{0} (internal, high priority (legacy))", Settings.InternalLandingZone));
 				}
 				else
@@ -85,12 +84,25 @@ namespace Tools.CrashReporter.CrashReportProcess
 				}
 			}
 
+			if (!string.IsNullOrEmpty(Settings.PS4LandingZone))
+			{
+				if (System.IO.Directory.Exists(Settings.PS4LandingZone))
+				{
+					ReportQueues.Add(new ReceiverReportQueue("PS4 Crashes", Settings.PS4LandingZone, StatusReportingEventNames.ProcessingStartedPS4Event, Config.Default.QueueLowerLimitForDiscard, Config.Default.QueueUpperLimitForDiscard));
+					CrashReporterProcessServicer.WriteEvent(string.Format("\t{0} (internal, PS4 non-retail hardware)", Settings.PS4LandingZone));
+				}
+				else
+				{
+					CrashReporterProcessServicer.WriteFailure(string.Format("\t{0} (internal, PS4 non-retail hardware) is not accessible", Settings.PS4LandingZone));
+				}
+			}
+
 #if !DEBUG
 			if (!string.IsNullOrEmpty(Settings.ExternalLandingZone))
 			{
 				if( System.IO.Directory.Exists( Settings.ExternalLandingZone ) )
 				{
-					ReportQueues.Add(new ReceiverReportQueue("External Crashes", Settings.ExternalLandingZone, Config.Default.QueueLowerLimitForDiscard, Config.Default.QueueUpperLimitForDiscard));
+					ReportQueues.Add(new ReceiverReportQueue("External Crashes", Settings.ExternalLandingZone, StatusReportingEventNames.ProcessingStartedReceiverEvent, Config.Default.QueueLowerLimitForDiscard, Config.Default.QueueUpperLimitForDiscard));
 					CrashReporterProcessServicer.WriteEvent( string.Format( "\t{0} (legacy)", Settings.ExternalLandingZone ) );
 				}
 				else
@@ -106,10 +118,6 @@ namespace Tools.CrashReporter.CrashReportProcess
 				CrashReporterProcessServicer.StatusReporter.InitQueue(Queue.QueueId, Queue.LandingZonePath);
 			}
 
-			ReportIndex.Filepath = Config.Default.ProcessedReportsIndexPath;
-			ReportIndex.Retention = TimeSpan.FromDays(Config.Default.ReportsIndexRetentionDays);
-			ReportIndex.ReadFromFile();
-
 			var Cancel = CancelSource.Token;
 			WatcherTask = new Task(async () =>
 			{
diff --git a/Engine/Source/Programs/CrashReporter/CrashReportProcess/StatusReporting.cs b/Engine/Source/Programs/CrashReporter/CrashReportProcess/StatusReporting.cs
--- a/Engine/Source/Programs/CrashReporter/CrashReportProcess/StatusReporting.cs
+++ b/Engine/Source/Programs/CrashReporter/CrashReportProcess/StatusReporting.cs
@@ -6,6 +6,7 @@ using System.Reflection;
 using System.Text;
 using System.Threading;
 using System.Threading.Tasks;
+using Tools.CrashReporter.CrashReportCommon;
 
 namespace Tools.CrashReporter.CrashReportProcess
 {
@@ -27,17 +28,20 @@ namespace Tools.CrashReporter.CrashReportProcess
 			// Collapse list of folders and labels into a set of drive letters and one or more labels
 			foreach (var FolderAndLabel in InFolderMonitors)
 			{
-				string Drive;
-				if (CrashReportCommon.StorageSpaceHelper.TryGetDriveLetter(FolderAndLabel.Key, out Drive))
+				if (!string.IsNullOrWhiteSpace(FolderAndLabel.Key))
 				{
-					if (FolderMonitors.ContainsKey(Drive))
-					{
-						// Add label to existing value
-						FolderMonitors[Drive] = FolderMonitors[Drive] + ", " + FolderAndLabel.Value;
-					}
-					else
+					string Drive;
+					if (CrashReportCommon.StorageSpaceHelper.TryGetDriveLetter(FolderAndLabel.Key, out Drive))
 					{
-						FolderMonitors.Add(Drive, FolderAndLabel.Value);
+						if (FolderMonitors.ContainsKey(Drive))
+						{
+							// Add label to existing value
+							FolderMonitors[Drive] = FolderMonitors[Drive] + ", " + FolderAndLabel.Value;
+						}
+						else
+						{
+							FolderMonitors.Add(Drive, FolderAndLabel.Value);
+						}
 					}
 				}
 			}
@@ -85,25 +89,31 @@ namespace Tools.CrashReporter.CrashReportProcess
 
 		public void AlertOnLowDisk(string Filepath, float AlertThresholdPercent)
 		{
-			string Drive;
-			if (!CrashReportCommon.StorageSpaceHelper.TryGetDriveLetter(Filepath, out Drive))
+			try
 			{
-				CrashReporterProcessServicer.WriteException("Failed to get drive letter for path " + Filepath);
-				return;
-			}
+				string Drive;
+				if (!StorageSpaceHelper.TryGetDriveLetter(Filepath, out Drive))
+				{
+					throw new CrashReporterException("Failed to get drive letter for path " + Filepath);
+				}
 
-			Int64 FreeSpace;
-			float FreePercent;
-			if (CrashReportCommon.StorageSpaceHelper.TryGetSpaceAvailable(Drive, out FreeSpace, out FreePercent))
-			{
-				if (FreePercent < AlertThresholdPercent)
+				Int64 FreeSpace;
+				float FreePercent;
+				if (StorageSpaceHelper.TryGetSpaceAvailable(Drive, out FreeSpace, out FreePercent))
+				{
+					if (FreePercent < AlertThresholdPercent)
+					{
+						Alert("AlertOnLowDisk" + Drive, "Low disk space warning on " + Drive + " =>> " + GetDiskSpaceString(FreeSpace, FreePercent), 3 * Config.Default.SlackAlertRepeatMinimumMinutes);
+					}
+				}
+				else
 				{
-					Alert("AlertOnLowDisk" + Drive, "Low disk space warning on " + Drive + " =>> " + GetDiskSpaceString(FreeSpace, FreePercent), Config.Default.SlackAlertRepeatMinimumMinutes);
+					CrashReporterProcessServicer.WriteEvent("Failed to read disk space for " + Drive);
 				}
 			}
-			else
+			catch (Exception Ex)
 			{
-				CrashReporterProcessServicer.WriteException("Failed to read disk space for " + Drive);
+				CrashReporterProcessServicer.WriteException("AlertOnLowDisk failed: " + Ex, Ex);
 			}
         }
 
@@ -188,7 +198,10 @@ namespace Tools.CrashReporter.CrashReportProcess
 							StatusReportMessage.AppendLine(string.Format("Events in the last {0}...", PeriodTimeString));
 							foreach (var CountInPeriod in CountsInPeriod)
 							{
-								StatusReportMessage.AppendLine("> " + CountInPeriod.Key + " " + CountInPeriod.Value);
+								if (CountInPeriod.Value > 0)
+								{
+									StatusReportMessage.AppendLine("> " + CountInPeriod.Key + " " + CountInPeriod.Value);
+								}
 							}
 						}
 					}
diff --git a/Engine/Source/Programs/CrashReporter/CrashReportProcess/StatusReportingConstants.cs b/Engine/Source/Programs/CrashReporter/CrashReportProcess/StatusReportingConstants.cs
--- a/Engine/Source/Programs/CrashReporter/CrashReportProcess/StatusReportingConstants.cs
+++ b/Engine/Source/Programs/CrashReporter/CrashReportProcess/StatusReportingConstants.cs
@@ -15,6 +15,7 @@ namespace Tools.CrashReporter.CrashReportProcess
 		public const string QueuedEvent = "Queued for processing";
 		public const string ProcessingStartedReceiverEvent = "Started processing (from CRR)";
 		public const string ProcessingStartedDataRouterEvent = "Started processing (from Data Router)";
+		public const string ProcessingStartedPS4Event = "Started processing (from PS4Services)";
 		public const string ProcessingSucceededEvent = "Processing succeeded";
 		public const string ProcessingFailedEvent = "Processing failed";
 		public const string ExceptionEvent = "Handled exceptions";	
diff --git a/Engine/Source/Programs/CrashReporter/CrashReportReceiver/CrashReportReceiver.csproj b/Engine/Source/Programs/CrashReporter/CrashReportReceiver/CrashReportReceiver.csproj
--- a/Engine/Source/Programs/CrashReporter/CrashReportReceiver/CrashReportReceiver.csproj
+++ b/Engine/Source/Programs/CrashReporter/CrashReportReceiver/CrashReportReceiver.csproj
@@ -51,6 +51,18 @@
     <CodeAnalysisRuleSet>MinimumRecommendedRules.ruleset</CodeAnalysisRuleSet>
     <Prefer32Bit>true</Prefer32Bit>
   </PropertyGroup>
+  <PropertyGroup Condition="'$(Configuration)|$(Platform)' == 'Dev Disabled|AnyCPU'">
+    <OutputPath>..\..\..\..\Binaries\DotNET\</OutputPath>
+    <DefineConstants>TRACE;CODE_ANALYSIS;DISABLED_CRR</DefineConstants>
+    <DocumentationFile>..\..\..\..\Binaries\DotNET\CrashReportReceiver.xml</DocumentationFile>
+    <Optimize>true</Optimize>
+    <DebugType>pdbonly</DebugType>
+    <PlatformTarget>AnyCPU</PlatformTarget>
+    <RunCodeAnalysis>true</RunCodeAnalysis>
+    <ErrorReport>prompt</ErrorReport>
+    <CodeAnalysisRuleSet>MinimumRecommendedRules.ruleset</CodeAnalysisRuleSet>
+    <Prefer32Bit>true</Prefer32Bit>
+  </PropertyGroup>
   <ItemGroup>
     <Reference Include="System" />
     <Reference Include="System.Configuration.Install" />
diff --git a/Engine/Source/Programs/CrashReporter/CrashReportReceiver/WebHandler.cs b/Engine/Source/Programs/CrashReporter/CrashReportReceiver/WebHandler.cs
--- a/Engine/Source/Programs/CrashReporter/CrashReportReceiver/WebHandler.cs
+++ b/Engine/Source/Programs/CrashReporter/CrashReportReceiver/WebHandler.cs
@@ -147,8 +147,11 @@ namespace Tools.CrashReporter.CrashReportReceiver
 		/// <returns>Result object, indicating whether the report has already been uploaded.</returns>
 		private CrashReporterResult CheckReport(HttpListenerRequest Request)
 		{
-			var ReportResult = new CrashReporterResult();
-
+			CrashReporterResult ReportResult = new CrashReporterResult();
+#if DISABLED_CRR
+			ReportResult.bSuccess = false;
+			CrashReporterReceiverServicer.WriteEvent("CheckReport() Report rejected by disabled CRR");
+#else
 			var RequestClass = new CheckReportRequest();
 			RequestClass.ReportId = GetReportIdFromPostData(GetContentStreamString(Request));
 
@@ -158,7 +161,7 @@ namespace Tools.CrashReporter.CrashReportReceiver
 			{
 				CrashReporterReceiverServicer.WriteEvent( string.Format( "Report \"{0}\" has already been received", RequestClass.ReportId ) );
 			}
-
+#endif
 			return ReportResult;
 		}
 
@@ -170,7 +173,11 @@ namespace Tools.CrashReporter.CrashReportReceiver
 		private CrashReporterResult CheckReportDetail( HttpListenerRequest Request )
 		{
 			CrashReporterResult ReportResult = new CrashReporterResult();
-
+#if DISABLED_CRR	
+			ReportResult.bSuccess = false;
+			ReportResult.Message = "CRR disabled";
+			CrashReporterReceiverServicer.WriteEvent("CheckReportDetail() Report rejected by disabled CRR");
+#else
 			string WERReportMetadataString = GetContentStreamString( Request );
 			WERReportMetadata WERData = null;
 
@@ -208,7 +215,7 @@ namespace Tools.CrashReporter.CrashReportReceiver
 					ReportResult.Message = "Rejecting Debug or DebugGame crash";
 				}
 			}
-
+#endif
 			return ReportResult;
 		}
 
