Copying //UE4/Dev-Core to //UE4/Dev-Main (Source: //UE4/Dev-Core @ 3151653)

diff --git a/Engine/Source/Programs/CrashReporter/CrashReportWebSite/Properties/Settings.settings b/Engine/Source/Programs/CrashReporter/CrashReportWebSite/Properties/Settings.settings
--- a/Engine/Source/Programs/CrashReporter/CrashReportWebSite/Properties/Settings.settings
+++ b/Engine/Source/Programs/CrashReporter/CrashReportWebSite/Properties/Settings.settings
@@ -12,7 +12,7 @@
       <Value Profile="(Default)">/CrashReporterCSV/</Value>
     </Setting>
     <Setting Name="JiraDeploymentAddress" Type="System.String" Scope="Application">
-      <Value Profile="(Default)">https://jira.ol.epicgames.net/rest/api/2</Value>
+      <Value Profile="(Default)">https://jira.it.epicgames.net/rest/api/2</Value>
     </Setting>
     <Setting Name="SitePath" Type="System.String" Scope="Application">
       <Value Profile="(Default)">\\devweb-02\Sites</Value>
@@ -29,155 +29,5 @@
     <Setting Name="SlackEmoji" Type="System.String" Scope="Application">
       <Value Profile="(Default)">:rngesus:</Value>
     </Setting>
-    <Setting Name="TestXML" Type="System.String" Scope="Application">
-      <Value Profile="(Default)">&lt;CrashDescription&gt;
-	&lt;CrashType /&gt;
-	&lt;BranchName&gt;//UE4/Release-4.12&lt;/BranchName&gt;
-	&lt;GameName&gt;UE4-Lenar_3&lt;/GameName&gt;
-	&lt;Platform&gt;Win64 [Windows 7 Service Pack 1 64b]&lt;/Platform&gt;
-	&lt;EngineMode&gt;Editor&lt;/EngineMode&gt;
-	&lt;BuildVersion&gt;4.12.5&lt;/BuildVersion&gt;
-	&lt;BuiltFromCL&gt;3039270&lt;/BuiltFromCL&gt;
-	&lt;CommandLine /&gt;
-	&lt;BaseDir&gt;C:/Program Files (x86)/Epic Games/4.12/Engine/Binaries/Win64/&lt;/BaseDir&gt;
-	&lt;Language&gt;zh_CN&lt;/Language&gt;
-	&lt;SystemLanguage /&gt;
-	&lt;UserName&gt;ohmimi&lt;/UserName&gt;
-	&lt;MachineGuid&gt;A24BAFFA41A33062DF1BCB9A31DDE162&lt;/MachineGuid&gt;
-	&lt;EpicAccountId&gt;cc732bdf5c4443f4a966a15bd4869c1e&lt;/EpicAccountId&gt;
-	&lt;CallStack&gt;UE4Editor_Core!FDebug::OptionallyLogFormattedEnsureMessageReturningFalse() [d:\build\++ue4+release-4.12+compile\sync\engine\source\runtime\core\private\misc\outputdevice.cpp:451]&lt;/CallStack&gt;
-	&lt;CallStack&gt;UE4Editor_Engine!UBlueprint::SetObjectBeingDebugged() [d:\build\++ue4+release-4.12+compile\sync\engine\source\runtime\engine\private\blueprint.cpp:685]&lt;/CallStack&gt;
-	&lt;CallStack&gt;UE4Editor_Engine!UWorld::TransferBlueprintDebugReferences() [d:\build\++ue4+release-4.12+compile\sync\engine\source\runtime\engine\private\world.cpp:1618]&lt;/CallStack&gt;
-	&lt;CallStack&gt;UE4Editor_UnrealEd!UEditorEngine::TeardownPlaySession() [d:\build\++ue4+release-4.12+compile\sync\engine\source\editor\unrealed\private\playlevel.cpp:642]&lt;/CallStack&gt;
-	&lt;CallStack&gt;UE4Editor_UnrealEd!UEditorEngine::EndPlayMap() [d:\build\++ue4+release-4.12+compile\sync\engine\source\editor\unrealed\private\playlevel.cpp:255]&lt;/CallStack&gt;
-	&lt;CallStack&gt;UE4Editor_LevelEditor!FLevelViewportLayout::~FLevelViewportLayout() [d:\build\++ue4+release-4.12+compile\sync\engine\source\editor\leveleditor\private\levelviewportlayout.cpp:145]&lt;/CallStack&gt;
-	&lt;CallStack&gt;UE4Editor_LevelEditor!FLevelViewportLayout2x2::`scalar deleting destructor'()&lt;/CallStack&gt;
-	&lt;CallStack&gt;UE4Editor_LevelEditor!FLevelViewportTabContent::~FLevelViewportTabContent()&lt;/CallStack&gt;
-	&lt;CallStack&gt;UE4Editor_LevelEditor!SharedPointerInternals::TReferenceControllerWithDeleter&amp;lt;FLevelViewportTabContent,SharedPointerInternals::DefaultDeleter&amp;lt;FLevelViewportTabContent&amp;gt; &amp;gt;::DestroyObject() [d:\build\++ue4+release-4.12+compile\sync\engine\source\runtime\core\public\templates\sharedpointerinternals.h:104]&lt;/CallStack&gt;
-	&lt;CallStack&gt;UE4Editor_LevelEditor!SViewportsOverlay::~SViewportsOverlay()&lt;/CallStack&gt;
-	&lt;CallStack&gt;UE4Editor_LevelEditor!SViewportsOverlay::`scalar deleting destructor'()&lt;/CallStack&gt;
-	&lt;CallStack&gt;UE4Editor_SlateCore!FSimpleSlot::~FSimpleSlot()&lt;/CallStack&gt;
-	&lt;CallStack&gt;UE4Editor_SlateCore!SCompoundWidget::~SCompoundWidget()&lt;/CallStack&gt;
-	&lt;CallStack&gt;UE4Editor_Slate!SBorder::`vector deleting destructor'()&lt;/CallStack&gt;
-	&lt;CallStack&gt;UE4Editor_SlateCore!TIndirectArray&amp;lt;SOverlay::FOverlaySlot,FDefaultAllocator&amp;gt;::DestructAndFreeItems() [d:\build\++ue4+release-4.12+compile\sync\engine\source\runtime\core\public\containers\array.h:3410]&lt;/CallStack&gt;
-	&lt;CallStack&gt;UE4Editor_SlateCore!SOverlay::~SOverlay()&lt;/CallStack&gt;
-	&lt;CallStack&gt;UE4Editor_Slate!SOverlay::`scalar deleting destructor'()&lt;/CallStack&gt;
-	&lt;CallStack&gt;UE4Editor_SlateCore!TIndirectArray&amp;lt;SBoxPanel::FSlot,FDefaultAllocator&amp;gt;::DestructAndFreeItems() [d:\build\++ue4+release-4.12+compile\sync\engine\source\runtime\core\public\containers\array.h:3410]&lt;/CallStack&gt;
-	&lt;CallStack&gt;UE4Editor_SlateCore!SBoxPanel::~SBoxPanel()&lt;/CallStack&gt;
-	&lt;CallStack&gt;UE4Editor_Slate!SVerticalBox::`scalar deleting destructor'()&lt;/CallStack&gt;
-	&lt;CallStack&gt;UE4Editor_SlateCore!FSimpleSlot::~FSimpleSlot()&lt;/CallStack&gt;
-	&lt;CallStack&gt;UE4Editor_SlateCore!SCompoundWidget::~SCompoundWidget()&lt;/CallStack&gt;
-	&lt;CallStack&gt;UE4Editor_Slate!SDockingTabStack::`vector deleting destructor'()&lt;/CallStack&gt;
-	&lt;CallStack&gt;UE4Editor_SlateCore!FSlotBase::~FSlotBase() [d:\build\++ue4+release-4.12+compile\sync\engine\source\runtime\slatecore\private\slotbase.cpp:34]&lt;/CallStack&gt;
-	&lt;CallStack&gt;UE4Editor_Slate!TIndirectArray&amp;lt;SSplitter::FSlot,FDefaultAllocator&amp;gt;::DestructAndFreeItems() [d:\build\++ue4+release-4.12+compile\sync\engine\source\runtime\core\public\containers\array.h:3410]&lt;/CallStack&gt;
-	&lt;CallStack&gt;UE4Editor_Slate!SSplitter::~SSplitter()&lt;/CallStack&gt;
-	&lt;CallStack&gt;UE4Editor_Slate!SSplitter::`vector deleting destructor'()&lt;/CallStack&gt;
-	&lt;CallStack&gt;UE4Editor_SlateCore!FSimpleSlot::~FSimpleSlot()&lt;/CallStack&gt;
-	&lt;CallStack&gt;UE4Editor_SlateCore!SCompoundWidget::~SCompoundWidget()&lt;/CallStack&gt;
-	&lt;CallStack&gt;UE4Editor_Slate!SDockingSplitter::`vector deleting destructor'()&lt;/CallStack&gt;
-	&lt;CallStack&gt;UE4Editor_SlateCore!FSlotBase::~FSlotBase() [d:\build\++ue4+release-4.12+compile\sync\engine\source\runtime\slatecore\private\slotbase.cpp:34]&lt;/CallStack&gt;
-	&lt;CallStack&gt;UE4Editor_Slate!TIndirectArray&amp;lt;SSplitter::FSlot,FDefaultAllocator&amp;gt;::DestructAndFreeItems() [d:\build\++ue4+release-4.12+compile\sync\engine\source\runtime\core\public\containers\array.h:3410]&lt;/CallStack&gt;
-	&lt;CallStack&gt;UE4Editor_Slate!SSplitter::~SSplitter()&lt;/CallStack&gt;
-	&lt;CallStack&gt;UE4Editor_Slate!SSplitter::`vector deleting destructor'()&lt;/CallStack&gt;
-	&lt;CallStack&gt;UE4Editor_SlateCore!FSimpleSlot::~FSimpleSlot()&lt;/CallStack&gt;
-	&lt;CallStack&gt;UE4Editor_SlateCore!SCompoundWidget::~SCompoundWidget()&lt;/CallStack&gt;
-	&lt;CallStack&gt;UE4Editor_Slate!SDockingSplitter::`vector deleting destructor'()&lt;/CallStack&gt;
-	&lt;CallStack&gt;UE4Editor_SlateCore!FSlotBase::~FSlotBase() [d:\build\++ue4+release-4.12+compile\sync\engine\source\runtime\slatecore\private\slotbase.cpp:34]&lt;/CallStack&gt;
-	&lt;CallStack&gt;UE4Editor_Slate!TIndirectArray&amp;lt;SSplitter::FSlot,FDefaultAllocator&amp;gt;::DestructAndFreeItems() [d:\build\++ue4+release-4.12+compile\sync\engine\source\runtime\core\public\containers\array.h:3410]&lt;/CallStack&gt;
-	&lt;CallStack&gt;UE4Editor_Slate!SSplitter::~SSplitter()&lt;/CallStack&gt;
-	&lt;CallStack&gt;UE4Editor_Slate!SSplitter::`vector deleting destructor'()&lt;/CallStack&gt;
-	&lt;CallStack&gt;UE4Editor_SlateCore!TIndirectArray&amp;lt;SOverlay::FOverlaySlot,FDefaultAllocator&amp;gt;::DestructAndFreeItems() [d:\build\++ue4+release-4.12+compile\sync\engine\source\runtime\core\public\containers\array.h:3410]&lt;/CallStack&gt;
-	&lt;CallStack&gt;UE4Editor_SlateCore!SOverlay::~SOverlay()&lt;/CallStack&gt;
-	&lt;CallStack&gt;UE4Editor_Slate!SOverlay::`scalar deleting destructor'()&lt;/CallStack&gt;
-	&lt;CallStack&gt;UE4Editor_SlateCore!FSimpleSlot::~FSimpleSlot()&lt;/CallStack&gt;
-	&lt;CallStack&gt;UE4Editor_SlateCore!SCompoundWidget::~SCompoundWidget()&lt;/CallStack&gt;
-	&lt;CallStack&gt;UE4Editor_Slate!SDockingArea::`vector deleting destructor'()&lt;/CallStack&gt;
-	&lt;CallStack&gt;UE4Editor_SlateCore!TIndirectArray&amp;lt;SBoxPanel::FSlot,FDefaultAllocator&amp;gt;::DestructAndFreeItems() [d:\build\++ue4+release-4.12+compile\sync\engine\source\runtime\core\public\containers\array.h:3410]&lt;/CallStack&gt;
-	&lt;CallStack&gt;UE4Editor_SlateCore!SBoxPanel::~SBoxPanel()&lt;/CallStack&gt;
-	&lt;CallStack&gt;UE4Editor_LevelEditor!SVerticalBox::`scalar deleting destructor'()&lt;/CallStack&gt;
-	&lt;CallStack&gt;UE4Editor_SlateCore!FSimpleSlot::~FSimpleSlot()&lt;/CallStack&gt;
-	&lt;CallStack&gt;UE4Editor_SlateCore!SCompoundWidget::~SCompoundWidget()&lt;/CallStack&gt;
-	&lt;CallStack&gt;UE4Editor_LevelEditor!SLevelEditor::`scalar deleting destructor'()&lt;/CallStack&gt;
-	&lt;CallStack&gt;UE4Editor_SlateCore!FSimpleSlot::~FSimpleSlot()&lt;/CallStack&gt;
-	&lt;CallStack&gt;UE4Editor_SlateCore!SCompoundWidget::~SCompoundWidget()&lt;/CallStack&gt;
-	&lt;CallStack&gt;UE4Editor_Slate!SBorder::`vector deleting destructor'()&lt;/CallStack&gt;
-	&lt;CallStack&gt;UE4Editor_SlateCore!TIndirectArray&amp;lt;SOverlay::FOverlaySlot,FDefaultAllocator&amp;gt;::DestructAndFreeItems() [d:\build\++ue4+release-4.12+compile\sync\engine\source\runtime\core\public\containers\array.h:3410]&lt;/CallStack&gt;
-	&lt;CallStack&gt;UE4Editor_SlateCore!SOverlay::~SOverlay()&lt;/CallStack&gt;
-	&lt;CallStack&gt;UE4Editor_Slate!SOverlay::`scalar deleting destructor'()&lt;/CallStack&gt;
-	&lt;CallStack&gt;UE4Editor_SlateCore!TIndirectArray&amp;lt;SBoxPanel::FSlot,FDefaultAllocator&amp;gt;::DestructAndFreeItems() [d:\build\++ue4+release-4.12+compile\sync\engine\source\runtime\core\public\containers\array.h:3410]&lt;/CallStack&gt;
-	&lt;CallStack&gt;UE4Editor_SlateCore!SBoxPanel::~SBoxPanel()&lt;/CallStack&gt;
-	&lt;CallStack&gt;UE4Editor_Slate!SVerticalBox::`scalar deleting destructor'()&lt;/CallStack&gt;
-	&lt;CallStack&gt;UE4Editor_SlateCore!FSimpleSlot::~FSimpleSlot()&lt;/CallStack&gt;
-	&lt;CallStack&gt;UE4Editor_SlateCore!SCompoundWidget::~SCompoundWidget()&lt;/CallStack&gt;
-	&lt;CallStack&gt;UE4Editor_Slate!SDockingTabStack::`vector deleting destructor'()&lt;/CallStack&gt;
-	&lt;CallStack&gt;UE4Editor_SlateCore!FSlotBase::~FSlotBase() [d:\build\++ue4+release-4.12+compile\sync\engine\source\runtime\slatecore\private\slotbase.cpp:34]&lt;/CallStack&gt;
-	&lt;CallStack&gt;UE4Editor_Slate!TIndirectArray&amp;lt;SSplitter::FSlot,FDefaultAllocator&amp;gt;::DestructAndFreeItems() [d:\build\++ue4+release-4.12+compile\sync\engine\source\runtime\core\public\containers\array.h:3410]&lt;/CallStack&gt;
-	&lt;CallStack&gt;UE4Editor_Slate!SSplitter::~SSplitter()&lt;/CallStack&gt;
-	&lt;CallStack&gt;UE4Editor_Slate!SSplitter::`vector deleting destructor'()&lt;/CallStack&gt;
-	&lt;CallStack&gt;UE4Editor_SlateCore!TIndirectArray&amp;lt;SOverlay::FOverlaySlot,FDefaultAllocator&amp;gt;::DestructAndFreeItems() [d:\build\++ue4+release-4.12+compile\sync\engine\source\runtime\core\public\containers\array.h:3410]&lt;/CallStack&gt;
-	&lt;CallStack&gt;UE4Editor_SlateCore!SOverlay::~SOverlay()&lt;/CallStack&gt;
-	&lt;CallStack&gt;UE4Editor_Slate!SOverlay::`scalar deleting destructor'()&lt;/CallStack&gt;
-	&lt;CallStack&gt;UE4Editor_SlateCore!TArray&amp;lt;FArrangedWidget,TInlineAllocator&amp;lt;16,FDefaultAllocator&amp;gt; &amp;gt;::~TArray&amp;lt;FArrangedWidget,TInlineAllocator&amp;lt;16,FDefaultAllocator&amp;gt; &amp;gt;() [d:\build\++ue4+release-4.12+compile\sync\engine\source\runtime\core\public\containers\array.h:723]&lt;/CallStack&gt;
-	&lt;CallStack&gt;UE4Editor_Slate!FSlateApplication::RoutePointerUpEvent() [d:\build\++ue4+release-4.12+compile\sync\engine\source\runtime\slate\private\framework\application\slateapplication.cpp:4667]&lt;/CallStack&gt;
-	&lt;CallStack&gt;UE4Editor_Slate!FSlateApplication::ProcessMouseButtonUpEvent() [d:\build\++ue4+release-4.12+compile\sync\engine\source\runtime\slate\private\framework\application\slateapplication.cpp:5041]&lt;/CallStack&gt;
-	&lt;CallStack&gt;UE4Editor_Slate!FSlateApplication::OnMouseUp() [d:\build\++ue4+release-4.12+compile\sync\engine\source\runtime\slate\private\framework\application\slateapplication.cpp:5021]&lt;/CallStack&gt;
-	&lt;CallStack&gt;UE4Editor_Core!FWindowsApplication::ProcessDeferredMessage() [d:\build\++ue4+release-4.12+compile\sync\engine\source\runtime\core\private\windows\windowsapplication.cpp:1510]&lt;/CallStack&gt;
-	&lt;CallStack&gt;UE4Editor_Core!FWindowsApplication::DeferMessage() [d:\build\++ue4+release-4.12+compile\sync\engine\source\runtime\core\private\windows\windowsapplication.cpp:1850]&lt;/CallStack&gt;
-	&lt;CallStack&gt;UE4Editor_Core!FWindowsApplication::ProcessMessage() [d:\build\++ue4+release-4.12+compile\sync\engine\source\runtime\core\private\windows\windowsapplication.cpp:745]&lt;/CallStack&gt;
-	&lt;CallStack&gt;UE4Editor_Core!FWindowsApplication::AppWndProc() [d:\build\++ue4+release-4.12+compile\sync\engine\source\runtime\core\private\windows\windowsapplication.cpp:667]&lt;/CallStack&gt;
-	&lt;CallStack&gt;user32!UserCallWinProcCheckWow()&lt;/CallStack&gt;
-	&lt;CallStack&gt;user32!DispatchMessageWorker()&lt;/CallStack&gt;
-	&lt;CallStack&gt;UE4Editor_Core!FWindowsPlatformMisc::PumpMessages() [d:\build\++ue4+release-4.12+compile\sync\engine\source\runtime\core\private\windows\windowsplatformmisc.cpp:903]&lt;/CallStack&gt;
-	&lt;CallStack&gt;UE4Editor!FEngineLoop::Tick() [d:\build\++ue4+release-4.12+compile\sync\engine\source\runtime\launch\private\launchengineloop.cpp:2729]&lt;/CallStack&gt;
-	&lt;CallStack&gt;UE4Editor!GuardedMain() [d:\build\++ue4+release-4.12+compile\sync\engine\source\runtime\launch\private\launch.cpp:148]&lt;/CallStack&gt;
-	&lt;CallStack&gt;UE4Editor!GuardedMainWrapper() [d:\build\++ue4+release-4.12+compile\sync\engine\source\runtime\launch\private\windows\launchwindows.cpp:126]&lt;/CallStack&gt;
-	&lt;CallStack&gt;UE4Editor!WinMain() [d:\build\++ue4+release-4.12+compile\sync\engine\source\runtime\launch\private\windows\launchwindows.cpp:200]&lt;/CallStack&gt;
-	&lt;CallStack&gt;UE4Editor!__scrt_common_main_seh() [f:\dd\vctools\crt\vcstartup\src\startup\exe_common.inl:264]&lt;/CallStack&gt;
-	&lt;CallStack&gt;kernel32!BaseThreadInitThunk()&lt;/CallStack&gt;
-	&lt;CallStack&gt;ntdll!RtlUserThreadStart()&lt;/CallStack&gt;
-	&lt;SourceContext&gt;  670       	{&lt;/SourceContext&gt;
-	&lt;SourceContext&gt;  671       		if (OldObject == NewObject)&lt;/SourceContext&gt;
-	&lt;SourceContext&gt;  672       		{&lt;/SourceContext&gt;
-	&lt;SourceContext&gt;  673       			// Nothing changed&lt;/SourceContext&gt;
-	&lt;SourceContext&gt;  674       			return;&lt;/SourceContext&gt;
-	&lt;SourceContext&gt;  675       		}&lt;/SourceContext&gt;
-	&lt;SourceContext&gt;  676       &lt;/SourceContext&gt;
-	&lt;SourceContext&gt;  677       		DebuggingWorldRegistrationHelper(OldObject, NULL);&lt;/SourceContext&gt;
-	&lt;SourceContext&gt;  678       	}&lt;/SourceContext&gt;
-	&lt;SourceContext&gt;  679       &lt;/SourceContext&gt;
-	&lt;SourceContext&gt;  680       	// Note that we allow macro Blueprints to bypass this check&lt;/SourceContext&gt;
-	&lt;SourceContext&gt;  681       	if ((NewObject != NULL) &amp;amp;&amp;amp; !GCompilingBlueprint &amp;amp;&amp;amp; BlueprintType != BPTYPE_MacroLibrary)&lt;/SourceContext&gt;
-	&lt;SourceContext&gt;  682       	{&lt;/SourceContext&gt;
-	&lt;SourceContext&gt;  683       		// You can only debug instances of this!&lt;/SourceContext&gt;
-	&lt;SourceContext&gt;  684 ***** 		if (!ensureMsgf(&lt;/SourceContext&gt;
-	&lt;SourceContext&gt;  685       				NewObject-&amp;gt;IsA(this-&amp;gt;GeneratedClass), &lt;/SourceContext&gt;
-	&lt;SourceContext&gt;  686       				TEXT("Type mismatch: Expected %s, Found %s"), &lt;/SourceContext&gt;
-	&lt;SourceContext&gt;  687       				this-&amp;gt;GeneratedClass ? *(this-&amp;gt;GeneratedClass-&amp;gt;GetName()) : TEXT("NULL"), &lt;/SourceContext&gt;
-	&lt;SourceContext&gt;  688       				NewObject-&amp;gt;GetClass() ? *(this-&amp;gt;GetClass()-&amp;gt;GetName()) : TEXT("NULL")))&lt;/SourceContext&gt;
-	&lt;SourceContext&gt;  689       		{&lt;/SourceContext&gt;
-	&lt;SourceContext&gt;  690       			NewObject = NULL;&lt;/SourceContext&gt;
-	&lt;SourceContext&gt;  691       		}&lt;/SourceContext&gt;
-	&lt;SourceContext&gt;  692       	}&lt;/SourceContext&gt;
-	&lt;SourceContext&gt;  693       &lt;/SourceContext&gt;
-	&lt;SourceContext&gt;  694       	// Update the current object being debugged&lt;/SourceContext&gt;
-	&lt;SourceContext&gt;  695       	CurrentObjectBeingDebugged = NewObject;&lt;/SourceContext&gt;
-	&lt;SourceContext&gt;  696       &lt;/SourceContext&gt;
-	&lt;SourceContext&gt;  697       	// Register the new object&lt;/SourceContext&gt;
-	&lt;SourceContext&gt;  698       	if (NewObject != NULL)&lt;/SourceContext&gt;
-	&lt;SourceContext&gt;  699       	{&lt;/SourceContext&gt;
-	&lt;UserDescription&gt;Sent in the unattended mode&lt;/UserDescription&gt;
-	&lt;UserActivityHint /&gt;
-	&lt;ErrorMessage&gt;Ensure condition failed: NewObject-&amp;gt;IsA(this-&amp;gt;GeneratedClass) [File:D:\Build\++UE4+Release-4.12+Compile\Sync\Engine\Source\Runtime\Engine\Private\Blueprint.cpp] [Line: 689] &lt;/ErrorMessage&gt;
-	&lt;ErrorMessage&gt;Type mismatch: Expected Untitled_C, Found LevelScriptBlueprint&lt;/ErrorMessage&gt;
-	&lt;CrashGUID&gt;UE4CC-Windows-7188BC094DE45EE1A7D39EBF6FB434F5_0000&lt;/CrashGUID&gt;
-	&lt;TimeofCrash&gt;2016-07-13T14:58:43.402875Z&lt;/TimeofCrash&gt;
-	&lt;bHasMiniDump&gt;true&lt;/bHasMiniDump&gt;
-	&lt;bHasLog&gt;true&lt;/bHasLog&gt;
-	&lt;bHasDiags&gt;false&lt;/bHasDiags&gt;
-	&lt;bHasWERData&gt;false&lt;/bHasWERData&gt;
-	&lt;bHasVideo&gt;false&lt;/bHasVideo&gt;
-	&lt;bAllowToBeContacted&gt;true&lt;/bAllowToBeContacted&gt;
-&lt;/CrashDescription&gt;</Value>
-    </Setting>
   </Settings>
 </SettingsFile>
\ No newline at end of file
