Copying //UE4/Dev-Core to //UE4/Dev-Main (Source: //UE4/Dev-Core @ 3151653)

diff --git a/Engine/Source/Programs/CrashReporter/CrashReportWebSite/Scripts/jquery-3.1.0.intellisense.js b/Engine/Source/Programs/CrashReporter/CrashReportWebSite/Scripts/jquery-3.1.0.intellisense.js
new file mode 100644
--- /dev/null
+++ b/Engine/Source/Programs/CrashReporter/CrashReportWebSite/Scripts/jquery-3.1.0.intellisense.js
@@ -0,0 +1,2670 @@
+intellisense.annotate(jQuery, {
+  'ajax': function() {
+    /// <signature>
+    ///   <summary>Perform an asynchronous HTTP (Ajax) request.</summary>
+    ///   <param name="url" type="String">A string containing the URL to which the request is sent.</param>
+    ///   <param name="settings" type="PlainObject">A set of key/value pairs that configure the Ajax request. All settings are optional. A default can be set for any option with $.ajaxSetup(). See jQuery.ajax( settings ) below for a complete list of all settings.</param>
+    ///   <returns type="jqXHR" />
+    /// </signature>
+    /// <signature>
+    ///   <summary>Perform an asynchronous HTTP (Ajax) request.</summary>
+    ///   <param name="settings" type="PlainObject">A set of key/value pairs that configure the Ajax request. All settings are optional. A default can be set for any option with $.ajaxSetup().</param>
+    ///   <returns type="jqXHR" />
+    /// </signature>
+  },
+  'ajaxPrefilter': function() {
+    /// <signature>
+    ///   <summary>Handle custom Ajax options or modify existing options before each request is sent and before they are processed by $.ajax().</summary>
+    ///   <param name="dataTypes" type="String">An optional string containing one or more space-separated dataTypes</param>
+    ///   <param name="handler(options, originalOptions, jqXHR)" type="Function">A handler to set default values for future Ajax requests.</param>
+    /// </signature>
+  },
+  'ajaxSetup': function() {
+    /// <signature>
+    ///   <summary>Set default values for future Ajax requests. Its use is not recommended.</summary>
+    ///   <param name="options" type="PlainObject">A set of key/value pairs that configure the default Ajax request. All options are optional.</param>
+    /// </signature>
+  },
+  'ajaxTransport': function() {
+    /// <signature>
+    ///   <summary>Creates an object that handles the actual transmission of Ajax data.</summary>
+    ///   <param name="dataType" type="String">A string identifying the data type to use</param>
+    ///   <param name="handler(options, originalOptions, jqXHR)" type="Function">A handler to return the new transport object to use with the data type provided in the first argument.</param>
+    /// </signature>
+  },
+  'boxModel': function() {
+    /// <summary>Deprecated in jQuery 1.3 (see jQuery.support). States if the current page, in the user's browser, is being rendered using the W3C CSS Box Model.</summary>
+    /// <returns type="Boolean" />
+  },
+  'browser': function() {
+    /// <summary>Contains flags for the useragent, read from navigator.userAgent. This property was removed in jQuery 1.9 and is available only through the jQuery.migrate plugin. Please try to use feature detection instead.</summary>
+    /// <returns type="PlainObject" />
+  },
+  'browser.version': function() {
+    /// <summary>The version number of the rendering engine for the user's browser. This property was removed in jQuery 1.9 and is available only through the jQuery.migrate plugin.</summary>
+    /// <returns type="String" />
+  },
+  'Callbacks': function() {
+    /// <signature>
+    ///   <summary>A multi-purpose callbacks list object that provides a powerful way to manage callback lists.</summary>
+    ///   <param name="flags" type="String">An optional list of space-separated flags that change how the callback list behaves.</param>
+    ///   <returns type="Callbacks" />
+    /// </signature>
+  },
+  'contains': function() {
+    /// <signature>
+    ///   <summary>Check to see if a DOM element is a descendant of another DOM element.</summary>
+    ///   <param name="container" type="Element">The DOM element that may contain the other element.</param>
+    ///   <param name="contained" type="Element">The DOM element that may be contained by (a descendant of) the other element.</param>
+    ///   <returns type="Boolean" />
+    /// </signature>
+  },
+  'cssHooks': function() {
+    /// <summary>Hook directly into jQuery to override how particular CSS properties are retrieved or set, normalize CSS property naming, or create custom properties.</summary>
+    /// <returns type="Object" />
+  },
+  'data': function() {
+    /// <signature>
+    ///   <summary>Returns value at named data store for the element, as set by jQuery.data(element, name, value), or the full data store for the element.</summary>
+    ///   <param name="element" type="Element">The DOM element to query for the data.</param>
+    ///   <param name="key" type="String">Name of the data stored.</param>
+    ///   <returns type="Object" />
+    /// </signature>
+    /// <signature>
+    ///   <summary>Returns value at named data store for the element, as set by jQuery.data(element, name, value), or the full data store for the element.</summary>
+    ///   <param name="element" type="Element">The DOM element to query for the data.</param>
+    ///   <returns type="Object" />
+    /// </signature>
+  },
+  'Deferred': function() {
+    /// <signature>
+    ///   <summary>A constructor function that returns a chainable utility object with methods to register multiple callbacks into callback queues, invoke callback queues, and relay the success or failure state of any synchronous or asynchronous function.</summary>
+    ///   <param name="beforeStart" type="Function">A function that is called just before the constructor returns.</param>
+    ///   <returns type="Deferred" />
+    /// </signature>
+  },
+  'dequeue': function() {
+    /// <signature>
+    ///   <summary>Execute the next function on the queue for the matched element.</summary>
+    ///   <param name="element" type="Element">A DOM element from which to remove and execute a queued function.</param>
+    ///   <param name="queueName" type="String">A string containing the name of the queue. Defaults to fx, the standard effects queue.</param>
+    /// </signature>
+  },
+  'each': function() {
+    /// <signature>
+    ///   <summary>A generic iterator function, which can be used to seamlessly iterate over both objects and arrays. Arrays and array-like objects with a length property (such as a function's arguments object) are iterated by numeric index, from 0 to length-1. Other objects are iterated via their named properties.</summary>
+    ///   <param name="collection" type="Object">The object or array to iterate over.</param>
+    ///   <param name="callback(indexInArray, valueOfElement)" type="Function">The function that will be executed on every object.</param>
+    ///   <returns type="Object" />
+    /// </signature>
+  },
+  'error': function() {
+    /// <signature>
+    ///   <summary>Takes a string and throws an exception containing it.</summary>
+    ///   <param name="message" type="String">The message to send out.</param>
+    /// </signature>
+  },
+  'extend': function() {
+    /// <signature>
+    ///   <summary>Merge the contents of two or more objects together into the first object.</summary>
+    ///   <param name="target" type="Object">An object that will receive the new properties if additional objects are passed in or that will extend the jQuery namespace if it is the sole argument.</param>
+    ///   <param name="object1" type="Object">An object containing additional properties to merge in.</param>
+    ///   <param name="objectN" type="Object">Additional objects containing properties to merge in.</param>
+    ///   <returns type="Object" />
+    /// </signature>
+    /// <signature>
+    ///   <summary>Merge the contents of two or more objects together into the first object.</summary>
+    ///   <param name="deep" type="Boolean">If true, the merge becomes recursive (aka. deep copy).</param>
+    ///   <param name="target" type="Object">The object to extend. It will receive the new properties.</param>
+    ///   <param name="object1" type="Object">An object containing additional properties to merge in.</param>
+    ///   <param name="objectN" type="Object">Additional objects containing properties to merge in.</param>
+    ///   <returns type="Object" />
+    /// </signature>
+  },
+  'fn.extend': function() {
+    /// <signature>
+    ///   <summary>Merge the contents of an object onto the jQuery prototype to provide new jQuery instance methods.</summary>
+    ///   <param name="object" type="Object">An object to merge onto the jQuery prototype.</param>
+    ///   <returns type="Object" />
+    /// </signature>
+  },
+  'get': function() {
+    /// <signature>
+    ///   <summary>Load data from the server using a HTTP GET request.</summary>
+    ///   <param name="url" type="String">A string containing the URL to which the request is sent.</param>
+    ///   <param name="data" type="">A plain object or string that is sent to the server with the request.</param>
+    ///   <param name="success(data, textStatus, jqXHR)" type="Function">A callback function that is executed if the request succeeds.</param>
+    ///   <param name="dataType" type="String">The type of data expected from the server. Default: Intelligent Guess (xml, json, script, or html).</param>
+    ///   <returns type="jqXHR" />
+    /// </signature>
+  },
+  'getJSON': function() {
+    /// <signature>
+    ///   <summary>Load JSON-encoded data from the server using a GET HTTP request.</summary>
+    ///   <param name="url" type="String">A string containing the URL to which the request is sent.</param>
+    ///   <param name="data" type="PlainObject">A plain object or string that is sent to the server with the request.</param>
+    ///   <param name="success(data, textStatus, jqXHR)" type="Function">A callback function that is executed if the request succeeds.</param>
+    ///   <returns type="jqXHR" />
+    /// </signature>
+  },
+  'getScript': function() {
+    /// <signature>
+    ///   <summary>Load a JavaScript file from the server using a GET HTTP request, then execute it.</summary>
+    ///   <param name="url" type="String">A string containing the URL to which the request is sent.</param>
+    ///   <param name="success(script, textStatus, jqXHR)" type="Function">A callback function that is executed if the request succeeds.</param>
+    ///   <returns type="jqXHR" />
+    /// </signature>
+  },
+  'globalEval': function() {
+    /// <signature>
+    ///   <summary>Execute some JavaScript code globally.</summary>
+    ///   <param name="code" type="String">The JavaScript code to execute.</param>
+    /// </signature>
+  },
+  'grep': function() {
+    /// <signature>
+    ///   <summary>Finds the elements of an array which satisfy a filter function. The original array is not affected.</summary>
+    ///   <param name="array" type="Array">The array to search through.</param>
+    ///   <param name="function(elementOfArray, indexInArray)" type="Function">The function to process each item against.  The first argument to the function is the item, and the second argument is the index.  The function should return a Boolean value.  this will be the global window object.</param>
+    ///   <param name="invert" type="Boolean">If "invert" is false, or not provided, then the function returns an array consisting of all elements for which "callback" returns true.  If "invert" is true, then the function returns an array consisting of all elements for which "callback" returns false.</param>
+    ///   <returns type="Array" />
+    /// </signature>
+  },
+  'hasData': function() {
+    /// <signature>
+    ///   <summary>Determine whether an element has any jQuery data associated with it.</summary>
+    ///   <param name="element" type="Element">A DOM element to be checked for data.</param>
+    ///   <returns type="Boolean" />
+    /// </signature>
+  },
+  'holdReady': function() {
+    /// <signature>
+    ///   <summary>Holds or releases the execution of jQuery's ready event.</summary>
+    ///   <param name="hold" type="Boolean">Indicates whether the ready hold is being requested or released</param>
+    /// </signature>
+  },
+  'inArray': function() {
+    /// <signature>
+    ///   <summary>Search for a specified value within an array and return its index (or -1 if not found).</summary>
+    ///   <param name="value" type="Anything">The value to search for.</param>
+    ///   <param name="array" type="Array">An array through which to search.</param>
+    ///   <param name="fromIndex" type="Number">The index of the array at which to begin the search. The default is 0, which will search the whole array.</param>
+    ///   <returns type="Number" />
+    /// </signature>
+  },
+  'isArray': function() {
+    /// <signature>
+    ///   <summary>Determine whether the argument is an array.</summary>
+    ///   <param name="obj" type="Object">Object to test whether or not it is an array.</param>
+    ///   <returns type="Boolean" />
+    /// </signature>
+  },
+  'isEmptyObject': function() {
+    /// <signature>
+    ///   <summary>Check to see if an object is empty (contains no enumerable properties).</summary>
+    ///   <param name="object" type="Object">The object that will be checked to see if it's empty.</param>
+    ///   <returns type="Boolean" />
+    /// </signature>
+  },
+  'isFunction': function() {
+    /// <signature>
+    ///   <summary>Determine if the argument passed is a Javascript function object.</summary>
+    ///   <param name="obj" type="PlainObject">Object to test whether or not it is a function.</param>
+    ///   <returns type="Boolean" />
+    /// </signature>
+  },
+  'isNumeric': function() {
+    /// <signature>
+    ///   <summary>Determines whether its argument is a number.</summary>
+    ///   <param name="value" type="PlainObject">The value to be tested.</param>
+    ///   <returns type="Boolean" />
+    /// </signature>
+  },
+  'isPlainObject': function() {
+    /// <signature>
+    ///   <summary>Check to see if an object is a plain object (created using "{}" or "new Object").</summary>
+    ///   <param name="object" type="PlainObject">The object that will be checked to see if it's a plain object.</param>
+    ///   <returns type="Boolean" />
+    /// </signature>
+  },
+  'isWindow': function() {
+    /// <signature>
+    ///   <summary>Determine whether the argument is a window.</summary>
+    ///   <param name="obj" type="PlainObject">Object to test whether or not it is a window.</param>
+    ///   <returns type="Boolean" />
+    /// </signature>
+  },
+  'isXMLDoc': function() {
+    /// <signature>
+    ///   <summary>Check to see if a DOM node is within an XML document (or is an XML document).</summary>
+    ///   <param name="node" type="Element">The DOM node that will be checked to see if it's in an XML document.</param>
+    ///   <returns type="Boolean" />
+    /// </signature>
+  },
+  'makeArray': function() {
+    /// <signature>
+    ///   <summary>Convert an array-like object into a true JavaScript array.</summary>
+    ///   <param name="obj" type="PlainObject">Any object to turn into a native Array.</param>
+    ///   <returns type="Array" />
+    /// </signature>
+  },
+  'map': function() {
+    /// <signature>
+    ///   <summary>Translate all items in an array or object to new array of items.</summary>
+    ///   <param name="array" type="Array">The Array to translate.</param>
+    ///   <param name="callback(elementOfArray, indexInArray)" type="Function">The function to process each item against.  The first argument to the function is the array item, the second argument is the index in array The function can return any value. Within the function, this refers to the global (window) object.</param>
+    ///   <returns type="Array" />
+    /// </signature>
+    /// <signature>
+    ///   <summary>Translate all items in an array or object to new array of items.</summary>
+    ///   <param name="arrayOrObject" type="">The Array or Object to translate.</param>
+    ///   <param name="callback( value, indexOrKey )" type="Function">The function to process each item against.  The first argument to the function is the value; the second argument is the index or key of the array or object property. The function can return any value to add to the array. A returned array will be flattened into the resulting array. Within the function, this refers to the global (window) object.</param>
+    ///   <returns type="Array" />
+    /// </signature>
+  },
+  'merge': function() {
+    /// <signature>
+    ///   <summary>Merge the contents of two arrays together into the first array.</summary>
+    ///   <param name="first" type="Array">The first array to merge, the elements of second added.</param>
+    ///   <param name="second" type="Array">The second array to merge into the first, unaltered.</param>
+    ///   <returns type="Array" />
+    /// </signature>
+  },
+  'noConflict': function() {
+    /// <signature>
+    ///   <summary>Relinquish jQuery's control of the $ variable.</summary>
+    ///   <param name="removeAll" type="Boolean">A Boolean indicating whether to remove all jQuery variables from the global scope (including jQuery itself).</param>
+    ///   <returns type="Object" />
+    /// </signature>
+  },
+  'noop': function() {
+    /// <summary>An empty function.</summary>
+  },
+  'now': function() {
+    /// <summary>Return a number representing the current time.</summary>
+    /// <returns type="Number" />
+  },
+  'param': function() {
+    /// <signature>
+    ///   <summary>Create a serialized representation of an array or object, suitable for use in a URL query string or Ajax request.</summary>
+    ///   <param name="obj" type="">An array or object to serialize.</param>
+    ///   <returns type="String" />
+    /// </signature>
+    /// <signature>
+    ///   <summary>Create a serialized representation of an array or object, suitable for use in a URL query string or Ajax request.</summary>
+    ///   <param name="obj" type="">An array or object to serialize.</param>
+    ///   <param name="traditional" type="Boolean">A Boolean indicating whether to perform a traditional "shallow" serialization.</param>
+    ///   <returns type="String" />
+    /// </signature>
+  },
+  'parseHTML': function() {
+    /// <signature>
+    ///   <summary>Parses a string into an array of DOM nodes.</summary>
+    ///   <param name="data" type="String">HTML string to be parsed</param>
+    ///   <param name="context" type="Element">Document element to serve as the context in which the HTML fragment will be created</param>
+    ///   <param name="keepScripts" type="Boolean">A Boolean indicating whether to include scripts passed in the HTML string</param>
+    ///   <returns type="Array" />
+    /// </signature>
+  },
+  'parseJSON': function() {
+    /// <signature>
+    ///   <summary>Takes a well-formed JSON string and returns the resulting JavaScript object.</summary>
+    ///   <param name="json" type="String">The JSON string to parse.</param>
+    ///   <returns type="Object" />
+    /// </signature>
+  },
+  'parseXML': function() {
+    /// <signature>
+    ///   <summary>Parses a string into an XML document.</summary>
+    ///   <param name="data" type="String">a well-formed XML string to be parsed</param>
+    ///   <returns type="XMLDocument" />
+    /// </signature>
+  },
+  'post': function() {
+    /// <signature>
+    ///   <summary>Load data from the server using a HTTP POST request.</summary>
+    ///   <param name="url" type="String">A string containing the URL to which the request is sent.</param>
+    ///   <param name="data" type="">A plain object or string that is sent to the server with the request.</param>
+    ///   <param name="success(data, textStatus, jqXHR)" type="Function">A callback function that is executed if the request succeeds. Required if dataType is provided, but can be null in that case.</param>
+    ///   <param name="dataType" type="String">The type of data expected from the server. Default: Intelligent Guess (xml, json, script, text, html).</param>
+    ///   <returns type="jqXHR" />
+    /// </signature>
+  },
+  'proxy': function() {
+    /// <signature>
+    ///   <summary>Takes a function and returns a new one that will always have a particular context.</summary>
+    ///   <param name="function" type="Function">The function whose context will be changed.</param>
+    ///   <param name="context" type="PlainObject">The object to which the context (this) of the function should be set.</param>
+    ///   <returns type="Function" />
+    /// </signature>
+    /// <signature>
+    ///   <summary>Takes a function and returns a new one that will always have a particular context.</summary>
+    ///   <param name="context" type="PlainObject">The object to which the context of the function should be set.</param>
+    ///   <param name="name" type="String">The name of the function whose context will be changed (should be a property of the context object).</param>
+    ///   <returns type="Function" />
+    /// </signature>
+    /// <signature>
+    ///   <summary>Takes a function and returns a new one that will always have a particular context.</summary>
+    ///   <param name="function" type="Function">The function whose context will be changed.</param>
+    ///   <param name="context" type="PlainObject">The object to which the context (this) of the function should be set.</param>
+    ///   <param name="additionalArguments" type="Anything">Any number of arguments to be passed to the function referenced in the function argument.</param>
+    ///   <returns type="Function" />
+    /// </signature>
+    /// <signature>
+    ///   <summary>Takes a function and returns a new one that will always have a particular context.</summary>
+    ///   <param name="context" type="PlainObject">The object to which the context of the function should be set.</param>
+    ///   <param name="name" type="String">The name of the function whose context will be changed (should be a property of the context object).</param>
+    ///   <param name="additionalArguments" type="Anything">Any number of arguments to be passed to the function named in the name argument.</param>
+    ///   <returns type="Function" />
+    /// </signature>
+  },
+  'queue': function() {
+    /// <signature>
+    ///   <summary>Manipulate the queue of functions to be executed on the matched element.</summary>
+    ///   <param name="element" type="Element">A DOM element where the array of queued functions is attached.</param>
+    ///   <param name="queueName" type="String">A string containing the name of the queue. Defaults to fx, the standard effects queue.</param>
+    ///   <param name="newQueue" type="Array">An array of functions to replace the current queue contents.</param>
+    ///   <returns type="jQuery" />
+    /// </signature>
+    /// <signature>
+    ///   <summary>Manipulate the queue of functions to be executed on the matched element.</summary>
+    ///   <param name="element" type="Element">A DOM element on which to add a queued function.</param>
+    ///   <param name="queueName" type="String">A string containing the name of the queue. Defaults to fx, the standard effects queue.</param>
+    ///   <param name="callback()" type="Function">The new function to add to the queue.</param>
+    ///   <returns type="jQuery" />
+    /// </signature>
+  },
+  'removeData': function() {
+    /// <signature>
+    ///   <summary>Remove a previously-stored piece of data.</summary>
+    ///   <param name="element" type="Element">A DOM element from which to remove data.</param>
+    ///   <param name="name" type="String">A string naming the piece of data to remove.</param>
+    ///   <returns type="jQuery" />
+    /// </signature>
+  },
+  'sub': function() {
+    /// <summary>Creates a new copy of jQuery whose properties and methods can be modified without affecting the original jQuery object.</summary>
+    /// <returns type="jQuery" />
+  },
+  'support': function() {
+    /// <summary>A collection of properties that represent the presence of different browser features or bugs. Primarily intended for jQuery's internal use; specific properties may be removed when they are no longer needed internally to improve page startup performance.</summary>
+    /// <returns type="Object" />
+  },
+  'trim': function() {
+    /// <signature>
+    ///   <summary>Remove the whitespace from the beginning and end of a string.</summary>
+    ///   <param name="str" type="String">The string to trim.</param>
+    ///   <returns type="String" />
+    /// </signature>
+  },
+  'type': function() {
+    /// <signature>
+    ///   <summary>Determine the internal JavaScript [[Class]] of an object.</summary>
+    ///   <param name="obj" type="PlainObject">Object to get the internal JavaScript [[Class]] of.</param>
+    ///   <returns type="String" />
+    /// </signature>
+  },
+  'unique': function() {
+    /// <signature>
+    ///   <summary>Sorts an array of DOM elements, in place, with the duplicates removed. Note that this only works on arrays of DOM elements, not strings or numbers.</summary>
+    ///   <param name="array" type="Array">The Array of DOM elements.</param>
+    ///   <returns type="Array" />
+    /// </signature>
+  },
+  'when': function() {
+    /// <signature>
+    ///   <summary>Provides a way to execute callback functions based on one or more objects, usually Deferred objects that represent asynchronous events.</summary>
+    ///   <param name="deferreds" type="Deferred">One or more Deferred objects, or plain JavaScript objects.</param>
+    ///   <returns type="Promise" />
+    /// </signature>
+  },
+});
+
+var _1228819969 = jQuery.Callbacks;
+jQuery.Callbacks = function(flags) {
+var _object = _1228819969(flags);
+intellisense.annotate(_object, {
+  'add': function() {
+    /// <signature>
+    ///   <summary>Add a callback or a collection of callbacks to a callback list.</summary>
+    ///   <param name="callbacks" type="">A function, or array of functions, that are to be added to the callback list.</param>
+    ///   <returns type="Callbacks" />
+    /// </signature>
+  },
+  'disable': function() {
+    /// <summary>Disable a callback list from doing anything more.</summary>
+    /// <returns type="Callbacks" />
+  },
+  'disabled': function() {
+    /// <summary>Determine if the callbacks list has been disabled.</summary>
+    /// <returns type="Boolean" />
+  },
+  'empty': function() {
+    /// <summary>Remove all of the callbacks from a list.</summary>
+    /// <returns type="Callbacks" />
+  },
+  'fire': function() {
+    /// <signature>
+    ///   <summary>Call all of the callbacks with the given arguments</summary>
+    ///   <param name="arguments" type="Anything">The argument or list of arguments to pass back to the callback list.</param>
+    ///   <returns type="Callbacks" />
+    /// </signature>
+  },
+  'fired': function() {
+    /// <summary>Determine if the callbacks have already been called at least once.</summary>
+    /// <returns type="Boolean" />
+  },
+  'fireWith': function() {
+    /// <signature>
+    ///   <summary>Call all callbacks in a list with the given context and arguments.</summary>
+    ///   <param name="context" type="">A reference to the context in which the callbacks in the list should be fired.</param>
+    ///   <param name="args" type="">An argument, or array of arguments, to pass to the callbacks in the list.</param>
+    ///   <returns type="Callbacks" />
+    /// </signature>
+  },
+  'has': function() {
+    /// <signature>
+    ///   <summary>Determine whether a supplied callback is in a list</summary>
+    ///   <param name="callback" type="Function">The callback to search for.</param>
+    ///   <returns type="Boolean" />
+    /// </signature>
+  },
+  'lock': function() {
+    /// <summary>Lock a callback list in its current state.</summary>
+    /// <returns type="Callbacks" />
+  },
+  'locked': function() {
+    /// <summary>Determine if the callbacks list has been locked.</summary>
+    /// <returns type="Boolean" />
+  },
+  'remove': function() {
+    /// <signature>
+    ///   <summary>Remove a callback or a collection of callbacks from a callback list.</summary>
+    ///   <param name="callbacks" type="">A function, or array of functions, that are to be removed from the callback list.</param>
+    ///   <returns type="Callbacks" />
+    /// </signature>
+  },
+});
+
+return _object;
+};
+intellisense.redirectDefinition(jQuery.Callbacks, _1228819969);
+
+var _731531622 = jQuery.Deferred;
+jQuery.Deferred = function(func) {
+var _object = _731531622(func);
+intellisense.annotate(_object, {
+  'always': function() {
+    /// <signature>
+    ///   <summary>Add handlers to be called when the Deferred object is either resolved or rejected.</summary>
+    ///   <param name="alwaysCallbacks" type="Function">A function, or array of functions, that is called when the Deferred is resolved or rejected.</param>
+    ///   <param name="alwaysCallbacks" type="Function">Optional additional functions, or arrays of functions, that are called when the Deferred is resolved or rejected.</param>
+    ///   <returns type="Deferred" />
+    /// </signature>
+  },
+  'done': function() {
+    /// <signature>
+    ///   <summary>Add handlers to be called when the Deferred object is resolved.</summary>
+    ///   <param name="doneCallbacks" type="Function">A function, or array of functions, that are called when the Deferred is resolved.</param>
+    ///   <param name="doneCallbacks" type="Function">Optional additional functions, or arrays of functions, that are called when the Deferred is resolved.</param>
+    ///   <returns type="Deferred" />
+    /// </signature>
+  },
+  'fail': function() {
+    /// <signature>
+    ///   <summary>Add handlers to be called when the Deferred object is rejected.</summary>
+    ///   <param name="failCallbacks" type="Function">A function, or array of functions, that are called when the Deferred is rejected.</param>
+    ///   <param name="failCallbacks" type="Function">Optional additional functions, or arrays of functions, that are called when the Deferred is rejected.</param>
+    ///   <returns type="Deferred" />
+    /// </signature>
+  },
+  'isRejected': function() {
+    /// <summary>Determine whether a Deferred object has been rejected.</summary>
+    /// <returns type="Boolean" />
+  },
+  'isResolved': function() {
+    /// <summary>Determine whether a Deferred object has been resolved.</summary>
+    /// <returns type="Boolean" />
+  },
+  'notify': function() {
+    /// <signature>
+    ///   <summary>Call the progressCallbacks on a Deferred object with the given args.</summary>
+    ///   <param name="args" type="Object">Optional arguments that are passed to the progressCallbacks.</param>
+    ///   <returns type="Deferred" />
+    /// </signature>
+  },
+  'notifyWith': function() {
+    /// <signature>
+    ///   <summary>Call the progressCallbacks on a Deferred object with the given context and args.</summary>
+    ///   <param name="context" type="Object">Context passed to the progressCallbacks as the this object.</param>
+    ///   <param name="args" type="Object">Optional arguments that are passed to the progressCallbacks.</param>
+    ///   <returns type="Deferred" />
+    /// </signature>
+  },
+  'pipe': function() {
+    /// <signature>
+    ///   <summary>Utility method to filter and/or chain Deferreds.</summary>
+    ///   <param name="doneFilter" type="Function">An optional function that is called when the Deferred is resolved.</param>
+    ///   <param name="failFilter" type="Function">An optional function that is called when the Deferred is rejected.</param>
+    ///   <returns type="Promise" />
+    /// </signature>
+    /// <signature>
+    ///   <summary>Utility method to filter and/or chain Deferreds.</summary>
+    ///   <param name="doneFilter" type="Function">An optional function that is called when the Deferred is resolved.</param>
+    ///   <param name="failFilter" type="Function">An optional function that is called when the Deferred is rejected.</param>
+    ///   <param name="progressFilter" type="Function">An optional function that is called when progress notifications are sent to the Deferred.</param>
+    ///   <returns type="Promise" />
+    /// </signature>
+  },
+  'progress': function() {
+    /// <signature>
+    ///   <summary>Add handlers to be called when the Deferred object generates progress notifications.</summary>
+    ///   <param name="progressCallbacks" type="">A function, or array of functions, to be called when the Deferred generates progress notifications.</param>
+    ///   <returns type="Deferred" />
+    /// </signature>
+  },
+  'promise': function() {
+    /// <signature>
+    ///   <summary>Return a Deferred's Promise object.</summary>
+    ///   <param name="target" type="Object">Object onto which the promise methods have to be attached</param>
+    ///   <returns type="Promise" />
+    /// </signature>
+  },
+  'reject': function() {
+    /// <signature>
+    ///   <summary>Reject a Deferred object and call any failCallbacks with the given args.</summary>
+    ///   <param name="args" type="Anything">Optional arguments that are passed to the failCallbacks.</param>
+    ///   <returns type="Deferred" />
+    /// </signature>
+  },
+  'rejectWith': function() {
+    /// <signature>
+    ///   <summary>Reject a Deferred object and call any failCallbacks with the given context and args.</summary>
+    ///   <param name="context" type="Object">Context passed to the failCallbacks as the this object.</param>
+    ///   <param name="args" type="Array">An optional array of arguments that are passed to the failCallbacks.</param>
+    ///   <returns type="Deferred" />
+    /// </signature>
+  },
+  'resolve': function() {
+    /// <signature>
+    ///   <summary>Resolve a Deferred object and call any doneCallbacks with the given args.</summary>
+    ///   <param name="args" type="Anything">Optional arguments that are passed to the doneCallbacks.</param>
+    ///   <returns type="Deferred" />
+    /// </signature>
+  },
+  'resolveWith': function() {
+    /// <signature>
+    ///   <summary>Resolve a Deferred object and call any doneCallbacks with the given context and args.</summary>
+    ///   <param name="context" type="Object">Context passed to the doneCallbacks as the this object.</param>
+    ///   <param name="args" type="Array">An optional array of arguments that are passed to the doneCallbacks.</param>
+    ///   <returns type="Deferred" />
+    /// </signature>
+  },
+  'state': function() {
+    /// <summary>Determine the current state of a Deferred object.</summary>
+    /// <returns type="String" />
+  },
+  'then': function() {
+    /// <signature>
+    ///   <summary>Add handlers to be called when the Deferred object is resolved, rejected, or still in progress.</summary>
+    ///   <param name="doneFilter" type="Function">A function that is called when the Deferred is resolved.</param>
+    ///   <param name="failFilter" type="Function">An optional function that is called when the Deferred is rejected.</param>
+    ///   <param name="progressFilter" type="Function">An optional function that is called when progress notifications are sent to the Deferred.</param>
+    ///   <returns type="Promise" />
+    /// </signature>
+    /// <signature>
+    ///   <summary>Add handlers to be called when the Deferred object is resolved, rejected, or still in progress.</summary>
+    ///   <param name="doneCallbacks" type="Function">A function, or array of functions, called when the Deferred is resolved.</param>
+    ///   <param name="failCallbacks" type="Function">A function, or array of functions, called when the Deferred is rejected.</param>
+    ///   <returns type="Promise" />
+    /// </signature>
+    /// <signature>
+    ///   <summary>Add handlers to be called when the Deferred object is resolved, rejected, or still in progress.</summary>
+    ///   <param name="doneCallbacks" type="Function">A function, or array of functions, called when the Deferred is resolved.</param>
+    ///   <param name="failCallbacks" type="Function">A function, or array of functions, called when the Deferred is rejected.</param>
+    ///   <param name="progressCallbacks" type="Function">A function, or array of functions, called when the Deferred notifies progress.</param>
+    ///   <returns type="Promise" />
+    /// </signature>
+  },
+});
+
+return _object;
+};
+intellisense.redirectDefinition(jQuery.Callbacks, _731531622);
+
+intellisense.annotate(jQuery.Event.prototype, {
+  'currentTarget': function() {
+    /// <summary>The current DOM element within the event bubbling phase.</summary>
+    /// <returns type="Element" />
+  },
+  'data': function() {
+    /// <summary>An optional object of data passed to an event method when the current executing handler is bound.</summary>
+    /// <returns type="Object" />
+  },
+  'delegateTarget': function() {
+    /// <summary>The element where the currently-called jQuery event handler was attached.</summary>
+    /// <returns type="Element" />
+  },
+  'isDefaultPrevented': function() {
+    /// <summary>Returns whether event.preventDefault() was ever called on this event object.</summary>
+    /// <returns type="Boolean" />
+  },
+  'isImmediatePropagationStopped': function() {
+    /// <summary>Returns whether event.stopImmediatePropagation() was ever called on this event object.</summary>
+    /// <returns type="Boolean" />
+  },
+  'isPropagationStopped': function() {
+    /// <summary>Returns whether event.stopPropagation() was ever called on this event object.</summary>
+    /// <returns type="Boolean" />
+  },
+  'metaKey': function() {
+    /// <summary>Indicates whether the META key was pressed when the event fired.</summary>
+    /// <returns type="Boolean" />
+  },
+  'namespace': function() {
+    /// <summary>The namespace specified when the event was triggered.</summary>
+    /// <returns type="String" />
+  },
+  'pageX': function() {
+    /// <summary>The mouse position relative to the left edge of the document.</summary>
+    /// <returns type="Number" />
+  },
+  'pageY': function() {
+    /// <summary>The mouse position relative to the top edge of the document.</summary>
+    /// <returns type="Number" />
+  },
+  'preventDefault': function() {
+    /// <summary>If this method is called, the default action of the event will not be triggered.</summary>
+  },
+  'relatedTarget': function() {
+    /// <summary>The other DOM element involved in the event, if any.</summary>
+    /// <returns type="Element" />
+  },
+  'result': function() {
+    /// <summary>The last value returned by an event handler that was triggered by this event, unless the value was undefined.</summary>
+    /// <returns type="Object" />
+  },
+  'stopImmediatePropagation': function() {
+    /// <summary>Keeps the rest of the handlers from being executed and prevents the event from bubbling up the DOM tree.</summary>
+  },
+  'stopPropagation': function() {
+    /// <summary>Prevents the event from bubbling up the DOM tree, preventing any parent handlers from being notified of the event.</summary>
+  },
+  'target': function() {
+    /// <summary>The DOM element that initiated the event.</summary>
+    /// <returns type="Element" />
+  },
+  'timeStamp': function() {
+    /// <summary>The difference in milliseconds between the time the browser created the event and January 1, 1970.</summary>
+    /// <returns type="Number" />
+  },
+  'type': function() {
+    /// <summary>Describes the nature of the event.</summary>
+    /// <returns type="String" />
+  },
+  'which': function() {
+    /// <summary>For key or mouse events, this property indicates the specific key or button that was pressed.</summary>
+    /// <returns type="Number" />
+  },
+});
+
+intellisense.annotate(jQuery.fn, {
+  'add': function() {
+    /// <signature>
+    ///   <summary>Add elements to the set of matched elements.</summary>
+    ///   <param name="selector" type="String">A string representing a selector expression to find additional elements to add to the set of matched elements.</param>
+    ///   <returns type="jQuery" />
+    /// </signature>
+    /// <signature>
+    ///   <summary>Add elements to the set of matched elements.</summary>
+    ///   <param name="elements" type="Array">One or more elements to add to the set of matched elements.</param>
+    ///   <returns type="jQuery" />
+    /// </signature>
+    /// <signature>
+    ///   <summary>Add elements to the set of matched elements.</summary>
+    ///   <param name="html" type="htmlString">An HTML fragment to add to the set of matched elements.</param>
+    ///   <returns type="jQuery" />
+    /// </signature>
+    /// <signature>
+    ///   <summary>Add elements to the set of matched elements.</summary>
+    ///   <param name="jQuery object" type="jQuery object ">An existing jQuery object to add to the set of matched elements.</param>
+    ///   <returns type="jQuery" />
+    /// </signature>
+    /// <signature>
+    ///   <summary>Add elements to the set of matched elements.</summary>
+    ///   <param name="selector" type="String">A string representing a selector expression to find additional elements to add to the set of matched elements.</param>
+    ///   <param name="context" type="Element">The point in the document at which the selector should begin matching; similar to the context argument of the $(selector, context) method.</param>
+    ///   <returns type="jQuery" />
+    /// </signature>
+  },
+  'addBack': function() {
+    /// <signature>
+    ///   <summary>Add the previous set of elements on the stack to the current set, optionally filtered by a selector.</summary>
+    ///   <param name="selector" type="String">A string containing a selector expression to match the current set of elements against.</param>
+    ///   <returns type="jQuery" />
+    /// </signature>
+  },
+  'addClass': function() {
+    /// <signature>
+    ///   <summary>Adds the specified class(es) to each of the set of matched elements.</summary>
+    ///   <param name="className" type="String">One or more space-separated classes to be added to the class attribute of each matched element.</param>
+    ///   <returns type="jQuery" />
+    /// </signature>
+    /// <signature>
+    ///   <summary>Adds the specified class(es) to each of the set of matched elements.</summary>
+    ///   <param name="function(index, currentClass)" type="Function">A function returning one or more space-separated class names to be added to the existing class name(s). Receives the index position of the element in the set and the existing class name(s) as arguments. Within the function, this refers to the current element in the set.</param>
+    ///   <returns type="jQuery" />
+    /// </signature>
+  },
+  'after': function() {
+    /// <signature>
+    ///   <summary>Insert content, specified by the parameter, after each element in the set of matched elements.</summary>
+    ///   <param name="content" type="">HTML string, DOM element, or jQuery object to insert after each element in the set of matched elements.</param>
+    ///   <param name="content" type="">One or more additional DOM elements, arrays of elements, HTML strings, or jQuery objects to insert after each element in the set of matched elements.</param>
+    ///   <returns type="jQuery" />
+    /// </signature>
+    /// <signature>
+    ///   <summary>Insert content, specified by the parameter, after each element in the set of matched elements.</summary>
+    ///   <param name="function(index)" type="Function">A function that returns an HTML string, DOM element(s), or jQuery object to insert after each element in the set of matched elements. Receives the index position of the element in the set as an argument. Within the function, this refers to the current element in the set.</param>
+    ///   <returns type="jQuery" />
+    /// </signature>
+  },
+  'ajaxComplete': function() {
+    /// <signature>
+    ///   <summary>Register a handler to be called when Ajax requests complete. This is an AjaxEvent.</summary>
+    ///   <param name="handler(event, XMLHttpRequest, ajaxOptions)" type="Function">The function to be invoked.</param>
+    ///   <returns type="jQuery" />
+    /// </signature>
+  },
+  'ajaxError': function() {
+    /// <signature>
+    ///   <summary>Register a handler to be called when Ajax requests complete with an error. This is an Ajax Event.</summary>
+    ///   <param name="handler(event, jqXHR, ajaxSettings, thrownError)" type="Function">The function to be invoked.</param>
+    ///   <returns type="jQuery" />
+    /// </signature>
+  },
+  'ajaxSend': function() {
+    /// <signature>
+    ///   <summary>Attach a function to be executed before an Ajax request is sent. This is an Ajax Event.</summary>
+    ///   <param name="handler(event, jqXHR, ajaxOptions)" type="Function">The function to be invoked.</param>
+    ///   <returns type="jQuery" />
+    /// </signature>
+  },
+  'ajaxStart': function() {
+    /// <signature>
+    ///   <summary>Register a handler to be called when the first Ajax request begins. This is an Ajax Event.</summary>
+    ///   <param name="handler()" type="Function">The function to be invoked.</param>
+    ///   <returns type="jQuery" />
+    /// </signature>
+  },
+  'ajaxStop': function() {
+    /// <signature>
+    ///   <summary>Register a handler to be called when all Ajax requests have completed. This is an Ajax Event.</summary>
+    ///   <param name="handler()" type="Function">The function to be invoked.</param>
+    ///   <returns type="jQuery" />
+    /// </signature>
+  },
+  'ajaxSuccess': function() {
+    /// <signature>
+    ///   <summary>Attach a function to be executed whenever an Ajax request completes successfully. This is an Ajax Event.</summary>
+    ///   <param name="handler(event, XMLHttpRequest, ajaxOptions)" type="Function">The function to be invoked.</param>
+    ///   <returns type="jQuery" />
+    /// </signature>
+  },
+  'all': function() {
+    /// <summary>Selects all elements.</summary>
+  },
+  'andSelf': function() {
+    /// <summary>Add the previous set of elements on the stack to the current set.</summary>
+    /// <returns type="jQuery" />
+  },
+  'animate': function() {
+    /// <signature>
+    ///   <summary>Perform a custom animation of a set of CSS properties.</summary>
+    ///   <param name="properties" type="PlainObject">An object of CSS properties and values that the animation will move toward.</param>
+    ///   <param name="duration" type="">A string or number determining how long the animation will run.</param>
+    ///   <param name="easing" type="String">A string indicating which easing function to use for the transition.</param>
+    ///   <param name="complete" type="Function">A function to call once the animation is complete.</param>
+    ///   <returns type="jQuery" />
+    /// </signature>
+    /// <signature>
+    ///   <summary>Perform a custom animation of a set of CSS properties.</summary>
+    ///   <param name="properties" type="PlainObject">An object of CSS properties and values that the animation will move toward.</param>
+    ///   <param name="options" type="PlainObject">A map of additional options to pass to the method.</param>
+    ///   <returns type="jQuery" />
+    /// </signature>
+  },
+  'animated': function() {
+    /// <summary>Select all elements that are in the progress of an animation at the time the selector is run.</summary>
+  },
+  'append': function() {
+    /// <signature>
+    ///   <summary>Insert content, specified by the parameter, to the end of each element in the set of matched elements.</summary>
+    ///   <param name="content" type="">DOM element, HTML string, or jQuery object to insert at the end of each element in the set of matched elements.</param>
+    ///   <param name="content" type="">One or more additional DOM elements, arrays of elements, HTML strings, or jQuery objects to insert at the end of each element in the set of matched elements.</param>
+    ///   <returns type="jQuery" />
+    /// </signature>
+    /// <signature>
+    ///   <summary>Insert content, specified by the parameter, to the end of each element in the set of matched elements.</summary>
+    ///   <param name="function(index, html)" type="Function">A function that returns an HTML string, DOM element(s), or jQuery object to insert at the end of each element in the set of matched elements. Receives the index position of the element in the set and the old HTML value of the element as arguments. Within the function, this refers to the current element in the set.</param>
+    ///   <returns type="jQuery" />
+    /// </signature>
+  },
+  'appendTo': function() {
+    /// <signature>
+    ///   <summary>Insert every element in the set of matched elements to the end of the target.</summary>
+    ///   <param name="target" type="">A selector, element, HTML string, or jQuery object; the matched set of elements will be inserted at the end of the element(s) specified by this parameter.</param>
+    ///   <returns type="jQuery" />
+    /// </signature>
+  },
+  'attr': function() {
+    /// <signature>
+    ///   <summary>Set one or more attributes for the set of matched elements.</summary>
+    ///   <param name="attributeName" type="String">The name of the attribute to set.</param>
+    ///   <param name="value" type="">A value to set for the attribute.</param>
+    ///   <returns type="jQuery" />
+    /// </signature>
+    /// <signature>
+    ///   <summary>Set one or more attributes for the set of matched elements.</summary>
+    ///   <param name="attributes" type="PlainObject">An object of attribute-value pairs to set.</param>
+    ///   <returns type="jQuery" />
+    /// </signature>
+    /// <signature>
+    ///   <summary>Set one or more attributes for the set of matched elements.</summary>
+    ///   <param name="attributeName" type="String">The name of the attribute to set.</param>
+    ///   <param name="function(index, attr)" type="Function">A function returning the value to set. this is the current element. Receives the index position of the element in the set and the old attribute value as arguments.</param>
+    ///   <returns type="jQuery" />
+    /// </signature>
+  },
+  'attributeContains': function() {
+    /// <signature>
+    ///   <summary>Selects elements that have the specified attribute with a value containing the a given substring.</summary>
+    ///   <param name="attribute" type="String">An attribute name.</param>
+    ///   <param name="value" type="String">An attribute value. Can be either an unquoted single word or a quoted string.</param>
+    /// </signature>
+  },
+  'attributeContainsPrefix': function() {
+    /// <signature>
+    ///   <summary>Selects elements that have the specified attribute with a value either equal to a given string or starting with that string followed by a hyphen (-).</summary>
+    ///   <param name="attribute" type="String">An attribute name.</param>
+    ///   <param name="value" type="String">An attribute value. Can be either an unquoted single word or a quoted string.</param>
+    /// </signature>
+  },
+  'attributeContainsWord': function() {
+    /// <signature>
+    ///   <summary>Selects elements that have the specified attribute with a value containing a given word, delimited by spaces.</summary>
+    ///   <param name="attribute" type="String">An attribute name.</param>
+    ///   <param name="value" type="String">An attribute value. Can be either an unquoted single word or a quoted string.</param>
+    /// </signature>
+  },
+  'attributeEndsWith': function() {
+    /// <signature>
+    ///   <summary>Selects elements that have the specified attribute with a value ending exactly with a given string. The comparison is case sensitive.</summary>
+    ///   <param name="attribute" type="String">An attribute name.</param>
+    ///   <param name="value" type="String">An attribute value. Can be either an unquoted single word or a quoted string.</param>
+    /// </signature>
+  },
+  'attributeEquals': function() {
+    /// <signature>
+    ///   <summary>Selects elements that have the specified attribute with a value exactly equal to a certain value.</summary>
+    ///   <param name="attribute" type="String">An attribute name.</param>
+    ///   <param name="value" type="String">An attribute value. Can be either an unquoted single word or a quoted string.</param>
+    /// </signature>
+  },
+  'attributeHas': function() {
+    /// <signature>
+    ///   <summary>Selects elements that have the specified attribute, with any value.</summary>
+    ///   <param name="attribute" type="String">An attribute name.</param>
+    /// </signature>
+  },
+  'attributeMultiple': function() {
+    /// <signature>
+    ///   <summary>Matches elements that match all of the specified attribute filters.</summary>
+    ///   <param name="attributeFilter1" type="String">An attribute filter.</param>
+    ///   <param name="attributeFilter2" type="String">Another attribute filter, reducing the selection even more</param>
+    ///   <param name="attributeFilterN" type="String">As many more attribute filters as necessary</param>
+    /// </signature>
+  },
+  'attributeNotEqual': function() {
+    /// <signature>
+    ///   <summary>Select elements that either don't have the specified attribute, or do have the specified attribute but not with a certain value.</summary>
+    ///   <param name="attribute" type="String">An attribute name.</param>
+    ///   <param name="value" type="String">An attribute value. Can be either an unquoted single word or a quoted string.</param>
+    /// </signature>
+  },
+  'attributeStartsWith': function() {
+    /// <signature>
+    ///   <summary>Selects elements that have the specified attribute with a value beginning exactly with a given string.</summary>
+    ///   <param name="attribute" type="String">An attribute name.</param>
+    ///   <param name="value" type="String">An attribute value. Can be either an unquoted single word or a quoted string.</param>
+    /// </signature>
+  },
+  'before': function() {
+    /// <signature>
+    ///   <summary>Insert content, specified by the parameter, before each element in the set of matched elements.</summary>
+    ///   <param name="content" type="">HTML string, DOM element, or jQuery object to insert before each element in the set of matched elements.</param>
+    ///   <param name="content" type="">One or more additional DOM elements, arrays of elements, HTML strings, or jQuery objects to insert before each element in the set of matched elements.</param>
+    ///   <returns type="jQuery" />
+    /// </signature>
+    /// <signature>
+    ///   <summary>Insert content, specified by the parameter, before each element in the set of matched elements.</summary>
+    ///   <param name="function" type="Function">A function that returns an HTML string, DOM element(s), or jQuery object to insert before each element in the set of matched elements. Receives the index position of the element in the set as an argument. Within the function, this refers to the current element in the set.</param>
+    ///   <returns type="jQuery" />
+    /// </signature>
+  },
+  'bind': function() {
+    /// <signature>
+    ///   <summary>Attach a handler to an event for the elements.</summary>
+    ///   <param name="eventType" type="String">A string containing one or more DOM event types, such as "click" or "submit," or custom event names.</param>
+    ///   <param name="eventData" type="Object">An object containing data that will be passed to the event handler.</param>
+    ///   <param name="handler(eventObject)" type="Function">A function to execute each time the event is triggered.</param>
+    ///   <returns type="jQuery" />
+    /// </signature>
+    /// <signature>
+    ///   <summary>Attach a handler to an event for the elements.</summary>
+    ///   <param name="eventType" type="String">A string containing one or more DOM event types, such as "click" or "submit," or custom event names.</param>
+    ///   <param name="eventData" type="Object">An object containing data that will be passed to the event handler.</param>
+    ///   <param name="preventBubble" type="Boolean">Setting the third argument to false will attach a function that prevents the default action from occurring and stops the event from bubbling. The default is true.</param>
+    ///   <returns type="jQuery" />
+    /// </signature>
+    /// <signature>
+    ///   <summary>Attach a handler to an event for the elements.</summary>
+    ///   <param name="events" type="Object">An object containing one or more DOM event types and functions to execute for them.</param>
+    ///   <returns type="jQuery" />
+    /// </signature>
+  },
+  'blur': function() {
+    /// <signature>
+    ///   <summary>Bind an event handler to the "blur" JavaScript event, or trigger that event on an element.</summary>
+    ///   <param name="handler(eventObject)" type="Function">A function to execute each time the event is triggered.</param>
+    ///   <returns type="jQuery" />
+    /// </signature>
+    /// <signature>
+    ///   <summary>Bind an event handler to the "blur" JavaScript event, or trigger that event on an element.</summary>
+    ///   <param name="eventData" type="Object">An object containing data that will be passed to the event handler.</param>
+    ///   <param name="handler(eventObject)" type="Function">A function to execute each time the event is triggered.</param>
+    ///   <returns type="jQuery" />
+    /// </signature>
+  },
+  'button': function() {
+    /// <summary>Selects all button elements and elements of type button.</summary>
+  },
+  'change': function() {
+    /// <signature>
+    ///   <summary>Bind an event handler to the "change" JavaScript event, or trigger that event on an element.</summary>
+    ///   <param name="handler(eventObject)" type="Function">A function to execute each time the event is triggered.</param>
+    ///   <returns type="jQuery" />
+    /// </signature>
+    /// <signature>
+    ///   <summary>Bind an event handler to the "change" JavaScript event, or trigger that event on an element.</summary>
+    ///   <param name="eventData" type="Object">An object containing data that will be passed to the event handler.</param>
+    ///   <param name="handler(eventObject)" type="Function">A function to execute each time the event is triggered.</param>
+    ///   <returns type="jQuery" />
+    /// </signature>
+  },
+  'checkbox': function() {
+    /// <summary>Selects all elements of type checkbox.</summary>
+  },
+  'checked': function() {
+    /// <summary>Matches all elements that are checked or selected.</summary>
+  },
+  'child': function() {
+    /// <signature>
+    ///   <summary>Selects all direct child elements specified by "child" of elements specified by "parent".</summary>
+    ///   <param name="parent" type="String">Any valid selector.</param>
+    ///   <param name="child" type="String">A selector to filter the child elements.</param>
+    /// </signature>
+  },
+  'children': function() {
+    /// <signature>
+    ///   <summary>Get the children of each element in the set of matched elements, optionally filtered by a selector.</summary>
+    ///   <param name="selector" type="String">A string containing a selector expression to match elements against.</param>
+    ///   <returns type="jQuery" />
+    /// </signature>
+  },
+  'class': function() {
+    /// <signature>
+    ///   <summary>Selects all elements with the given class.</summary>
+    ///   <param name="class" type="String">A class to search for. An element can have multiple classes; only one of them must match.</param>
+    /// </signature>
+  },
+  'clearQueue': function() {
+    /// <signature>
+    ///   <summary>Remove from the queue all items that have not yet been run.</summary>
+    ///   <param name="queueName" type="String">A string containing the name of the queue. Defaults to fx, the standard effects queue.</param>
+    ///   <returns type="jQuery" />
+    /// </signature>
+  },
+  'click': function() {
+    /// <signature>
+    ///   <summary>Bind an event handler to the "click" JavaScript event, or trigger that event on an element.</summary>
+    ///   <param name="handler(eventObject)" type="Function">A function to execute each time the event is triggered.</param>
+    ///   <returns type="jQuery" />
+    /// </signature>
+    /// <signature>
+    ///   <summary>Bind an event handler to the "click" JavaScript event, or trigger that event on an element.</summary>
+    ///   <param name="eventData" type="Object">An object containing data that will be passed to the event handler.</param>
+    ///   <param name="handler(eventObject)" type="Function">A function to execute each time the event is triggered.</param>
+    ///   <returns type="jQuery" />
+    /// </signature>
+  },
+  'clone': function() {
+    /// <signature>
+    ///   <summary>Create a deep copy of the set of matched elements.</summary>
+    ///   <param name="withDataAndEvents" type="Boolean">A Boolean indicating whether event handlers should be copied along with the elements. As of jQuery 1.4, element data will be copied as well.</param>
+    ///   <returns type="jQuery" />
+    /// </signature>
+    /// <signature>
+    ///   <summary>Create a deep copy of the set of matched elements.</summary>
+    ///   <param name="withDataAndEvents" type="Boolean">A Boolean indicating whether event handlers and data should be copied along with the elements. The default value is false. *In jQuery 1.5.0 the default value was incorrectly true; it was changed back to false in 1.5.1 and up.</param>
+    ///   <param name="deepWithDataAndEvents" type="Boolean">A Boolean indicating whether event handlers and data for all children of the cloned element should be copied. By default its value matches the first argument's value (which defaults to false).</param>
+    ///   <returns type="jQuery" />
+    /// </signature>
+  },
+  'closest': function() {
+    /// <signature>
+    ///   <summary>For each element in the set, get the first element that matches the selector by testing the element itself and traversing up through its ancestors in the DOM tree.</summary>
+    ///   <param name="selector" type="String">A string containing a selector expression to match elements against.</param>
+    ///   <returns type="jQuery" />
+    /// </signature>
+    /// <signature>
+    ///   <summary>For each element in the set, get the first element that matches the selector by testing the element itself and traversing up through its ancestors in the DOM tree.</summary>
+    ///   <param name="selector" type="String">A string containing a selector expression to match elements against.</param>
+    ///   <param name="context" type="Element">A DOM element within which a matching element may be found. If no context is passed in then the context of the jQuery set will be used instead.</param>
+    ///   <returns type="jQuery" />
+    /// </signature>
+    /// <signature>
+    ///   <summary>For each element in the set, get the first element that matches the selector by testing the element itself and traversing up through its ancestors in the DOM tree.</summary>
+    ///   <param name="jQuery object" type="jQuery">A jQuery object to match elements against.</param>
+    ///   <returns type="jQuery" />
+    /// </signature>
+    /// <signature>
+    ///   <summary>For each element in the set, get the first element that matches the selector by testing the element itself and traversing up through its ancestors in the DOM tree.</summary>
+    ///   <param name="element" type="Element">An element to match elements against.</param>
+    ///   <returns type="jQuery" />
+    /// </signature>
+  },
+  'contains': function() {
+    /// <signature>
+    ///   <summary>Select all elements that contain the specified text.</summary>
+    ///   <param name="text" type="String">A string of text to look for. It's case sensitive.</param>
+    /// </signature>
+  },
+  'contents': function() {
+    /// <summary>Get the children of each element in the set of matched elements, including text and comment nodes.</summary>
+    /// <returns type="jQuery" />
+  },
+  'context': function() {
+    /// <summary>The DOM node context originally passed to jQuery(); if none was passed then context will likely be the document.</summary>
+    /// <returns type="Element" />
+  },
+  'css': function() {
+    /// <signature>
+    ///   <summary>Set one or more CSS properties for the set of matched elements.</summary>
+    ///   <param name="propertyName" type="String">A CSS property name.</param>
+    ///   <param name="value" type="">A value to set for the property.</param>
+    ///   <returns type="jQuery" />
+    /// </signature>
+    /// <signature>
+    ///   <summary>Set one or more CSS properties for the set of matched elements.</summary>
+    ///   <param name="propertyName" type="String">A CSS property name.</param>
+    ///   <param name="function(index, value)" type="Function">A function returning the value to set. this is the current element. Receives the index position of the element in the set and the old value as arguments.</param>
+    ///   <returns type="jQuery" />
+    /// </signature>
+    /// <signature>
+    ///   <summary>Set one or more CSS properties for the set of matched elements.</summary>
+    ///   <param name="properties" type="PlainObject">An object of property-value pairs to set.</param>
+    ///   <returns type="jQuery" />
+    /// </signature>
+  },
+  'data': function() {
+    /// <signature>
+    ///   <summary>Store arbitrary data associated with the matched elements.</summary>
+    ///   <param name="key" type="String">A string naming the piece of data to set.</param>
+    ///   <param name="value" type="Object">The new data value; it can be any Javascript type including Array or Object.</param>
+    ///   <returns type="jQuery" />
+    /// </signature>
+    /// <signature>
+    ///   <summary>Store arbitrary data associated with the matched elements.</summary>
+    ///   <param name="obj" type="Object">An object of key-value pairs of data to update.</param>
+    ///   <returns type="jQuery" />
+    /// </signature>
+  },
+  'dblclick': function() {
+    /// <signature>
+    ///   <summary>Bind an event handler to the "dblclick" JavaScript event, or trigger that event on an element.</summary>
+    ///   <param name="handler(eventObject)" type="Function">A function to execute each time the event is triggered.</param>
+    ///   <returns type="jQuery" />
+    /// </signature>
+    /// <signature>
+    ///   <summary>Bind an event handler to the "dblclick" JavaScript event, or trigger that event on an element.</summary>
+    ///   <param name="eventData" type="Object">An object containing data that will be passed to the event handler.</param>
+    ///   <param name="handler(eventObject)" type="Function">A function to execute each time the event is triggered.</param>
+    ///   <returns type="jQuery" />
+    /// </signature>
+  },
+  'delay': function() {
+    /// <signature>
+    ///   <summary>Set a timer to delay execution of subsequent items in the queue.</summary>
+    ///   <param name="duration" type="Number">An integer indicating the number of milliseconds to delay execution of the next item in the queue.</param>
+    ///   <param name="queueName" type="String">A string containing the name of the queue. Defaults to fx, the standard effects queue.</param>
+    ///   <returns type="jQuery" />
+    /// </signature>
+  },
+  'delegate': function() {
+    /// <signature>
+    ///   <summary>Attach a handler to one or more events for all elements that match the selector, now or in the future, based on a specific set of root elements.</summary>
+    ///   <param name="selector" type="String">A selector to filter the elements that trigger the event.</param>
+    ///   <param name="eventType" type="String">A string containing one or more space-separated JavaScript event types, such as "click" or "keydown," or custom event names.</param>
+    ///   <param name="handler(eventObject)" type="Function">A function to execute at the time the event is triggered.</param>
+    ///   <returns type="jQuery" />
+    /// </signature>
+    /// <signature>
+    ///   <summary>Attach a handler to one or more events for all elements that match the selector, now or in the future, based on a specific set of root elements.</summary>
+    ///   <param name="selector" type="String">A selector to filter the elements that trigger the event.</param>
+    ///   <param name="eventType" type="String">A string containing one or more space-separated JavaScript event types, such as "click" or "keydown," or custom event names.</param>
+    ///   <param name="eventData" type="Object">An object containing data that will be passed to the event handler.</param>
+    ///   <param name="handler(eventObject)" type="Function">A function to execute at the time the event is triggered.</param>
+    ///   <returns type="jQuery" />
+    /// </signature>
+    /// <signature>
+    ///   <summary>Attach a handler to one or more events for all elements that match the selector, now or in the future, based on a specific set of root elements.</summary>
+    ///   <param name="selector" type="String">A selector to filter the elements that trigger the event.</param>
+    ///   <param name="events" type="PlainObject">A plain object of one or more event types and functions to execute for them.</param>
+    ///   <returns type="jQuery" />
+    /// </signature>
+  },
+  'dequeue': function() {
+    /// <signature>
+    ///   <summary>Execute the next function on the queue for the matched elements.</summary>
+    ///   <param name="queueName" type="String">A string containing the name of the queue. Defaults to fx, the standard effects queue.</param>
+    ///   <returns type="jQuery" />
+    /// </signature>
+  },
+  'descendant': function() {
+    /// <signature>
+    ///   <summary>Selects all elements that are descendants of a given ancestor.</summary>
+    ///   <param name="ancestor" type="String">Any valid selector.</param>
+    ///   <param name="descendant" type="String">A selector to filter the descendant elements.</param>
+    /// </signature>
+  },
+  'detach': function() {
+    /// <signature>
+    ///   <summary>Remove the set of matched elements from the DOM.</summary>
+    ///   <param name="selector" type="String">A selector expression that filters the set of matched elements to be removed.</param>
+    ///   <returns type="jQuery" />
+    /// </signature>
+  },
+  'die': function() {
+    /// <signature>
+    ///   <summary>Remove event handlers previously attached using .live() from the elements.</summary>
+    ///   <param name="eventType" type="String">A string containing a JavaScript event type, such as click or keydown.</param>
+    ///   <param name="handler" type="String">The function that is no longer to be executed.</param>
+    ///   <returns type="jQuery" />
+    /// </signature>
+    /// <signature>
+    ///   <summary>Remove event handlers previously attached using .live() from the elements.</summary>
+    ///   <param name="events" type="PlainObject">A plain object of one or more event types, such as click or keydown and their corresponding functions that are no longer to be executed.</param>
+    ///   <returns type="jQuery" />
+    /// </signature>
+  },
+  'disabled': function() {
+    /// <summary>Selects all elements that are disabled.</summary>
+  },
+  'each': function() {
+    /// <signature>
+    ///   <summary>Iterate over a jQuery object, executing a function for each matched element.</summary>
+    ///   <param name="function(index, Element)" type="Function">A function to execute for each matched element.</param>
+    ///   <returns type="jQuery" />
+    /// </signature>
+  },
+  'element': function() {
+    /// <signature>
+    ///   <summary>Selects all elements with the given tag name.</summary>
+    ///   <param name="element" type="String">An element to search for. Refers to the tagName of DOM nodes.</param>
+    /// </signature>
+  },
+  'empty': function() {
+    /// <summary>Select all elements that have no children (including text nodes).</summary>
+  },
+  'enabled': function() {
+    /// <summary>Selects all elements that are enabled.</summary>
+  },
+  'end': function() {
+    /// <summary>End the most recent filtering operation in the current chain and return the set of matched elements to its previous state.</summary>
+    /// <returns type="jQuery" />
+  },
+  'eq': function() {
+    /// <signature>
+    ///   <summary>Select the element at index n within the matched set.</summary>
+    ///   <param name="index" type="Number">Zero-based index of the element to match.</param>
+    /// </signature>
+    /// <signature>
+    ///   <summary>Select the element at index n within the matched set.</summary>
+    ///   <param name="-index" type="Number">Zero-based index of the element to match, counting backwards from the last element.</param>
+    /// </signature>
+  },
+  'error': function() {
+    /// <signature>
+    ///   <summary>Bind an event handler to the "error" JavaScript event.</summary>
+    ///   <param name="handler(eventObject)" type="Function">A function to execute when the event is triggered.</param>
+    ///   <returns type="jQuery" />
+    /// </signature>
+    /// <signature>
+    ///   <summary>Bind an event handler to the "error" JavaScript event.</summary>
+    ///   <param name="eventData" type="Object">An object containing data that will be passed to the event handler.</param>
+    ///   <param name="handler(eventObject)" type="Function">A function to execute each time the event is triggered.</param>
+    ///   <returns type="jQuery" />
+    /// </signature>
+  },
+  'even': function() {
+    /// <summary>Selects even elements, zero-indexed.  See also odd.</summary>
+  },
+  'fadeIn': function() {
+    /// <signature>
+    ///   <summary>Display the matched elements by fading them to opaque.</summary>
+    ///   <param name="duration" type="">A string or number determining how long the animation will run.</param>
+    ///   <param name="complete" type="Function">A function to call once the animation is complete.</param>
+    ///   <returns type="jQuery" />
+    /// </signature>
+    /// <signature>
+    ///   <summary>Display the matched elements by fading them to opaque.</summary>
+    ///   <param name="options" type="PlainObject">A map of additional options to pass to the method.</param>
+    ///   <returns type="jQuery" />
+    /// </signature>
+    /// <signature>
+    ///   <summary>Display the matched elements by fading them to opaque.</summary>
+    ///   <param name="duration" type="">A string or number determining how long the animation will run.</param>
+    ///   <param name="easing" type="String">A string indicating which easing function to use for the transition.</param>
+    ///   <param name="complete" type="Function">A function to call once the animation is complete.</param>
+    ///   <returns type="jQuery" />
+    /// </signature>
+  },
+  'fadeOut': function() {
+    /// <signature>
+    ///   <summary>Hide the matched elements by fading them to transparent.</summary>
+    ///   <param name="duration" type="">A string or number determining how long the animation will run.</param>
+    ///   <param name="complete" type="Function">A function to call once the animation is complete.</param>
+    ///   <returns type="jQuery" />
+    /// </signature>
+    /// <signature>
+    ///   <summary>Hide the matched elements by fading them to transparent.</summary>
+    ///   <param name="options" type="PlainObject">A map of additional options to pass to the method.</param>
+    ///   <returns type="jQuery" />
+    /// </signature>
+    /// <signature>
+    ///   <summary>Hide the matched elements by fading them to transparent.</summary>
+    ///   <param name="duration" type="">A string or number determining how long the animation will run.</param>
+    ///   <param name="easing" type="String">A string indicating which easing function to use for the transition.</param>
+    ///   <param name="complete" type="Function">A function to call once the animation is complete.</param>
+    ///   <returns type="jQuery" />
+    /// </signature>
+  },
+  'fadeTo': function() {
+    /// <signature>
+    ///   <summary>Adjust the opacity of the matched elements.</summary>
+    ///   <param name="duration" type="">A string or number determining how long the animation will run.</param>
+    ///   <param name="opacity" type="Number">A number between 0 and 1 denoting the target opacity.</param>
+    ///   <param name="complete" type="Function">A function to call once the animation is complete.</param>
+    ///   <returns type="jQuery" />
+    /// </signature>
+    /// <signature>
+    ///   <summary>Adjust the opacity of the matched elements.</summary>
+    ///   <param name="duration" type="">A string or number determining how long the animation will run.</param>
+    ///   <param name="opacity" type="Number">A number between 0 and 1 denoting the target opacity.</param>
+    ///   <param name="easing" type="String">A string indicating which easing function to use for the transition.</param>
+    ///   <param name="complete" type="Function">A function to call once the animation is complete.</param>
+    ///   <returns type="jQuery" />
+    /// </signature>
+  },
+  'fadeToggle': function() {
+    /// <signature>
+    ///   <summary>Display or hide the matched elements by animating their opacity.</summary>
+    ///   <param name="duration" type="">A string or number determining how long the animation will run.</param>
+    ///   <param name="easing" type="String">A string indicating which easing function to use for the transition.</param>
+    ///   <param name="complete" type="Function">A function to call once the animation is complete.</param>
+    ///   <returns type="jQuery" />
+    /// </signature>
+    /// <signature>
+    ///   <summary>Display or hide the matched elements by animating their opacity.</summary>
+    ///   <param name="options" type="PlainObject">A map of additional options to pass to the method.</param>
+    ///   <returns type="jQuery" />
+    /// </signature>
+  },
+  'file': function() {
+    /// <summary>Selects all elements of type file.</summary>
+  },
+  'filter': function() {
+    /// <signature>
+    ///   <summary>Reduce the set of matched elements to those that match the selector or pass the function's test.</summary>
+    ///   <param name="selector" type="String">A string containing a selector expression to match the current set of elements against.</param>
+    ///   <returns type="jQuery" />
+    /// </signature>
+    /// <signature>
+    ///   <summary>Reduce the set of matched elements to those that match the selector or pass the function's test.</summary>
+    ///   <param name="function(index)" type="Function">A function used as a test for each element in the set. this is the current DOM element.</param>
+    ///   <returns type="jQuery" />
+    /// </signature>
+    /// <signature>
+    ///   <summary>Reduce the set of matched elements to those that match the selector or pass the function's test.</summary>
+    ///   <param name="element" type="Element">An element to match the current set of elements against.</param>
+    ///   <returns type="jQuery" />
+    /// </signature>
+    /// <signature>
+    ///   <summary>Reduce the set of matched elements to those that match the selector or pass the function's test.</summary>
+    ///   <param name="jQuery object" type="Object">An existing jQuery object to match the current set of elements against.</param>
+    ///   <returns type="jQuery" />
+    /// </signature>
+  },
+  'find': function() {
+    /// <signature>
+    ///   <summary>Get the descendants of each element in the current set of matched elements, filtered by a selector, jQuery object, or element.</summary>
+    ///   <param name="selector" type="String">A string containing a selector expression to match elements against.</param>
+    ///   <returns type="jQuery" />
+    /// </signature>
+    /// <signature>
+    ///   <summary>Get the descendants of each element in the current set of matched elements, filtered by a selector, jQuery object, or element.</summary>
+    ///   <param name="jQuery object" type="Object">A jQuery object to match elements against.</param>
+    ///   <returns type="jQuery" />
+    /// </signature>
+    /// <signature>
+    ///   <summary>Get the descendants of each element in the current set of matched elements, filtered by a selector, jQuery object, or element.</summary>
+    ///   <param name="element" type="Element">An element to match elements against.</param>
+    ///   <returns type="jQuery" />
+    /// </signature>
+  },
+  'finish': function() {
+    /// <signature>
+    ///   <summary>Stop the currently-running animation, remove all queued animations, and complete all animations for the matched elements.</summary>
+    ///   <param name="queue" type="String">The name of the queue in which to stop animations.</param>
+    ///   <returns type="jQuery" />
+    /// </signature>
+  },
+  'first': function() {
+    /// <summary>Selects the first matched element.</summary>
+  },
+  'first-child': function() {
+    /// <summary>Selects all elements that are the first child of their parent.</summary>
+  },
+  'first-of-type': function() {
+    /// <summary>Selects all elements that are the first among siblings of the same element name.</summary>
+  },
+  'focus': function() {
+    /// <signature>
+    ///   <summary>Bind an event handler to the "focus" JavaScript event, or trigger that event on an element.</summary>
+    ///   <param name="handler(eventObject)" type="Function">A function to execute each time the event is triggered.</param>
+    ///   <returns type="jQuery" />
+    /// </signature>
+    /// <signature>
+    ///   <summary>Bind an event handler to the "focus" JavaScript event, or trigger that event on an element.</summary>
+    ///   <param name="eventData" type="Object">An object containing data that will be passed to the event handler.</param>
+    ///   <param name="handler(eventObject)" type="Function">A function to execute each time the event is triggered.</param>
+    ///   <returns type="jQuery" />
+    /// </signature>
+  },
+  'focusin': function() {
+    /// <signature>
+    ///   <summary>Bind an event handler to the "focusin" event.</summary>
+    ///   <param name="handler(eventObject)" type="Function">A function to execute each time the event is triggered.</param>
+    ///   <returns type="jQuery" />
+    /// </signature>
+    /// <signature>
+    ///   <summary>Bind an event handler to the "focusin" event.</summary>
+    ///   <param name="eventData" type="Object">An object containing data that will be passed to the event handler.</param>
+    ///   <param name="handler(eventObject)" type="Function">A function to execute each time the event is triggered.</param>
+    ///   <returns type="jQuery" />
+    /// </signature>
+  },
+  'focusout': function() {
+    /// <signature>
+    ///   <summary>Bind an event handler to the "focusout" JavaScript event.</summary>
+    ///   <param name="handler(eventObject)" type="Function">A function to execute each time the event is triggered.</param>
+    ///   <returns type="jQuery" />
+    /// </signature>
+    /// <signature>
+    ///   <summary>Bind an event handler to the "focusout" JavaScript event.</summary>
+    ///   <param name="eventData" type="Object">An object containing data that will be passed to the event handler.</param>
+    ///   <param name="handler(eventObject)" type="Function">A function to execute each time the event is triggered.</param>
+    ///   <returns type="jQuery" />
+    /// </signature>
+  },
+  'get': function() {
+    /// <signature>
+    ///   <summary>Retrieve one of the DOM elements matched by the jQuery object.</summary>
+    ///   <param name="index" type="Number">A zero-based integer indicating which element to retrieve.</param>
+    ///   <returns type="Element" />
+    /// </signature>
+  },
+  'gt': function() {
+    /// <signature>
+    ///   <summary>Select all elements at an index greater than index within the matched set.</summary>
+    ///   <param name="index" type="Number">Zero-based index.</param>
+    /// </signature>
+    /// <signature>
+    ///   <summary>Select all elements at an index greater than index within the matched set.</summary>
+    ///   <param name="-index" type="Number">Zero-based index, counting backwards from the last element.</param>
+    /// </signature>
+  },
+  'has': function() {
+    /// <signature>
+    ///   <summary>Reduce the set of matched elements to those that have a descendant that matches the selector or DOM element.</summary>
+    ///   <param name="selector" type="String">A string containing a selector expression to match elements against.</param>
+    ///   <returns type="jQuery" />
+    /// </signature>
+    /// <signature>
+    ///   <summary>Reduce the set of matched elements to those that have a descendant that matches the selector or DOM element.</summary>
+    ///   <param name="contained" type="Element">A DOM element to match elements against.</param>
+    ///   <returns type="jQuery" />
+    /// </signature>
+  },
+  'hasClass': function() {
+    /// <signature>
+    ///   <summary>Determine whether any of the matched elements are assigned the given class.</summary>
+    ///   <param name="className" type="String">The class name to search for.</param>
+    ///   <returns type="Boolean" />
+    /// </signature>
+  },
+  'header': function() {
+    /// <summary>Selects all elements that are headers, like h1, h2, h3 and so on.</summary>
+  },
+  'height': function() {
+    /// <signature>
+    ///   <summary>Set the CSS height of every matched element.</summary>
+    ///   <param name="value" type="">An integer representing the number of pixels, or an integer with an optional unit of measure appended (as a string).</param>
+    ///   <returns type="jQuery" />
+    /// </signature>
+    /// <signature>
+    ///   <summary>Set the CSS height of every matched element.</summary>
+    ///   <param name="function(index, height)" type="Function">A function returning the height to set. Receives the index position of the element in the set and the old height as arguments. Within the function, this refers to the current element in the set.</param>
+    ///   <returns type="jQuery" />
+    /// </signature>
+  },
+  'hidden': function() {
+    /// <summary>Selects all elements that are hidden.</summary>
+  },
+  'hide': function() {
+    /// <signature>
+    ///   <summary>Hide the matched elements.</summary>
+    ///   <param name="duration" type="">A string or number determining how long the animation will run.</param>
+    ///   <param name="complete" type="Function">A function to call once the animation is complete.</param>
+    ///   <returns type="jQuery" />
+    /// </signature>
+    /// <signature>
+    ///   <summary>Hide the matched elements.</summary>
+    ///   <param name="options" type="PlainObject">A map of additional options to pass to the method.</param>
+    ///   <returns type="jQuery" />
+    /// </signature>
+    /// <signature>
+    ///   <summary>Hide the matched elements.</summary>
+    ///   <param name="duration" type="">A string or number determining how long the animation will run.</param>
+    ///   <param name="easing" type="String">A string indicating which easing function to use for the transition.</param>
+    ///   <param name="complete" type="Function">A function to call once the animation is complete.</param>
+    ///   <returns type="jQuery" />
+    /// </signature>
+  },
+  'hover': function() {
+    /// <signature>
+    ///   <summary>Bind two handlers to the matched elements, to be executed when the mouse pointer enters and leaves the elements.</summary>
+    ///   <param name="handlerIn(eventObject)" type="Function">A function to execute when the mouse pointer enters the element.</param>
+    ///   <param name="handlerOut(eventObject)" type="Function">A function to execute when the mouse pointer leaves the element.</param>
+    ///   <returns type="jQuery" />
+    /// </signature>
+  },
+  'html': function() {
+    /// <signature>
+    ///   <summary>Set the HTML contents of each element in the set of matched elements.</summary>
+    ///   <param name="htmlString" type="htmlString">A string of HTML to set as the content of each matched element.</param>
+    ///   <returns type="jQuery" />
+    /// </signature>
+    /// <signature>
+    ///   <summary>Set the HTML contents of each element in the set of matched elements.</summary>
+    ///   <param name="function(index, oldhtml)" type="Function">A function returning the HTML content to set. Receives the           index position of the element in the set and the old HTML value as arguments.           jQuery empties the element before calling the function;           use the oldhtml argument to reference the previous content.           Within the function, this refers to the current element in the set.</param>
+    ///   <returns type="jQuery" />
+    /// </signature>
+  },
+  'id': function() {
+    /// <signature>
+    ///   <summary>Selects a single element with the given id attribute.</summary>
+    ///   <param name="id" type="String">An ID to search for, specified via the id attribute of an element.</param>
+    /// </signature>
+  },
+  'image': function() {
+    /// <summary>Selects all elements of type image.</summary>
+  },
+  'index': function() {
+    /// <signature>
+    ///   <summary>Search for a given element from among the matched elements.</summary>
+    ///   <param name="selector" type="String">A selector representing a jQuery collection in which to look for an element.</param>
+    ///   <returns type="Number" />
+    /// </signature>
+    /// <signature>
+    ///   <summary>Search for a given element from among the matched elements.</summary>
+    ///   <param name="element" type="">The DOM element or first element within the jQuery object to look for.</param>
+    ///   <returns type="Number" />
+    /// </signature>
+  },
+  'init': function() {
+    /// <signature>
+    ///   <summary>Accepts a string containing a CSS selector which is then used to match a set of elements.</summary>
+    ///   <param name="selector" type="String">A string containing a selector expression</param>
+    ///   <param name="context" type="">A DOM Element, Document, or jQuery to use as context</param>
+    ///   <returns type="jQuery" />
+    /// </signature>
+    /// <signature>
+    ///   <summary>Accepts a string containing a CSS selector which is then used to match a set of elements.</summary>
+    ///   <param name="element" type="Element">A DOM element to wrap in a jQuery object.</param>
+    ///   <returns type="jQuery" />
+    /// </signature>
+    /// <signature>
+    ///   <summary>Accepts a string containing a CSS selector which is then used to match a set of elements.</summary>
+    ///   <param name="elementArray" type="Array">An array containing a set of DOM elements to wrap in a jQuery object.</param>
+    ///   <returns type="jQuery" />
+    /// </signature>
+    /// <signature>
+    ///   <summary>Accepts a string containing a CSS selector which is then used to match a set of elements.</summary>
+    ///   <param name="object" type="PlainObject">A plain object to wrap in a jQuery object.</param>
+    ///   <returns type="jQuery" />
+    /// </signature>
+    /// <signature>
+    ///   <summary>Accepts a string containing a CSS selector which is then used to match a set of elements.</summary>
+    ///   <param name="jQuery object" type="PlainObject">An existing jQuery object to clone.</param>
+    ///   <returns type="jQuery" />
+    /// </signature>
+  },
+  'innerHeight': function() {
+    /// <summary>Get the current computed height for the first element in the set of matched elements, including padding but not border.</summary>
+    /// <returns type="Number" />
+  },
+  'innerWidth': function() {
+    /// <summary>Get the current computed width for the first element in the set of matched elements, including padding but not border.</summary>
+    /// <returns type="Number" />
+  },
+  'input': function() {
+    /// <summary>Selects all input, textarea, select and button elements.</summary>
+  },
+  'insertAfter': function() {
+    /// <signature>
+    ///   <summary>Insert every element in the set of matched elements after the target.</summary>
+    ///   <param name="target" type="">A selector, element, HTML string, or jQuery object; the matched set of elements will be inserted after the element(s) specified by this parameter.</param>
+    ///   <returns type="jQuery" />
+    /// </signature>
+  },
+  'insertBefore': function() {
+    /// <signature>
+    ///   <summary>Insert every element in the set of matched elements before the target.</summary>
+    ///   <param name="target" type="">A selector, element, HTML string, or jQuery object; the matched set of elements will be inserted before the element(s) specified by this parameter.</param>
+    ///   <returns type="jQuery" />
+    /// </signature>
+  },
+  'is': function() {
+    /// <signature>
+    ///   <summary>Check the current matched set of elements against a selector, element, or jQuery object and return true if at least one of these elements matches the given arguments.</summary>
+    ///   <param name="selector" type="String">A string containing a selector expression to match elements against.</param>
+    ///   <returns type="Boolean" />
+    /// </signature>
+    /// <signature>
+    ///   <summary>Check the current matched set of elements against a selector, element, or jQuery object and return true if at least one of these elements matches the given arguments.</summary>
+    ///   <param name="function(index)" type="Function">A function used as a test for the set of elements. It accepts one argument, index, which is the element's index in the jQuery collection.Within the function, this refers to the current DOM element.</param>
+    ///   <returns type="Boolean" />
+    /// </signature>
+    /// <signature>
+    ///   <summary>Check the current matched set of elements against a selector, element, or jQuery object and return true if at least one of these elements matches the given arguments.</summary>
+    ///   <param name="jQuery object" type="Object">An existing jQuery object to match the current set of elements against.</param>
+    ///   <returns type="Boolean" />
+    /// </signature>
+    /// <signature>
+    ///   <summary>Check the current matched set of elements against a selector, element, or jQuery object and return true if at least one of these elements matches the given arguments.</summary>
+    ///   <param name="element" type="Element">An element to match the current set of elements against.</param>
+    ///   <returns type="Boolean" />
+    /// </signature>
+  },
+  'jquery': function() {
+    /// <summary>A string containing the jQuery version number.</summary>
+    /// <returns type="String" />
+  },
+  'keydown': function() {
+    /// <signature>
+    ///   <summary>Bind an event handler to the "keydown" JavaScript event, or trigger that event on an element.</summary>
+    ///   <param name="handler(eventObject)" type="Function">A function to execute each time the event is triggered.</param>
+    ///   <returns type="jQuery" />
+    /// </signature>
+    /// <signature>
+    ///   <summary>Bind an event handler to the "keydown" JavaScript event, or trigger that event on an element.</summary>
+    ///   <param name="eventData" type="PlainObject">An object containing data that will be passed to the event handler.</param>
+    ///   <param name="handler(eventObject)" type="Function">A function to execute each time the event is triggered.</param>
+    ///   <returns type="jQuery" />
+    /// </signature>
+  },
+  'keypress': function() {
+    /// <signature>
+    ///   <summary>Bind an event handler to the "keypress" JavaScript event, or trigger that event on an element.</summary>
+    ///   <param name="handler(eventObject)" type="Function">A function to execute each time the event is triggered.</param>
+    ///   <returns type="jQuery" />
+    /// </signature>
+    /// <signature>
+    ///   <summary>Bind an event handler to the "keypress" JavaScript event, or trigger that event on an element.</summary>
+    ///   <param name="eventData" type="PlainObject">An object containing data that will be passed to the event handler.</param>
+    ///   <param name="handler(eventObject)" type="Function">A function to execute each time the event is triggered.</param>
+    ///   <returns type="jQuery" />
+    /// </signature>
+  },
+  'keyup': function() {
+    /// <signature>
+    ///   <summary>Bind an event handler to the "keyup" JavaScript event, or trigger that event on an element.</summary>
+    ///   <param name="handler(eventObject)" type="Function">A function to execute each time the event is triggered.</param>
+    ///   <returns type="jQuery" />
+    /// </signature>
+    /// <signature>
+    ///   <summary>Bind an event handler to the "keyup" JavaScript event, or trigger that event on an element.</summary>
+    ///   <param name="eventData" type="PlainObject">An object containing data that will be passed to the event handler.</param>
+    ///   <param name="handler(eventObject)" type="Function">A function to execute each time the event is triggered.</param>
+    ///   <returns type="jQuery" />
+    /// </signature>
+  },
+  'lang': function() {
+    /// <signature>
+    ///   <summary>Selects all elements of the specified language.</summary>
+    ///   <param name="language" type="String">A language code.</param>
+    /// </signature>
+  },
+  'last': function() {
+    /// <summary>Selects the last matched element.</summary>
+  },
+  'last-child': function() {
+    /// <summary>Selects all elements that are the last child of their parent.</summary>
+  },
+  'last-of-type': function() {
+    /// <summary>Selects all elements that are the last among siblings of the same element name.</summary>
+  },
+  'length': function() {
+    /// <summary>The number of elements in the jQuery object.</summary>
+    /// <returns type="Number" />
+  },
+  'live': function() {
+    /// <signature>
+    ///   <summary>Attach an event handler for all elements which match the current selector, now and in the future.</summary>
+    ///   <param name="events" type="String">A string containing a JavaScript event type, such as "click" or "keydown." As of jQuery 1.4 the string can contain multiple, space-separated event types or custom event names.</param>
+    ///   <param name="handler(eventObject)" type="Function">A function to execute at the time the event is triggered.</param>
+    ///   <returns type="jQuery" />
+    /// </signature>
+    /// <signature>
+    ///   <summary>Attach an event handler for all elements which match the current selector, now and in the future.</summary>
+    ///   <param name="events" type="String">A string containing a JavaScript event type, such as "click" or "keydown." As of jQuery 1.4 the string can contain multiple, space-separated event types or custom event names.</param>
+    ///   <param name="data" type="PlainObject">An object containing data that will be passed to the event handler.</param>
+    ///   <param name="handler(eventObject)" type="Function">A function to execute at the time the event is triggered.</param>
+    ///   <returns type="jQuery" />
+    /// </signature>
+    /// <signature>
+    ///   <summary>Attach an event handler for all elements which match the current selector, now and in the future.</summary>
+    ///   <param name="events" type="PlainObject">A plain object of one or more JavaScript event types and functions to execute for them.</param>
+    ///   <returns type="jQuery" />
+    /// </signature>
+  },
+  'load': function() {
+    /// <signature>
+    ///   <summary>Bind an event handler to the "load" JavaScript event.</summary>
+    ///   <param name="handler(eventObject)" type="Function">A function to execute when the event is triggered.</param>
+    ///   <returns type="jQuery" />
+    /// </signature>
+    /// <signature>
+    ///   <summary>Bind an event handler to the "load" JavaScript event.</summary>
+    ///   <param name="eventData" type="PlainObject">An object containing data that will be passed to the event handler.</param>
+    ///   <param name="handler(eventObject)" type="Function">A function to execute each time the event is triggered.</param>
+    ///   <returns type="jQuery" />
+    /// </signature>
+  },
+  'lt': function() {
+    /// <signature>
+    ///   <summary>Select all elements at an index less than index within the matched set.</summary>
+    ///   <param name="index" type="Number">Zero-based index.</param>
+    /// </signature>
+    /// <signature>
+    ///   <summary>Select all elements at an index less than index within the matched set.</summary>
+    ///   <param name="-index" type="Number">Zero-based index, counting backwards from the last element.</param>
+    /// </signature>
+  },
+  'map': function() {
+    /// <signature>
+    ///   <summary>Pass each element in the current matched set through a function, producing a new jQuery object containing the return values.</summary>
+    ///   <param name="callback(index, domElement)" type="Function">A function object that will be invoked for each element in the current set.</param>
+    ///   <returns type="jQuery" />
+    /// </signature>
+  },
+  'mousedown': function() {
+    /// <signature>
+    ///   <summary>Bind an event handler to the "mousedown" JavaScript event, or trigger that event on an element.</summary>
+    ///   <param name="handler(eventObject)" type="Function">A function to execute each time the event is triggered.</param>
+    ///   <returns type="jQuery" />
+    /// </signature>
+    /// <signature>
+    ///   <summary>Bind an event handler to the "mousedown" JavaScript event, or trigger that event on an element.</summary>
+    ///   <param name="eventData" type="PlainObject">An object containing data that will be passed to the event handler.</param>
+    ///   <param name="handler(eventObject)" type="Function">A function to execute each time the event is triggered.</param>
+    ///   <returns type="jQuery" />
+    /// </signature>
+  },
+  'mouseenter': function() {
+    /// <signature>
+    ///   <summary>Bind an event handler to be fired when the mouse enters an element, or trigger that handler on an element.</summary>
+    ///   <param name="handler(eventObject)" type="Function">A function to execute each time the event is triggered.</param>
+    ///   <returns type="jQuery" />
+    /// </signature>
+    /// <signature>
+    ///   <summary>Bind an event handler to be fired when the mouse enters an element, or trigger that handler on an element.</summary>
+    ///   <param name="eventData" type="PlainObject">An object containing data that will be passed to the event handler.</param>
+    ///   <param name="handler(eventObject)" type="Function">A function to execute each time the event is triggered.</param>
+    ///   <returns type="jQuery" />
+    /// </signature>
+  },
+  'mouseleave': function() {
+    /// <signature>
+    ///   <summary>Bind an event handler to be fired when the mouse leaves an element, or trigger that handler on an element.</summary>
+    ///   <param name="handler(eventObject)" type="Function">A function to execute each time the event is triggered.</param>
+    ///   <returns type="jQuery" />
+    /// </signature>
+    /// <signature>
+    ///   <summary>Bind an event handler to be fired when the mouse leaves an element, or trigger that handler on an element.</summary>
+    ///   <param name="eventData" type="PlainObject">An object containing data that will be passed to the event handler.</param>
+    ///   <param name="handler(eventObject)" type="Function">A function to execute each time the event is triggered.</param>
+    ///   <returns type="jQuery" />
+    /// </signature>
+  },
+  'mousemove': function() {
+    /// <signature>
+    ///   <summary>Bind an event handler to the "mousemove" JavaScript event, or trigger that event on an element.</summary>
+    ///   <param name="handler(eventObject)" type="Function">A function to execute each time the event is triggered.</param>
+    ///   <returns type="jQuery" />
+    /// </signature>
+    /// <signature>
+    ///   <summary>Bind an event handler to the "mousemove" JavaScript event, or trigger that event on an element.</summary>
+    ///   <param name="eventData" type="PlainObject">An object containing data that will be passed to the event handler.</param>
+    ///   <param name="handler(eventObject)" type="Function">A function to execute each time the event is triggered.</param>
+    ///   <returns type="jQuery" />
+    /// </signature>
+  },
+  'mouseout': function() {
+    /// <signature>
+    ///   <summary>Bind an event handler to the "mouseout" JavaScript event, or trigger that event on an element.</summary>
+    ///   <param name="handler(eventObject)" type="Function">A function to execute each time the event is triggered.</param>
+    ///   <returns type="jQuery" />
+    /// </signature>
+    /// <signature>
+    ///   <summary>Bind an event handler to the "mouseout" JavaScript event, or trigger that event on an element.</summary>
+    ///   <param name="eventData" type="PlainObject">An object containing data that will be passed to the event handler.</param>
+    ///   <param name="handler(eventObject)" type="Function">A function to execute each time the event is triggered.</param>
+    ///   <returns type="jQuery" />
+    /// </signature>
+  },
+  'mouseover': function() {
+    /// <signature>
+    ///   <summary>Bind an event handler to the "mouseover" JavaScript event, or trigger that event on an element.</summary>
+    ///   <param name="handler(eventObject)" type="Function">A function to execute each time the event is triggered.</param>
+    ///   <returns type="jQuery" />
+    /// </signature>
+    /// <signature>
+    ///   <summary>Bind an event handler to the "mouseover" JavaScript event, or trigger that event on an element.</summary>
+    ///   <param name="eventData" type="PlainObject">An object containing data that will be passed to the event handler.</param>
+    ///   <param name="handler(eventObject)" type="Function">A function to execute each time the event is triggered.</param>
+    ///   <returns type="jQuery" />
+    /// </signature>
+  },
+  'mouseup': function() {
+    /// <signature>
+    ///   <summary>Bind an event handler to the "mouseup" JavaScript event, or trigger that event on an element.</summary>
+    ///   <param name="handler(eventObject)" type="Function">A function to execute each time the event is triggered.</param>
+    ///   <returns type="jQuery" />
+    /// </signature>
+    /// <signature>
+    ///   <summary>Bind an event handler to the "mouseup" JavaScript event, or trigger that event on an element.</summary>
+    ///   <param name="eventData" type="PlainObject">An object containing data that will be passed to the event handler.</param>
+    ///   <param name="handler(eventObject)" type="Function">A function to execute each time the event is triggered.</param>
+    ///   <returns type="jQuery" />
+    /// </signature>
+  },
+  'multiple': function() {
+    /// <signature>
+    ///   <summary>Selects the combined results of all the specified selectors.</summary>
+    ///   <param name="selector1" type="String">Any valid selector.</param>
+    ///   <param name="selector2" type="String">Another valid selector.</param>
+    ///   <param name="selectorN" type="String">As many more valid selectors as you like.</param>
+    /// </signature>
+  },
+  'next': function() {
+    /// <signature>
+    ///   <summary>Get the immediately following sibling of each element in the set of matched elements. If a selector is provided, it retrieves the next sibling only if it matches that selector.</summary>
+    ///   <param name="selector" type="String">A string containing a selector expression to match elements against.</param>
+    ///   <returns type="jQuery" />
+    /// </signature>
+  },
+  'next adjacent': function() {
+    /// <signature>
+    ///   <summary>Selects all next elements matching "next" that are immediately preceded by a sibling "prev".</summary>
+    ///   <param name="prev" type="String">Any valid selector.</param>
+    ///   <param name="next" type="String">A selector to match the element that is next to the first selector.</param>
+    /// </signature>
+  },
+  'next siblings': function() {
+    /// <signature>
+    ///   <summary>Selects all sibling elements that follow after the "prev" element, have the same parent, and match the filtering "siblings" selector.</summary>
+    ///   <param name="prev" type="String">Any valid selector.</param>
+    ///   <param name="siblings" type="String">A selector to filter elements that are the following siblings of the first selector.</param>
+    /// </signature>
+  },
+  'nextAll': function() {
+    /// <signature>
+    ///   <summary>Get all following siblings of each element in the set of matched elements, optionally filtered by a selector.</summary>
+    ///   <param name="selector" type="String">A string containing a selector expression to match elements against.</param>
+    ///   <returns type="jQuery" />
+    /// </signature>
+  },
+  'nextUntil': function() {
+    /// <signature>
+    ///   <summary>Get all following siblings of each element up to but not including the element matched by the selector, DOM node, or jQuery object passed.</summary>
+    ///   <param name="selector" type="String">A string containing a selector expression to indicate where to stop matching following sibling elements.</param>
+    ///   <param name="filter" type="String">A string containing a selector expression to match elements against.</param>
+    ///   <returns type="jQuery" />
+    /// </signature>
+    /// <signature>
+    ///   <summary>Get all following siblings of each element up to but not including the element matched by the selector, DOM node, or jQuery object passed.</summary>
+    ///   <param name="element" type="Element">A DOM node or jQuery object indicating where to stop matching following sibling elements.</param>
+    ///   <param name="filter" type="String">A string containing a selector expression to match elements against.</param>
+    ///   <returns type="jQuery" />
+    /// </signature>
+  },
+  'not': function() {
+    /// <signature>
+    ///   <summary>Remove elements from the set of matched elements.</summary>
+    ///   <param name="selector" type="String">A string containing a selector expression to match elements against.</param>
+    ///   <returns type="jQuery" />
+    /// </signature>
+    /// <signature>
+    ///   <summary>Remove elements from the set of matched elements.</summary>
+    ///   <param name="elements" type="Array">One or more DOM elements to remove from the matched set.</param>
+    ///   <returns type="jQuery" />
+    /// </signature>
+    /// <signature>
+    ///   <summary>Remove elements from the set of matched elements.</summary>
+    ///   <param name="function(index)" type="Function">A function used as a test for each element in the set. this is the current DOM element.</param>
+    ///   <returns type="jQuery" />
+    /// </signature>
+    /// <signature>
+    ///   <summary>Remove elements from the set of matched elements.</summary>
+    ///   <param name="jQuery object" type="PlainObject">An existing jQuery object to match the current set of elements against.</param>
+    ///   <returns type="jQuery" />
+    /// </signature>
+  },
+  'nth-child': function() {
+    /// <signature>
+    ///   <summary>Selects all elements that are the nth-child of their parent.</summary>
+    ///   <param name="index" type="String">The index of each child to match, starting with 1, the string even or odd, or an equation ( eg. :nth-child(even), :nth-child(4n) )</param>
+    /// </signature>
+  },
+  'nth-last-child': function() {
+    /// <signature>
+    ///   <summary>Selects all elements that are the nth-child of their parent, counting from the last element to the first.</summary>
+    ///   <param name="index" type="String">The index of each child to match, starting with the last one (1), the string even or odd, or an equation ( eg. :nth-last-child(even), :nth-last-child(4n) )</param>
+    /// </signature>
+  },
+  'nth-last-of-type': function() {
+    /// <signature>
+    ///   <summary>Selects all elements that are the nth-child of their parent, counting from the last element to the first.</summary>
+    ///   <param name="index" type="String">The index of each child to match, starting with the last one (1), the string even or odd, or an equation ( eg. :nth-last-of-type(even), :nth-last-of-type(4n) )</param>
+    /// </signature>
+  },
+  'nth-of-type': function() {
+    /// <signature>
+    ///   <summary>Selects all elements that are the nth child of their parent in relation to siblings with the same element name.</summary>
+    ///   <param name="index" type="String">The index of each child to match, starting with 1, the string even or odd, or an equation ( eg. :nth-of-type(even), :nth-of-type(4n) )</param>
+    /// </signature>
+  },
+  'odd': function() {
+    /// <summary>Selects odd elements, zero-indexed.  See also even.</summary>
+  },
+  'off': function() {
+    /// <signature>
+    ///   <summary>Remove an event handler.</summary>
+    ///   <param name="events" type="String">One or more space-separated event types and optional namespaces, or just namespaces, such as "click", "keydown.myPlugin", or ".myPlugin".</param>
+    ///   <param name="selector" type="String">A selector which should match the one originally passed to .on() when attaching event handlers.</param>
+    ///   <param name="handler(eventObject)" type="Function">A handler function previously attached for the event(s), or the special value false.</param>
+    ///   <returns type="jQuery" />
+    /// </signature>
+    /// <signature>
+    ///   <summary>Remove an event handler.</summary>
+    ///   <param name="events" type="PlainObject">An object where the string keys represent one or more space-separated event types and optional namespaces, and the values represent handler functions previously attached for the event(s).</param>
+    ///   <param name="selector" type="String">A selector which should match the one originally passed to .on() when attaching event handlers.</param>
+    ///   <returns type="jQuery" />
+    /// </signature>
+  },
+  'offset': function() {
+    /// <signature>
+    ///   <summary>Set the current coordinates of every element in the set of matched elements, relative to the document.</summary>
+    ///   <param name="coordinates" type="PlainObject">An object containing the properties top and left, which are integers indicating the new top and left coordinates for the elements.</param>
+    ///   <returns type="jQuery" />
+    /// </signature>
+    /// <signature>
+    ///   <summary>Set the current coordinates of every element in the set of matched elements, relative to the document.</summary>
+    ///   <param name="function(index, coords)" type="Function">A function to return the coordinates to set. Receives the index of the element in the collection as the first argument and the current coordinates as the second argument. The function should return an object with the new top and left properties.</param>
+    ///   <returns type="jQuery" />
+    /// </signature>
+  },
+  'offsetParent': function() {
+    /// <summary>Get the closest ancestor element that is positioned.</summary>
+    /// <returns type="jQuery" />
+  },
+  'on': function() {
+    /// <signature>
+    ///   <summary>Attach an event handler function for one or more events to the selected elements.</summary>
+    ///   <param name="events" type="String">One or more space-separated event types and optional namespaces, such as "click" or "keydown.myPlugin".</param>
+    ///   <param name="selector" type="String">A selector string to filter the descendants of the selected elements that trigger the event. If the selector is null or omitted, the event is always triggered when it reaches the selected element.</param>
+    ///   <param name="data" type="Anything">Data to be passed to the handler in event.data when an event is triggered.</param>
+    ///   <param name="handler(eventObject)" type="Function">A function to execute when the event is triggered. The value false is also allowed as a shorthand for a function that simply does return false.</param>
+    ///   <returns type="jQuery" />
+    /// </signature>
+    /// <signature>
+    ///   <summary>Attach an event handler function for one or more events to the selected elements.</summary>
+    ///   <param name="events" type="PlainObject">An object in which the string keys represent one or more space-separated event types and optional namespaces, and the values represent a handler function to be called for the event(s).</param>
+    ///   <param name="selector" type="String">A selector string to filter the descendants of the selected elements that will call the handler. If the selector is null or omitted, the handler is always called when it reaches the selected element.</param>
+    ///   <param name="data" type="Anything">Data to be passed to the handler in event.data when an event occurs.</param>
+    ///   <returns type="jQuery" />
+    /// </signature>
+  },
+  'one': function() {
+    /// <signature>
+    ///   <summary>Attach a handler to an event for the elements. The handler is executed at most once per element.</summary>
+    ///   <param name="events" type="String">A string containing one or more JavaScript event types, such as "click" or "submit," or custom event names.</param>
+    ///   <param name="data" type="PlainObject">An object containing data that will be passed to the event handler.</param>
+    ///   <param name="handler(eventObject)" type="Function">A function to execute at the time the event is triggered.</param>
+    ///   <returns type="jQuery" />
+    /// </signature>
+    /// <signature>
+    ///   <summary>Attach a handler to an event for the elements. The handler is executed at most once per element.</summary>
+    ///   <param name="events" type="String">One or more space-separated event types and optional namespaces, such as "click" or "keydown.myPlugin".</param>
+    ///   <param name="selector" type="String">A selector string to filter the descendants of the selected elements that trigger the event. If the selector is null or omitted, the event is always triggered when it reaches the selected element.</param>
+    ///   <param name="data" type="Anything">Data to be passed to the handler in event.data when an event is triggered.</param>
+    ///   <param name="handler(eventObject)" type="Function">A function to execute when the event is triggered. The value false is also allowed as a shorthand for a function that simply does return false.</param>
+    ///   <returns type="jQuery" />
+    /// </signature>
+    /// <signature>
+    ///   <summary>Attach a handler to an event for the elements. The handler is executed at most once per element.</summary>
+    ///   <param name="events" type="PlainObject">An object in which the string keys represent one or more space-separated event types and optional namespaces, and the values represent a handler function to be called for the event(s).</param>
+    ///   <param name="selector" type="String">A selector string to filter the descendants of the selected elements that will call the handler. If the selector is null or omitted, the handler is always called when it reaches the selected element.</param>
+    ///   <param name="data" type="Anything">Data to be passed to the handler in event.data when an event occurs.</param>
+    ///   <returns type="jQuery" />
+    /// </signature>
+  },
+  'only-child': function() {
+    /// <summary>Selects all elements that are the only child of their parent.</summary>
+  },
+  'only-of-type': function() {
+    /// <summary>Selects all elements that have no siblings with the same element name.</summary>
+  },
+  'outerHeight': function() {
+    /// <signature>
+    ///   <summary>Get the current computed height for the first element in the set of matched elements, including padding, border, and optionally margin. Returns an integer (without "px") representation of the value or null if called on an empty set of elements.</summary>
+    ///   <param name="includeMargin" type="Boolean">A Boolean indicating whether to include the element's margin in the calculation.</param>
+    ///   <returns type="Number" />
+    /// </signature>
+  },
+  'outerWidth': function() {
+    /// <signature>
+    ///   <summary>Get the current computed width for the first element in the set of matched elements, including padding and border.</summary>
+    ///   <param name="includeMargin" type="Boolean">A Boolean indicating whether to include the element's margin in the calculation.</param>
+    ///   <returns type="Number" />
+    /// </signature>
+  },
+  'parent': function() {
+    /// <signature>
+    ///   <summary>Get the parent of each element in the current set of matched elements, optionally filtered by a selector.</summary>
+    ///   <param name="selector" type="String">A string containing a selector expression to match elements against.</param>
+    ///   <returns type="jQuery" />
+    /// </signature>
+  },
+  'parents': function() {
+    /// <signature>
+    ///   <summary>Get the ancestors of each element in the current set of matched elements, optionally filtered by a selector.</summary>
+    ///   <param name="selector" type="String">A string containing a selector expression to match elements against.</param>
+    ///   <returns type="jQuery" />
+    /// </signature>
+  },
+  'parentsUntil': function() {
+    /// <signature>
+    ///   <summary>Get the ancestors of each element in the current set of matched elements, up to but not including the element matched by the selector, DOM node, or jQuery object.</summary>
+    ///   <param name="selector" type="String">A string containing a selector expression to indicate where to stop matching ancestor elements.</param>
+    ///   <param name="filter" type="String">A string containing a selector expression to match elements against.</param>
+    ///   <returns type="jQuery" />
+    /// </signature>
+    /// <signature>
+    ///   <summary>Get the ancestors of each element in the current set of matched elements, up to but not including the element matched by the selector, DOM node, or jQuery object.</summary>
+    ///   <param name="element" type="Element">A DOM node or jQuery object indicating where to stop matching ancestor elements.</param>
+    ///   <param name="filter" type="String">A string containing a selector expression to match elements against.</param>
+    ///   <returns type="jQuery" />
+    /// </signature>
+  },
+  'password': function() {
+    /// <summary>Selects all elements of type password.</summary>
+  },
+  'position': function() {
+    /// <summary>Get the current coordinates of the first element in the set of matched elements, relative to the offset parent.</summary>
+    /// <returns type="Object" />
+  },
+  'prepend': function() {
+    /// <signature>
+    ///   <summary>Insert content, specified by the parameter, to the beginning of each element in the set of matched elements.</summary>
+    ///   <param name="content" type="">DOM element, array of elements, HTML string, or jQuery object to insert at the beginning of each element in the set of matched elements.</param>
+    ///   <param name="content" type="">One or more additional DOM elements, arrays of elements, HTML strings, or jQuery objects to insert at the beginning of each element in the set of matched elements.</param>
+    ///   <returns type="jQuery" />
+    /// </signature>
+    /// <signature>
+    ///   <summary>Insert content, specified by the parameter, to the beginning of each element in the set of matched elements.</summary>
+    ///   <param name="function(index, html)" type="Function">A function that returns an HTML string, DOM element(s), or jQuery object to insert at the beginning of each element in the set of matched elements. Receives the index position of the element in the set and the old HTML value of the element as arguments. Within the function, this refers to the current element in the set.</param>
+    ///   <returns type="jQuery" />
+    /// </signature>
+  },
+  'prependTo': function() {
+    /// <signature>
+    ///   <summary>Insert every element in the set of matched elements to the beginning of the target.</summary>
+    ///   <param name="target" type="">A selector, element, HTML string, or jQuery object; the matched set of elements will be inserted at the beginning of the element(s) specified by this parameter.</param>
+    ///   <returns type="jQuery" />
+    /// </signature>
+  },
+  'prev': function() {
+    /// <signature>
+    ///   <summary>Get the immediately preceding sibling of each element in the set of matched elements, optionally filtered by a selector.</summary>
+    ///   <param name="selector" type="String">A string containing a selector expression to match elements against.</param>
+    ///   <returns type="jQuery" />
+    /// </signature>
+  },
+  'prevAll': function() {
+    /// <signature>
+    ///   <summary>Get all preceding siblings of each element in the set of matched elements, optionally filtered by a selector.</summary>
+    ///   <param name="selector" type="String">A string containing a selector expression to match elements against.</param>
+    ///   <returns type="jQuery" />
+    /// </signature>
+  },
+  'prevUntil': function() {
+    /// <signature>
+    ///   <summary>Get all preceding siblings of each element up to but not including the element matched by the selector, DOM node, or jQuery object.</summary>
+    ///   <param name="selector" type="String">A string containing a selector expression to indicate where to stop matching preceding sibling elements.</param>
+    ///   <param name="filter" type="String">A string containing a selector expression to match elements against.</param>
+    ///   <returns type="jQuery" />
+    /// </signature>
+    /// <signature>
+    ///   <summary>Get all preceding siblings of each element up to but not including the element matched by the selector, DOM node, or jQuery object.</summary>
+    ///   <param name="element" type="Element">A DOM node or jQuery object indicating where to stop matching preceding sibling elements.</param>
+    ///   <param name="filter" type="String">A string containing a selector expression to match elements against.</param>
+    ///   <returns type="jQuery" />
+    /// </signature>
+  },
+  'promise': function() {
+    /// <signature>
+    ///   <summary>Return a Promise object to observe when all actions of a certain type bound to the collection, queued or not, have finished.</summary>
+    ///   <param name="type" type="String">The type of queue that needs to be observed.</param>
+    ///   <param name="target" type="PlainObject">Object onto which the promise methods have to be attached</param>
+    ///   <returns type="Promise" />
+    /// </signature>
+  },
+  'prop': function() {
+    /// <signature>
+    ///   <summary>Set one or more properties for the set of matched elements.</summary>
+    ///   <param name="propertyName" type="String">The name of the property to set.</param>
+    ///   <param name="value" type="">A value to set for the property.</param>
+    ///   <returns type="jQuery" />
+    /// </signature>
+    /// <signature>
+    ///   <summary>Set one or more properties for the set of matched elements.</summary>
+    ///   <param name="properties" type="PlainObject">An object of property-value pairs to set.</param>
+    ///   <returns type="jQuery" />
+    /// </signature>
+    /// <signature>
+    ///   <summary>Set one or more properties for the set of matched elements.</summary>
+    ///   <param name="propertyName" type="String">The name of the property to set.</param>
+    ///   <param name="function(index, oldPropertyValue)" type="Function">A function returning the value to set. Receives the index position of the element in the set and the old property value as arguments. Within the function, the keyword this refers to the current element.</param>
+    ///   <returns type="jQuery" />
+    /// </signature>
+  },
+  'pushStack': function() {
+    /// <signature>
+    ///   <summary>Add a collection of DOM elements onto the jQuery stack.</summary>
+    ///   <param name="elements" type="Array">An array of elements to push onto the stack and make into a new jQuery object.</param>
+    ///   <returns type="jQuery" />
+    /// </signature>
+    /// <signature>
+    ///   <summary>Add a collection of DOM elements onto the jQuery stack.</summary>
+    ///   <param name="elements" type="Array">An array of elements to push onto the stack and make into a new jQuery object.</param>
+    ///   <param name="name" type="String">The name of a jQuery method that generated the array of elements.</param>
+    ///   <param name="arguments" type="Array">The arguments that were passed in to the jQuery method (for serialization).</param>
+    ///   <returns type="jQuery" />
+    /// </signature>
+  },
+  'queue': function() {
+    /// <signature>
+    ///   <summary>Manipulate the queue of functions to be executed, once for each matched element.</summary>
+    ///   <param name="queueName" type="String">A string containing the name of the queue. Defaults to fx, the standard effects queue.</param>
+    ///   <param name="newQueue" type="Array">An array of functions to replace the current queue contents.</param>
+    ///   <returns type="jQuery" />
+    /// </signature>
+    /// <signature>
+    ///   <summary>Manipulate the queue of functions to be executed, once for each matched element.</summary>
+    ///   <param name="queueName" type="String">A string containing the name of the queue. Defaults to fx, the standard effects queue.</param>
+    ///   <param name="callback( next )" type="Function">The new function to add to the queue, with a function to call that will dequeue the next item.</param>
+    ///   <returns type="jQuery" />
+    /// </signature>
+  },
+  'radio': function() {
+    /// <summary>Selects all  elements of type radio.</summary>
+  },
+  'ready': function() {
+    /// <signature>
+    ///   <summary>Specify a function to execute when the DOM is fully loaded.</summary>
+    ///   <param name="handler" type="Function">A function to execute after the DOM is ready.</param>
+    ///   <returns type="jQuery" />
+    /// </signature>
+  },
+  'remove': function() {
+    /// <signature>
+    ///   <summary>Remove the set of matched elements from the DOM.</summary>
+    ///   <param name="selector" type="String">A selector expression that filters the set of matched elements to be removed.</param>
+    ///   <returns type="jQuery" />
+    /// </signature>
+  },
+  'removeAttr': function() {
+    /// <signature>
+    ///   <summary>Remove an attribute from each element in the set of matched elements.</summary>
+    ///   <param name="attributeName" type="String">An attribute to remove; as of version 1.7, it can be a space-separated list of attributes.</param>
+    ///   <returns type="jQuery" />
+    /// </signature>
+  },
+  'removeClass': function() {
+    /// <signature>
+    ///   <summary>Remove a single class, multiple classes, or all classes from each element in the set of matched elements.</summary>
+    ///   <param name="className" type="String">One or more space-separated classes to be removed from the class attribute of each matched element.</param>
+    ///   <returns type="jQuery" />
+    /// </signature>
+    /// <signature>
+    ///   <summary>Remove a single class, multiple classes, or all classes from each element in the set of matched elements.</summary>
+    ///   <param name="function(index, class)" type="Function">A function returning one or more space-separated class names to be removed. Receives the index position of the element in the set and the old class value as arguments.</param>
+    ///   <returns type="jQuery" />
+    /// </signature>
+  },
+  'removeData': function() {
+    /// <signature>
+    ///   <summary>Remove a previously-stored piece of data.</summary>
+    ///   <param name="name" type="String">A string naming the piece of data to delete.</param>
+    ///   <returns type="jQuery" />
+    /// </signature>
+    /// <signature>
+    ///   <summary>Remove a previously-stored piece of data.</summary>
+    ///   <param name="list" type="">An array or space-separated string naming the pieces of data to delete.</param>
+    ///   <returns type="jQuery" />
+    /// </signature>
+  },
+  'removeProp': function() {
+    /// <signature>
+    ///   <summary>Remove a property for the set of matched elements.</summary>
+    ///   <param name="propertyName" type="String">The name of the property to remove.</param>
+    ///   <returns type="jQuery" />
+    /// </signature>
+  },
+  'replaceAll': function() {
+    /// <signature>
+    ///   <summary>Replace each target element with the set of matched elements.</summary>
+    ///   <param name="target" type="">A selector string, jQuery object, or DOM element reference indicating which element(s) to replace.</param>
+    ///   <returns type="jQuery" />
+    /// </signature>
+  },
+  'replaceWith': function() {
+    /// <signature>
+    ///   <summary>Replace each element in the set of matched elements with the provided new content and return the set of elements that was removed.</summary>
+    ///   <param name="newContent" type="">The content to insert. May be an HTML string, DOM element, or jQuery object.</param>
+    ///   <returns type="jQuery" />
+    /// </signature>
+    /// <signature>
+    ///   <summary>Replace each element in the set of matched elements with the provided new content and return the set of elements that was removed.</summary>
+    ///   <param name="function" type="Function">A function that returns content with which to replace the set of matched elements.</param>
+    ///   <returns type="jQuery" />
+    /// </signature>
+  },
+  'reset': function() {
+    /// <summary>Selects all elements of type reset.</summary>
+  },
+  'resize': function() {
+    /// <signature>
+    ///   <summary>Bind an event handler to the "resize" JavaScript event, or trigger that event on an element.</summary>
+    ///   <param name="handler(eventObject)" type="Function">A function to execute each time the event is triggered.</param>
+    ///   <returns type="jQuery" />
+    /// </signature>
+    /// <signature>
+    ///   <summary>Bind an event handler to the "resize" JavaScript event, or trigger that event on an element.</summary>
+    ///   <param name="eventData" type="PlainObject">An object containing data that will be passed to the event handler.</param>
+    ///   <param name="handler(eventObject)" type="Function">A function to execute each time the event is triggered.</param>
+    ///   <returns type="jQuery" />
+    /// </signature>
+  },
+  'root': function() {
+    /// <summary>Selects the element that is the root of the document.</summary>
+  },
+  'scroll': function() {
+    /// <signature>
+    ///   <summary>Bind an event handler to the "scroll" JavaScript event, or trigger that event on an element.</summary>
+    ///   <param name="handler(eventObject)" type="Function">A function to execute each time the event is triggered.</param>
+    ///   <returns type="jQuery" />
+    /// </signature>
+    /// <signature>
+    ///   <summary>Bind an event handler to the "scroll" JavaScript event, or trigger that event on an element.</summary>
+    ///   <param name="eventData" type="PlainObject">An object containing data that will be passed to the event handler.</param>
+    ///   <param name="handler(eventObject)" type="Function">A function to execute each time the event is triggered.</param>
+    ///   <returns type="jQuery" />
+    /// </signature>
+  },
+  'scrollLeft': function() {
+    /// <signature>
+    ///   <summary>Set the current horizontal position of the scroll bar for each of the set of matched elements.</summary>
+    ///   <param name="value" type="Number">An integer indicating the new position to set the scroll bar to.</param>
+    ///   <returns type="jQuery" />
+    /// </signature>
+  },
+  'scrollTop': function() {
+    /// <signature>
+    ///   <summary>Set the current vertical position of the scroll bar for each of the set of matched elements.</summary>
+    ///   <param name="value" type="Number">An integer indicating the new position to set the scroll bar to.</param>
+    ///   <returns type="jQuery" />
+    /// </signature>
+  },
+  'select': function() {
+    /// <signature>
+    ///   <summary>Bind an event handler to the "select" JavaScript event, or trigger that event on an element.</summary>
+    ///   <param name="handler(eventObject)" type="Function">A function to execute each time the event is triggered.</param>
+    ///   <returns type="jQuery" />
+    /// </signature>
+    /// <signature>
+    ///   <summary>Bind an event handler to the "select" JavaScript event, or trigger that event on an element.</summary>
+    ///   <param name="eventData" type="PlainObject">An object containing data that will be passed to the event handler.</param>
+    ///   <param name="handler(eventObject)" type="Function">A function to execute each time the event is triggered.</param>
+    ///   <returns type="jQuery" />
+    /// </signature>
+  },
+  'selected': function() {
+    /// <summary>Selects all elements that are selected.</summary>
+  },
+  'selector': function() {
+    /// <summary>A selector representing selector passed to jQuery(), if any, when creating the original set.</summary>
+    /// <returns type="String" />
+  },
+  'serialize': function() {
+    /// <summary>Encode a set of form elements as a string for submission.</summary>
+    /// <returns type="String" />
+  },
+  'serializeArray': function() {
+    /// <summary>Encode a set of form elements as an array of names and values.</summary>
+    /// <returns type="Array" />
+  },
+  'show': function() {
+    /// <signature>
+    ///   <summary>Display the matched elements.</summary>
+    ///   <param name="duration" type="">A string or number determining how long the animation will run.</param>
+    ///   <param name="complete" type="Function">A function to call once the animation is complete.</param>
+    ///   <returns type="jQuery" />
+    /// </signature>
+    /// <signature>
+    ///   <summary>Display the matched elements.</summary>
+    ///   <param name="options" type="PlainObject">A map of additional options to pass to the method.</param>
+    ///   <returns type="jQuery" />
+    /// </signature>
+    /// <signature>
+    ///   <summary>Display the matched elements.</summary>
+    ///   <param name="duration" type="">A string or number determining how long the animation will run.</param>
+    ///   <param name="easing" type="String">A string indicating which easing function to use for the transition.</param>
+    ///   <param name="complete" type="Function">A function to call once the animation is complete.</param>
+    ///   <returns type="jQuery" />
+    /// </signature>
+  },
+  'siblings': function() {
+    /// <signature>
+    ///   <summary>Get the siblings of each element in the set of matched elements, optionally filtered by a selector.</summary>
+    ///   <param name="selector" type="String">A string containing a selector expression to match elements against.</param>
+    ///   <returns type="jQuery" />
+    /// </signature>
+  },
+  'size': function() {
+    /// <summary>Return the number of elements in the jQuery object.</summary>
+    /// <returns type="Number" />
+  },
+  'slice': function() {
+    /// <signature>
+    ///   <summary>Reduce the set of matched elements to a subset specified by a range of indices.</summary>
+    ///   <param name="start" type="Number">An integer indicating the 0-based position at which the elements begin to be selected. If negative, it indicates an offset from the end of the set.</param>
+    ///   <param name="end" type="Number">An integer indicating the 0-based position at which the elements stop being selected. If negative, it indicates an offset from the end of the set. If omitted, the range continues until the end of the set.</param>
+    ///   <returns type="jQuery" />
+    /// </signature>
+  },
+  'slideDown': function() {
+    /// <signature>
+    ///   <summary>Display the matched elements with a sliding motion.</summary>
+    ///   <param name="duration" type="">A string or number determining how long the animation will run.</param>
+    ///   <param name="complete" type="Function">A function to call once the animation is complete.</param>
+    ///   <returns type="jQuery" />
+    /// </signature>
+    /// <signature>
+    ///   <summary>Display the matched elements with a sliding motion.</summary>
+    ///   <param name="options" type="PlainObject">A map of additional options to pass to the method.</param>
+    ///   <returns type="jQuery" />
+    /// </signature>
+    /// <signature>
+    ///   <summary>Display the matched elements with a sliding motion.</summary>
+    ///   <param name="duration" type="">A string or number determining how long the animation will run.</param>
+    ///   <param name="easing" type="String">A string indicating which easing function to use for the transition.</param>
+    ///   <param name="complete" type="Function">A function to call once the animation is complete.</param>
+    ///   <returns type="jQuery" />
+    /// </signature>
+  },
+  'slideToggle': function() {
+    /// <signature>
+    ///   <summary>Display or hide the matched elements with a sliding motion.</summary>
+    ///   <param name="duration" type="">A string or number determining how long the animation will run.</param>
+    ///   <param name="complete" type="Function">A function to call once the animation is complete.</param>
+    ///   <returns type="jQuery" />
+    /// </signature>
+    /// <signature>
+    ///   <summary>Display or hide the matched elements with a sliding motion.</summary>
+    ///   <param name="options" type="PlainObject">A map of additional options to pass to the method.</param>
+    ///   <returns type="jQuery" />
+    /// </signature>
+    /// <signature>
+    ///   <summary>Display or hide the matched elements with a sliding motion.</summary>
+    ///   <param name="duration" type="">A string or number determining how long the animation will run.</param>
+    ///   <param name="easing" type="String">A string indicating which easing function to use for the transition.</param>
+    ///   <param name="complete" type="Function">A function to call once the animation is complete.</param>
+    ///   <returns type="jQuery" />
+    /// </signature>
+  },
+  'slideUp': function() {
+    /// <signature>
+    ///   <summary>Hide the matched elements with a sliding motion.</summary>
+    ///   <param name="duration" type="">A string or number determining how long the animation will run.</param>
+    ///   <param name="complete" type="Function">A function to call once the animation is complete.</param>
+    ///   <returns type="jQuery" />
+    /// </signature>
+    /// <signature>
+    ///   <summary>Hide the matched elements with a sliding motion.</summary>
+    ///   <param name="options" type="PlainObject">A map of additional options to pass to the method.</param>
+    ///   <returns type="jQuery" />
+    /// </signature>
+    /// <signature>
+    ///   <summary>Hide the matched elements with a sliding motion.</summary>
+    ///   <param name="duration" type="">A string or number determining how long the animation will run.</param>
+    ///   <param name="easing" type="String">A string indicating which easing function to use for the transition.</param>
+    ///   <param name="complete" type="Function">A function to call once the animation is complete.</param>
+    ///   <returns type="jQuery" />
+    /// </signature>
+  },
+  'stop': function() {
+    /// <signature>
+    ///   <summary>Stop the currently-running animation on the matched elements.</summary>
+    ///   <param name="clearQueue" type="Boolean">A Boolean indicating whether to remove queued animation as well. Defaults to false.</param>
+    ///   <param name="jumpToEnd" type="Boolean">A Boolean indicating whether to complete the current animation immediately. Defaults to false.</param>
+    ///   <returns type="jQuery" />
+    /// </signature>
+    /// <signature>
+    ///   <summary>Stop the currently-running animation on the matched elements.</summary>
+    ///   <param name="queue" type="String">The name of the queue in which to stop animations.</param>
+    ///   <param name="clearQueue" type="Boolean">A Boolean indicating whether to remove queued animation as well. Defaults to false.</param>
+    ///   <param name="jumpToEnd" type="Boolean">A Boolean indicating whether to complete the current animation immediately. Defaults to false.</param>
+    ///   <returns type="jQuery" />
+    /// </signature>
+  },
+  'submit': function() {
+    /// <signature>
+    ///   <summary>Bind an event handler to the "submit" JavaScript event, or trigger that event on an element.</summary>
+    ///   <param name="handler(eventObject)" type="Function">A function to execute each time the event is triggered.</param>
+    ///   <returns type="jQuery" />
+    /// </signature>
+    /// <signature>
+    ///   <summary>Bind an event handler to the "submit" JavaScript event, or trigger that event on an element.</summary>
+    ///   <param name="eventData" type="PlainObject">An object containing data that will be passed to the event handler.</param>
+    ///   <param name="handler(eventObject)" type="Function">A function to execute each time the event is triggered.</param>
+    ///   <returns type="jQuery" />
+    /// </signature>
+  },
+  'target': function() {
+    /// <summary>Selects the target element indicated by the fragment identifier of the document's URI.</summary>
+  },
+  'text': function() {
+    /// <signature>
+    ///   <summary>Set the content of each element in the set of matched elements to the specified text.</summary>
+    ///   <param name="textString" type="String">A string of text to set as the content of each matched element.</param>
+    ///   <returns type="jQuery" />
+    /// </signature>
+    /// <signature>
+    ///   <summary>Set the content of each element in the set of matched elements to the specified text.</summary>
+    ///   <param name="function(index, text)" type="Function">A function returning the text content to set. Receives the index position of the element in the set and the old text value as arguments.</param>
+    ///   <returns type="jQuery" />
+    /// </signature>
+  },
+  'toArray': function() {
+    /// <summary>Retrieve all the DOM elements contained in the jQuery set, as an array.</summary>
+    /// <returns type="Array" />
+  },
+  'toggle': function() {
+    /// <signature>
+    ///   <summary>Display or hide the matched elements.</summary>
+    ///   <param name="duration" type="">A string or number determining how long the animation will run.</param>
+    ///   <param name="complete" type="Function">A function to call once the animation is complete.</param>
+    ///   <returns type="jQuery" />
+    /// </signature>
+    /// <signature>
+    ///   <summary>Display or hide the matched elements.</summary>
+    ///   <param name="options" type="PlainObject">A map of additional options to pass to the method.</param>
+    ///   <returns type="jQuery" />
+    /// </signature>
+    /// <signature>
+    ///   <summary>Display or hide the matched elements.</summary>
+    ///   <param name="duration" type="">A string or number determining how long the animation will run.</param>
+    ///   <param name="easing" type="String">A string indicating which easing function to use for the transition.</param>
+    ///   <param name="complete" type="Function">A function to call once the animation is complete.</param>
+    ///   <returns type="jQuery" />
+    /// </signature>
+    /// <signature>
+    ///   <summary>Display or hide the matched elements.</summary>
+    ///   <param name="showOrHide" type="Boolean">A Boolean indicating whether to show or hide the elements.</param>
+    ///   <returns type="jQuery" />
+    /// </signature>
+  },
+  'toggleClass': function() {
+    /// <signature>
+    ///   <summary>Add or remove one or more classes from each element in the set of matched elements, depending on either the class's presence or the value of the switch argument.</summary>
+    ///   <param name="className" type="String">One or more class names (separated by spaces) to be toggled for each element in the matched set.</param>
+    ///   <returns type="jQuery" />
+    /// </signature>
+    /// <signature>
+    ///   <summary>Add or remove one or more classes from each element in the set of matched elements, depending on either the class's presence or the value of the switch argument.</summary>
+    ///   <param name="className" type="String">One or more class names (separated by spaces) to be toggled for each element in the matched set.</param>
+    ///   <param name="switch" type="Boolean">A Boolean (not just truthy/falsy) value to determine whether the class should be added or removed.</param>
+    ///   <returns type="jQuery" />
+    /// </signature>
+    /// <signature>
+    ///   <summary>Add or remove one or more classes from each element in the set of matched elements, depending on either the class's presence or the value of the switch argument.</summary>
+    ///   <param name="switch" type="Boolean">A boolean value to determine whether the class should be added or removed.</param>
+    ///   <returns type="jQuery" />
+    /// </signature>
+    /// <signature>
+    ///   <summary>Add or remove one or more classes from each element in the set of matched elements, depending on either the class's presence or the value of the switch argument.</summary>
+    ///   <param name="function(index, class, switch)" type="Function">A function that returns class names to be toggled in the class attribute of each element in the matched set. Receives the index position of the element in the set, the old class value, and the switch as arguments.</param>
+    ///   <param name="switch" type="Boolean">A boolean value to determine whether the class should be added or removed.</param>
+    ///   <returns type="jQuery" />
+    /// </signature>
+  },
+  'trigger': function() {
+    /// <signature>
+    ///   <summary>Execute all handlers and behaviors attached to the matched elements for the given event type.</summary>
+    ///   <param name="eventType" type="String">A string containing a JavaScript event type, such as click or submit.</param>
+    ///   <param name="extraParameters" type="">Additional parameters to pass along to the event handler.</param>
+    ///   <returns type="jQuery" />
+    /// </signature>
+    /// <signature>
+    ///   <summary>Execute all handlers and behaviors attached to the matched elements for the given event type.</summary>
+    ///   <param name="event" type="Event">A jQuery.Event object.</param>
+    ///   <param name="extraParameters" type="">Additional parameters to pass along to the event handler.</param>
+    ///   <returns type="jQuery" />
+    /// </signature>
+  },
+  'triggerHandler': function() {
+    /// <signature>
+    ///   <summary>Execute all handlers attached to an element for an event.</summary>
+    ///   <param name="eventType" type="String">A string containing a JavaScript event type, such as click or submit.</param>
+    ///   <param name="extraParameters" type="Array">An array of additional parameters to pass along to the event handler.</param>
+    ///   <returns type="Object" />
+    /// </signature>
+  },
+  'unbind': function() {
+    /// <signature>
+    ///   <summary>Remove a previously-attached event handler from the elements.</summary>
+    ///   <param name="eventType" type="String">A string containing a JavaScript event type, such as click or submit.</param>
+    ///   <param name="handler(eventObject)" type="Function">The function that is to be no longer executed.</param>
+    ///   <returns type="jQuery" />
+    /// </signature>
+    /// <signature>
+    ///   <summary>Remove a previously-attached event handler from the elements.</summary>
+    ///   <param name="eventType" type="String">A string containing a JavaScript event type, such as click or submit.</param>
+    ///   <param name="false" type="Boolean">Unbinds the corresponding 'return false' function that was bound using .bind( eventType, false ).</param>
+    ///   <returns type="jQuery" />
+    /// </signature>
+    /// <signature>
+    ///   <summary>Remove a previously-attached event handler from the elements.</summary>
+    ///   <param name="event" type="Object">A JavaScript event object as passed to an event handler.</param>
+    ///   <returns type="jQuery" />
+    /// </signature>
+  },
+  'undelegate': function() {
+    /// <signature>
+    ///   <summary>Remove a handler from the event for all elements which match the current selector, based upon a specific set of root elements.</summary>
+    ///   <param name="selector" type="String">A selector which will be used to filter the event results.</param>
+    ///   <param name="eventType" type="String">A string containing a JavaScript event type, such as "click" or "keydown"</param>
+    ///   <returns type="jQuery" />
+    /// </signature>
+    /// <signature>
+    ///   <summary>Remove a handler from the event for all elements which match the current selector, based upon a specific set of root elements.</summary>
+    ///   <param name="selector" type="String">A selector which will be used to filter the event results.</param>
+    ///   <param name="eventType" type="String">A string containing a JavaScript event type, such as "click" or "keydown"</param>
+    ///   <param name="handler(eventObject)" type="Function">A function to execute at the time the event is triggered.</param>
+    ///   <returns type="jQuery" />
+    /// </signature>
+    /// <signature>
+    ///   <summary>Remove a handler from the event for all elements which match the current selector, based upon a specific set of root elements.</summary>
+    ///   <param name="selector" type="String">A selector which will be used to filter the event results.</param>
+    ///   <param name="events" type="PlainObject">An object of one or more event types and previously bound functions to unbind from them.</param>
+    ///   <returns type="jQuery" />
+    /// </signature>
+    /// <signature>
+    ///   <summary>Remove a handler from the event for all elements which match the current selector, based upon a specific set of root elements.</summary>
+    ///   <param name="namespace" type="String">A string containing a namespace to unbind all events from.</param>
+    ///   <returns type="jQuery" />
+    /// </signature>
+  },
+  'unload': function() {
+    /// <signature>
+    ///   <summary>Bind an event handler to the "unload" JavaScript event.</summary>
+    ///   <param name="handler(eventObject)" type="Function">A function to execute when the event is triggered.</param>
+    ///   <returns type="jQuery" />
+    /// </signature>
+    /// <signature>
+    ///   <summary>Bind an event handler to the "unload" JavaScript event.</summary>
+    ///   <param name="eventData" type="Object">A plain object of data that will be passed to the event handler.</param>
+    ///   <param name="handler(eventObject)" type="Function">A function to execute each time the event is triggered.</param>
+    ///   <returns type="jQuery" />
+    /// </signature>
+  },
+  'unwrap': function() {
+    /// <summary>Remove the parents of the set of matched elements from the DOM, leaving the matched elements in their place.</summary>
+    /// <returns type="jQuery" />
+  },
+  'val': function() {
+    /// <signature>
+    ///   <summary>Set the value of each element in the set of matched elements.</summary>
+    ///   <param name="value" type="">A string of text or an array of strings corresponding to the value of each matched element to set as selected/checked.</param>
+    ///   <returns type="jQuery" />
+    /// </signature>
+    /// <signature>
+    ///   <summary>Set the value of each element in the set of matched elements.</summary>
+    ///   <param name="function(index, value)" type="Function">A function returning the value to set. this is the current element. Receives the index position of the element in the set and the old value as arguments.</param>
+    ///   <returns type="jQuery" />
+    /// </signature>
+  },
+  'visible': function() {
+    /// <summary>Selects all elements that are visible.</summary>
+  },
+  'width': function() {
+    /// <signature>
+    ///   <summary>Set the CSS width of each element in the set of matched elements.</summary>
+    ///   <param name="value" type="">An integer representing the number of pixels, or an integer along with an optional unit of measure appended (as a string).</param>
+    ///   <returns type="jQuery" />
+    /// </signature>
+    /// <signature>
+    ///   <summary>Set the CSS width of each element in the set of matched elements.</summary>
+    ///   <param name="function(index, width)" type="Function">A function returning the width to set. Receives the index position of the element in the set and the old width as arguments. Within the function, this refers to the current element in the set.</param>
+    ///   <returns type="jQuery" />
+    /// </signature>
+  },
+  'wrap': function() {
+    /// <signature>
+    ///   <summary>Wrap an HTML structure around each element in the set of matched elements.</summary>
+    ///   <param name="wrappingElement" type="">A selector, element, HTML string, or jQuery object specifying the structure to wrap around the matched elements.</param>
+    ///   <returns type="jQuery" />
+    /// </signature>
+    /// <signature>
+    ///   <summary>Wrap an HTML structure around each element in the set of matched elements.</summary>
+    ///   <param name="function(index)" type="Function">A callback function returning the HTML content or jQuery object to wrap around the matched elements. Receives the index position of the element in the set as an argument. Within the function, this refers to the current element in the set.</param>
+    ///   <returns type="jQuery" />
+    /// </signature>
+  },
+  'wrapAll': function() {
+    /// <signature>
+    ///   <summary>Wrap an HTML structure around all elements in the set of matched elements.</summary>
+    ///   <param name="wrappingElement" type="">A selector, element, HTML string, or jQuery object specifying the structure to wrap around the matched elements.</param>
+    ///   <returns type="jQuery" />
+    /// </signature>
+  },
+  'wrapInner': function() {
+    /// <signature>
+    ///   <summary>Wrap an HTML structure around the content of each element in the set of matched elements.</summary>
+    ///   <param name="wrappingElement" type="String">An HTML snippet, selector expression, jQuery object, or DOM element specifying the structure to wrap around the content of the matched elements.</param>
+    ///   <returns type="jQuery" />
+    /// </signature>
+    /// <signature>
+    ///   <summary>Wrap an HTML structure around the content of each element in the set of matched elements.</summary>
+    ///   <param name="function(index)" type="Function">A callback function which generates a structure to wrap around the content of the matched elements. Receives the index position of the element in the set as an argument. Within the function, this refers to the current element in the set.</param>
+    ///   <returns type="jQuery" />
+    /// </signature>
+  },
+});
+
+intellisense.annotate(window, {
+  '$': function() {
+    /// <signature>
+    ///   <summary>Accepts a string containing a CSS selector which is then used to match a set of elements.</summary>
+    ///   <param name="selector" type="String">A string containing a selector expression</param>
+    ///   <param name="context" type="">A DOM Element, Document, or jQuery to use as context</param>
+    ///   <returns type="jQuery" />
+    /// </signature>
+    /// <signature>
+    ///   <summary>Accepts a string containing a CSS selector which is then used to match a set of elements.</summary>
+    ///   <param name="element" type="Element">A DOM element to wrap in a jQuery object.</param>
+    ///   <returns type="jQuery" />
+    /// </signature>
+    /// <signature>
+    ///   <summary>Accepts a string containing a CSS selector which is then used to match a set of elements.</summary>
+    ///   <param name="elementArray" type="Array">An array containing a set of DOM elements to wrap in a jQuery object.</param>
+    ///   <returns type="jQuery" />
+    /// </signature>
+    /// <signature>
+    ///   <summary>Accepts a string containing a CSS selector which is then used to match a set of elements.</summary>
+    ///   <param name="object" type="PlainObject">A plain object to wrap in a jQuery object.</param>
+    ///   <returns type="jQuery" />
+    /// </signature>
+    /// <signature>
+    ///   <summary>Accepts a string containing a CSS selector which is then used to match a set of elements.</summary>
+    ///   <param name="jQuery object" type="PlainObject">An existing jQuery object to clone.</param>
+    ///   <returns type="jQuery" />
+    /// </signature>
+  },
+});
+
