Copying //UE4/Dev-Core to //UE4/Dev-Main (Source: //UE4/Dev-Core @ 3151653)

diff --git a/Engine/Source/Programs/CrashReporter/CrashReportWebSite/Scripts/jquery-ui-1.12.0.js b/Engine/Source/Programs/CrashReporter/CrashReportWebSite/Scripts/jquery-ui-1.12.0.js
new file mode 100644
--- /dev/null
+++ b/Engine/Source/Programs/CrashReporter/CrashReportWebSite/Scripts/jquery-ui-1.12.0.js
@@ -0,0 +1,18686 @@
+/*! jQuery UI - v1.12.0 - 2016-07-08
+* http://jqueryui.com
+* Includes: widget.js, position.js, data.js, disable-selection.js, effect.js, effects/effect-blind.js, effects/effect-bounce.js, effects/effect-clip.js, effects/effect-drop.js, effects/effect-explode.js, effects/effect-fade.js, effects/effect-fold.js, effects/effect-highlight.js, effects/effect-puff.js, effects/effect-pulsate.js, effects/effect-scale.js, effects/effect-shake.js, effects/effect-size.js, effects/effect-slide.js, effects/effect-transfer.js, focusable.js, form-reset-mixin.js, jquery-1-7.js, keycode.js, labels.js, scroll-parent.js, tabbable.js, unique-id.js, widgets/accordion.js, widgets/autocomplete.js, widgets/button.js, widgets/checkboxradio.js, widgets/controlgroup.js, widgets/datepicker.js, widgets/dialog.js, widgets/draggable.js, widgets/droppable.js, widgets/menu.js, widgets/mouse.js, widgets/progressbar.js, widgets/resizable.js, widgets/selectable.js, widgets/selectmenu.js, widgets/slider.js, widgets/sortable.js, widgets/spinner.js, widgets/tabs.js, widgets/tooltip.js
+* Copyright jQuery Foundation and other contributors; Licensed MIT */
+
+(function( factory ) {
+	if ( typeof define === "function" && define.amd ) {
+
+		// AMD. Register as an anonymous module.
+		define([ "jquery" ], factory );
+	} else {
+
+		// Browser globals
+		factory( jQuery );
+	}
+}(function( $ ) {
+
+$.ui = $.ui || {};
+
+var version = $.ui.version = "1.12.0";
+
+
+/*!
+ * jQuery UI Widget 1.12.0
+ * http://jqueryui.com
+ *
+ * Copyright jQuery Foundation and other contributors
+ * Released under the MIT license.
+ * http://jquery.org/license
+ */
+
+//>>label: Widget
+//>>group: Core
+//>>description: Provides a factory for creating stateful widgets with a common API.
+//>>docs: http://api.jqueryui.com/jQuery.widget/
+//>>demos: http://jqueryui.com/widget/
+
+
+
+var widgetUuid = 0;
+var widgetSlice = Array.prototype.slice;
+
+$.cleanData = ( function( orig ) {
+	return function( elems ) {
+		var events, elem, i;
+		for ( i = 0; ( elem = elems[ i ] ) != null; i++ ) {
+			try {
+
+				// Only trigger remove when necessary to save time
+				events = $._data( elem, "events" );
+				if ( events && events.remove ) {
+					$( elem ).triggerHandler( "remove" );
+				}
+
+			// Http://bugs.jquery.com/ticket/8235
+			} catch ( e ) {}
+		}
+		orig( elems );
+	};
+} )( $.cleanData );
+
+$.widget = function( name, base, prototype ) {
+	var existingConstructor, constructor, basePrototype;
+
+	// ProxiedPrototype allows the provided prototype to remain unmodified
+	// so that it can be used as a mixin for multiple widgets (#8876)
+	var proxiedPrototype = {};
+
+	var namespace = name.split( "." )[ 0 ];
+	name = name.split( "." )[ 1 ];
+	var fullName = namespace + "-" + name;
+
+	if ( !prototype ) {
+		prototype = base;
+		base = $.Widget;
+	}
+
+	if ( $.isArray( prototype ) ) {
+		prototype = $.extend.apply( null, [ {} ].concat( prototype ) );
+	}
+
+	// Create selector for plugin
+	$.expr[ ":" ][ fullName.toLowerCase() ] = function( elem ) {
+		return !!$.data( elem, fullName );
+	};
+
+	$[ namespace ] = $[ namespace ] || {};
+	existingConstructor = $[ namespace ][ name ];
+	constructor = $[ namespace ][ name ] = function( options, element ) {
+
+		// Allow instantiation without "new" keyword
+		if ( !this._createWidget ) {
+			return new constructor( options, element );
+		}
+
+		// Allow instantiation without initializing for simple inheritance
+		// must use "new" keyword (the code above always passes args)
+		if ( arguments.length ) {
+			this._createWidget( options, element );
+		}
+	};
+
+	// Extend with the existing constructor to carry over any static properties
+	$.extend( constructor, existingConstructor, {
+		version: prototype.version,
+
+		// Copy the object used to create the prototype in case we need to
+		// redefine the widget later
+		_proto: $.extend( {}, prototype ),
+
+		// Track widgets that inherit from this widget in case this widget is
+		// redefined after a widget inherits from it
+		_childConstructors: []
+	} );
+
+	basePrototype = new base();
+
+	// We need to make the options hash a property directly on the new instance
+	// otherwise we'll modify the options hash on the prototype that we're
+	// inheriting from
+	basePrototype.options = $.widget.extend( {}, basePrototype.options );
+	$.each( prototype, function( prop, value ) {
+		if ( !$.isFunction( value ) ) {
+			proxiedPrototype[ prop ] = value;
+			return;
+		}
+		proxiedPrototype[ prop ] = ( function() {
+			function _super() {
+				return base.prototype[ prop ].apply( this, arguments );
+			}
+
+			function _superApply( args ) {
+				return base.prototype[ prop ].apply( this, args );
+			}
+
+			return function() {
+				var __super = this._super;
+				var __superApply = this._superApply;
+				var returnValue;
+
+				this._super = _super;
+				this._superApply = _superApply;
+
+				returnValue = value.apply( this, arguments );
+
+				this._super = __super;
+				this._superApply = __superApply;
+
+				return returnValue;
+			};
+		} )();
+	} );
+	constructor.prototype = $.widget.extend( basePrototype, {
+
+		// TODO: remove support for widgetEventPrefix
+		// always use the name + a colon as the prefix, e.g., draggable:start
+		// don't prefix for widgets that aren't DOM-based
+		widgetEventPrefix: existingConstructor ? ( basePrototype.widgetEventPrefix || name ) : name
+	}, proxiedPrototype, {
+		constructor: constructor,
+		namespace: namespace,
+		widgetName: name,
+		widgetFullName: fullName
+	} );
+
+	// If this widget is being redefined then we need to find all widgets that
+	// are inheriting from it and redefine all of them so that they inherit from
+	// the new version of this widget. We're essentially trying to replace one
+	// level in the prototype chain.
+	if ( existingConstructor ) {
+		$.each( existingConstructor._childConstructors, function( i, child ) {
+			var childPrototype = child.prototype;
+
+			// Redefine the child widget using the same prototype that was
+			// originally used, but inherit from the new version of the base
+			$.widget( childPrototype.namespace + "." + childPrototype.widgetName, constructor,
+				child._proto );
+		} );
+
+		// Remove the list of existing child constructors from the old constructor
+		// so the old child constructors can be garbage collected
+		delete existingConstructor._childConstructors;
+	} else {
+		base._childConstructors.push( constructor );
+	}
+
+	$.widget.bridge( name, constructor );
+
+	return constructor;
+};
+
+$.widget.extend = function( target ) {
+	var input = widgetSlice.call( arguments, 1 );
+	var inputIndex = 0;
+	var inputLength = input.length;
+	var key;
+	var value;
+
+	for ( ; inputIndex < inputLength; inputIndex++ ) {
+		for ( key in input[ inputIndex ] ) {
+			value = input[ inputIndex ][ key ];
+			if ( input[ inputIndex ].hasOwnProperty( key ) && value !== undefined ) {
+
+				// Clone objects
+				if ( $.isPlainObject( value ) ) {
+					target[ key ] = $.isPlainObject( target[ key ] ) ?
+						$.widget.extend( {}, target[ key ], value ) :
+
+						// Don't extend strings, arrays, etc. with objects
+						$.widget.extend( {}, value );
+
+				// Copy everything else by reference
+				} else {
+					target[ key ] = value;
+				}
+			}
+		}
+	}
+	return target;
+};
+
+$.widget.bridge = function( name, object ) {
+	var fullName = object.prototype.widgetFullName || name;
+	$.fn[ name ] = function( options ) {
+		var isMethodCall = typeof options === "string";
+		var args = widgetSlice.call( arguments, 1 );
+		var returnValue = this;
+
+		if ( isMethodCall ) {
+			this.each( function() {
+				var methodValue;
+				var instance = $.data( this, fullName );
+
+				if ( options === "instance" ) {
+					returnValue = instance;
+					return false;
+				}
+
+				if ( !instance ) {
+					return $.error( "cannot call methods on " + name +
+						" prior to initialization; " +
+						"attempted to call method '" + options + "'" );
+				}
+
+				if ( !$.isFunction( instance[ options ] ) || options.charAt( 0 ) === "_" ) {
+					return $.error( "no such method '" + options + "' for " + name +
+						" widget instance" );
+				}
+
+				methodValue = instance[ options ].apply( instance, args );
+
+				if ( methodValue !== instance && methodValue !== undefined ) {
+					returnValue = methodValue && methodValue.jquery ?
+						returnValue.pushStack( methodValue.get() ) :
+						methodValue;
+					return false;
+				}
+			} );
+		} else {
+
+			// Allow multiple hashes to be passed on init
+			if ( args.length ) {
+				options = $.widget.extend.apply( null, [ options ].concat( args ) );
+			}
+
+			this.each( function() {
+				var instance = $.data( this, fullName );
+				if ( instance ) {
+					instance.option( options || {} );
+					if ( instance._init ) {
+						instance._init();
+					}
+				} else {
+					$.data( this, fullName, new object( options, this ) );
+				}
+			} );
+		}
+
+		return returnValue;
+	};
+};
+
+$.Widget = function( /* options, element */ ) {};
+$.Widget._childConstructors = [];
+
+$.Widget.prototype = {
+	widgetName: "widget",
+	widgetEventPrefix: "",
+	defaultElement: "<div>",
+
+	options: {
+		classes: {},
+		disabled: false,
+
+		// Callbacks
+		create: null
+	},
+
+	_createWidget: function( options, element ) {
+		element = $( element || this.defaultElement || this )[ 0 ];
+		this.element = $( element );
+		this.uuid = widgetUuid++;
+		this.eventNamespace = "." + this.widgetName + this.uuid;
+
+		this.bindings = $();
+		this.hoverable = $();
+		this.focusable = $();
+		this.classesElementLookup = {};
+
+		if ( element !== this ) {
+			$.data( element, this.widgetFullName, this );
+			this._on( true, this.element, {
+				remove: function( event ) {
+					if ( event.target === element ) {
+						this.destroy();
+					}
+				}
+			} );
+			this.document = $( element.style ?
+
+				// Element within the document
+				element.ownerDocument :
+
+				// Element is window or document
+				element.document || element );
+			this.window = $( this.document[ 0 ].defaultView || this.document[ 0 ].parentWindow );
+		}
+
+		this.options = $.widget.extend( {},
+			this.options,
+			this._getCreateOptions(),
+			options );
+
+		this._create();
+
+		if ( this.options.disabled ) {
+			this._setOptionDisabled( this.options.disabled );
+		}
+
+		this._trigger( "create", null, this._getCreateEventData() );
+		this._init();
+	},
+
+	_getCreateOptions: function() {
+		return {};
+	},
+
+	_getCreateEventData: $.noop,
+
+	_create: $.noop,
+
+	_init: $.noop,
+
+	destroy: function() {
+		var that = this;
+
+		this._destroy();
+		$.each( this.classesElementLookup, function( key, value ) {
+			that._removeClass( value, key );
+		} );
+
+		// We can probably remove the unbind calls in 2.0
+		// all event bindings should go through this._on()
+		this.element
+			.off( this.eventNamespace )
+			.removeData( this.widgetFullName );
+		this.widget()
+			.off( this.eventNamespace )
+			.removeAttr( "aria-disabled" );
+
+		// Clean up events and states
+		this.bindings.off( this.eventNamespace );
+	},
+
+	_destroy: $.noop,
+
+	widget: function() {
+		return this.element;
+	},
+
+	option: function( key, value ) {
+		var options = key;
+		var parts;
+		var curOption;
+		var i;
+
+		if ( arguments.length === 0 ) {
+
+			// Don't return a reference to the internal hash
+			return $.widget.extend( {}, this.options );
+		}
+
+		if ( typeof key === "string" ) {
+
+			// Handle nested keys, e.g., "foo.bar" => { foo: { bar: ___ } }
+			options = {};
+			parts = key.split( "." );
+			key = parts.shift();
+			if ( parts.length ) {
+				curOption = options[ key ] = $.widget.extend( {}, this.options[ key ] );
+				for ( i = 0; i < parts.length - 1; i++ ) {
+					curOption[ parts[ i ] ] = curOption[ parts[ i ] ] || {};
+					curOption = curOption[ parts[ i ] ];
+				}
+				key = parts.pop();
+				if ( arguments.length === 1 ) {
+					return curOption[ key ] === undefined ? null : curOption[ key ];
+				}
+				curOption[ key ] = value;
+			} else {
+				if ( arguments.length === 1 ) {
+					return this.options[ key ] === undefined ? null : this.options[ key ];
+				}
+				options[ key ] = value;
+			}
+		}
+
+		this._setOptions( options );
+
+		return this;
+	},
+
+	_setOptions: function( options ) {
+		var key;
+
+		for ( key in options ) {
+			this._setOption( key, options[ key ] );
+		}
+
+		return this;
+	},
+
+	_setOption: function( key, value ) {
+		if ( key === "classes" ) {
+			this._setOptionClasses( value );
+		}
+
+		this.options[ key ] = value;
+
+		if ( key === "disabled" ) {
+			this._setOptionDisabled( value );
+		}
+
+		return this;
+	},
+
+	_setOptionClasses: function( value ) {
+		var classKey, elements, currentElements;
+
+		for ( classKey in value ) {
+			currentElements = this.classesElementLookup[ classKey ];
+			if ( value[ classKey ] === this.options.classes[ classKey ] ||
+					!currentElements ||
+					!currentElements.length ) {
+				continue;
+			}
+
+			// We are doing this to create a new jQuery object because the _removeClass() call
+			// on the next line is going to destroy the reference to the current elements being
+			// tracked. We need to save a copy of this collection so that we can add the new classes
+			// below.
+			elements = $( currentElements.get() );
+			this._removeClass( currentElements, classKey );
+
+			// We don't use _addClass() here, because that uses this.options.classes
+			// for generating the string of classes. We want to use the value passed in from
+			// _setOption(), this is the new value of the classes option which was passed to
+			// _setOption(). We pass this value directly to _classes().
+			elements.addClass( this._classes( {
+				element: elements,
+				keys: classKey,
+				classes: value,
+				add: true
+			} ) );
+		}
+	},
+
+	_setOptionDisabled: function( value ) {
+		this._toggleClass( this.widget(), this.widgetFullName + "-disabled", null, !!value );
+
+		// If the widget is becoming disabled, then nothing is interactive
+		if ( value ) {
+			this._removeClass( this.hoverable, null, "ui-state-hover" );
+			this._removeClass( this.focusable, null, "ui-state-focus" );
+		}
+	},
+
+	enable: function() {
+		return this._setOptions( { disabled: false } );
+	},
+
+	disable: function() {
+		return this._setOptions( { disabled: true } );
+	},
+
+	_classes: function( options ) {
+		var full = [];
+		var that = this;
+
+		options = $.extend( {
+			element: this.element,
+			classes: this.options.classes || {}
+		}, options );
+
+		function processClassString( classes, checkOption ) {
+			var current, i;
+			for ( i = 0; i < classes.length; i++ ) {
+				current = that.classesElementLookup[ classes[ i ] ] || $();
+				if ( options.add ) {
+					current = $( $.unique( current.get().concat( options.element.get() ) ) );
+				} else {
+					current = $( current.not( options.element ).get() );
+				}
+				that.classesElementLookup[ classes[ i ] ] = current;
+				full.push( classes[ i ] );
+				if ( checkOption && options.classes[ classes[ i ] ] ) {
+					full.push( options.classes[ classes[ i ] ] );
+				}
+			}
+		}
+
+		if ( options.keys ) {
+			processClassString( options.keys.match( /\S+/g ) || [], true );
+		}
+		if ( options.extra ) {
+			processClassString( options.extra.match( /\S+/g ) || [] );
+		}
+
+		return full.join( " " );
+	},
+
+	_removeClass: function( element, keys, extra ) {
+		return this._toggleClass( element, keys, extra, false );
+	},
+
+	_addClass: function( element, keys, extra ) {
+		return this._toggleClass( element, keys, extra, true );
+	},
+
+	_toggleClass: function( element, keys, extra, add ) {
+		add = ( typeof add === "boolean" ) ? add : extra;
+		var shift = ( typeof element === "string" || element === null ),
+			options = {
+				extra: shift ? keys : extra,
+				keys: shift ? element : keys,
+				element: shift ? this.element : element,
+				add: add
+			};
+		options.element.toggleClass( this._classes( options ), add );
+		return this;
+	},
+
+	_on: function( suppressDisabledCheck, element, handlers ) {
+		var delegateElement;
+		var instance = this;
+
+		// No suppressDisabledCheck flag, shuffle arguments
+		if ( typeof suppressDisabledCheck !== "boolean" ) {
+			handlers = element;
+			element = suppressDisabledCheck;
+			suppressDisabledCheck = false;
+		}
+
+		// No element argument, shuffle and use this.element
+		if ( !handlers ) {
+			handlers = element;
+			element = this.element;
+			delegateElement = this.widget();
+		} else {
+			element = delegateElement = $( element );
+			this.bindings = this.bindings.add( element );
+		}
+
+		$.each( handlers, function( event, handler ) {
+			function handlerProxy() {
+
+				// Allow widgets to customize the disabled handling
+				// - disabled as an array instead of boolean
+				// - disabled class as method for disabling individual parts
+				if ( !suppressDisabledCheck &&
+						( instance.options.disabled === true ||
+						$( this ).hasClass( "ui-state-disabled" ) ) ) {
+					return;
+				}
+				return ( typeof handler === "string" ? instance[ handler ] : handler )
+					.apply( instance, arguments );
+			}
+
+			// Copy the guid so direct unbinding works
+			if ( typeof handler !== "string" ) {
+				handlerProxy.guid = handler.guid =
+					handler.guid || handlerProxy.guid || $.guid++;
+			}
+
+			var match = event.match( /^([\w:-]*)\s*(.*)$/ );
+			var eventName = match[ 1 ] + instance.eventNamespace;
+			var selector = match[ 2 ];
+
+			if ( selector ) {
+				delegateElement.on( eventName, selector, handlerProxy );
+			} else {
+				element.on( eventName, handlerProxy );
+			}
+		} );
+	},
+
+	_off: function( element, eventName ) {
+		eventName = ( eventName || "" ).split( " " ).join( this.eventNamespace + " " ) +
+			this.eventNamespace;
+		element.off( eventName ).off( eventName );
+
+		// Clear the stack to avoid memory leaks (#10056)
+		this.bindings = $( this.bindings.not( element ).get() );
+		this.focusable = $( this.focusable.not( element ).get() );
+		this.hoverable = $( this.hoverable.not( element ).get() );
+	},
+
+	_delay: function( handler, delay ) {
+		function handlerProxy() {
+			return ( typeof handler === "string" ? instance[ handler ] : handler )
+				.apply( instance, arguments );
+		}
+		var instance = this;
+		return setTimeout( handlerProxy, delay || 0 );
+	},
+
+	_hoverable: function( element ) {
+		this.hoverable = this.hoverable.add( element );
+		this._on( element, {
+			mouseenter: function( event ) {
+				this._addClass( $( event.currentTarget ), null, "ui-state-hover" );
+			},
+			mouseleave: function( event ) {
+				this._removeClass( $( event.currentTarget ), null, "ui-state-hover" );
+			}
+		} );
+	},
+
+	_focusable: function( element ) {
+		this.focusable = this.focusable.add( element );
+		this._on( element, {
+			focusin: function( event ) {
+				this._addClass( $( event.currentTarget ), null, "ui-state-focus" );
+			},
+			focusout: function( event ) {
+				this._removeClass( $( event.currentTarget ), null, "ui-state-focus" );
+			}
+		} );
+	},
+
+	_trigger: function( type, event, data ) {
+		var prop, orig;
+		var callback = this.options[ type ];
+
+		data = data || {};
+		event = $.Event( event );
+		event.type = ( type === this.widgetEventPrefix ?
+			type :
+			this.widgetEventPrefix + type ).toLowerCase();
+
+		// The original event may come from any element
+		// so we need to reset the target on the new event
+		event.target = this.element[ 0 ];
+
+		// Copy original event properties over to the new event
+		orig = event.originalEvent;
+		if ( orig ) {
+			for ( prop in orig ) {
+				if ( !( prop in event ) ) {
+					event[ prop ] = orig[ prop ];
+				}
+			}
+		}
+
+		this.element.trigger( event, data );
+		return !( $.isFunction( callback ) &&
+			callback.apply( this.element[ 0 ], [ event ].concat( data ) ) === false ||
+			event.isDefaultPrevented() );
+	}
+};
+
+$.each( { show: "fadeIn", hide: "fadeOut" }, function( method, defaultEffect ) {
+	$.Widget.prototype[ "_" + method ] = function( element, options, callback ) {
+		if ( typeof options === "string" ) {
+			options = { effect: options };
+		}
+
+		var hasOptions;
+		var effectName = !options ?
+			method :
+			options === true || typeof options === "number" ?
+				defaultEffect :
+				options.effect || defaultEffect;
+
+		options = options || {};
+		if ( typeof options === "number" ) {
+			options = { duration: options };
+		}
+
+		hasOptions = !$.isEmptyObject( options );
+		options.complete = callback;
+
+		if ( options.delay ) {
+			element.delay( options.delay );
+		}
+
+		if ( hasOptions && $.effects && $.effects.effect[ effectName ] ) {
+			element[ method ]( options );
+		} else if ( effectName !== method && element[ effectName ] ) {
+			element[ effectName ]( options.duration, options.easing, callback );
+		} else {
+			element.queue( function( next ) {
+				$( this )[ method ]();
+				if ( callback ) {
+					callback.call( element[ 0 ] );
+				}
+				next();
+			} );
+		}
+	};
+} );
+
+var widget = $.widget;
+
+
+/*!
+ * jQuery UI Position 1.12.0
+ * http://jqueryui.com
+ *
+ * Copyright jQuery Foundation and other contributors
+ * Released under the MIT license.
+ * http://jquery.org/license
+ *
+ * http://api.jqueryui.com/position/
+ */
+
+//>>label: Position
+//>>group: Core
+//>>description: Positions elements relative to other elements.
+//>>docs: http://api.jqueryui.com/position/
+//>>demos: http://jqueryui.com/position/
+
+
+( function() {
+var cachedScrollbarWidth, supportsOffsetFractions,
+	max = Math.max,
+	abs = Math.abs,
+	round = Math.round,
+	rhorizontal = /left|center|right/,
+	rvertical = /top|center|bottom/,
+	roffset = /[\+\-]\d+(\.[\d]+)?%?/,
+	rposition = /^\w+/,
+	rpercent = /%$/,
+	_position = $.fn.position;
+
+// Support: IE <=9 only
+supportsOffsetFractions = function() {
+	var element = $( "<div>" )
+			.css( "position", "absolute" )
+			.appendTo( "body" )
+			.offset( {
+				top: 1.5,
+				left: 1.5
+			} ),
+		support = element.offset().top === 1.5;
+
+	element.remove();
+
+	supportsOffsetFractions = function() {
+		return support;
+	};
+
+	return support;
+};
+
+function getOffsets( offsets, width, height ) {
+	return [
+		parseFloat( offsets[ 0 ] ) * ( rpercent.test( offsets[ 0 ] ) ? width / 100 : 1 ),
+		parseFloat( offsets[ 1 ] ) * ( rpercent.test( offsets[ 1 ] ) ? height / 100 : 1 )
+	];
+}
+
+function parseCss( element, property ) {
+	return parseInt( $.css( element, property ), 10 ) || 0;
+}
+
+function getDimensions( elem ) {
+	var raw = elem[ 0 ];
+	if ( raw.nodeType === 9 ) {
+		return {
+			width: elem.width(),
+			height: elem.height(),
+			offset: { top: 0, left: 0 }
+		};
+	}
+	if ( $.isWindow( raw ) ) {
+		return {
+			width: elem.width(),
+			height: elem.height(),
+			offset: { top: elem.scrollTop(), left: elem.scrollLeft() }
+		};
+	}
+	if ( raw.preventDefault ) {
+		return {
+			width: 0,
+			height: 0,
+			offset: { top: raw.pageY, left: raw.pageX }
+		};
+	}
+	return {
+		width: elem.outerWidth(),
+		height: elem.outerHeight(),
+		offset: elem.offset()
+	};
+}
+
+$.position = {
+	scrollbarWidth: function() {
+		if ( cachedScrollbarWidth !== undefined ) {
+			return cachedScrollbarWidth;
+		}
+		var w1, w2,
+			div = $( "<div " +
+				"style='display:block;position:absolute;width:50px;height:50px;overflow:hidden;'>" +
+				"<div style='height:100px;width:auto;'></div></div>" ),
+			innerDiv = div.children()[ 0 ];
+
+		$( "body" ).append( div );
+		w1 = innerDiv.offsetWidth;
+		div.css( "overflow", "scroll" );
+
+		w2 = innerDiv.offsetWidth;
+
+		if ( w1 === w2 ) {
+			w2 = div[ 0 ].clientWidth;
+		}
+
+		div.remove();
+
+		return ( cachedScrollbarWidth = w1 - w2 );
+	},
+	getScrollInfo: function( within ) {
+		var overflowX = within.isWindow || within.isDocument ? "" :
+				within.element.css( "overflow-x" ),
+			overflowY = within.isWindow || within.isDocument ? "" :
+				within.element.css( "overflow-y" ),
+			hasOverflowX = overflowX === "scroll" ||
+				( overflowX === "auto" && within.width < within.element[ 0 ].scrollWidth ),
+			hasOverflowY = overflowY === "scroll" ||
+				( overflowY === "auto" && within.height < within.element[ 0 ].scrollHeight );
+		return {
+			width: hasOverflowY ? $.position.scrollbarWidth() : 0,
+			height: hasOverflowX ? $.position.scrollbarWidth() : 0
+		};
+	},
+	getWithinInfo: function( element ) {
+		var withinElement = $( element || window ),
+			isWindow = $.isWindow( withinElement[ 0 ] ),
+			isDocument = !!withinElement[ 0 ] && withinElement[ 0 ].nodeType === 9,
+			hasOffset = !isWindow && !isDocument;
+		return {
+			element: withinElement,
+			isWindow: isWindow,
+			isDocument: isDocument,
+			offset: hasOffset ? $( element ).offset() : { left: 0, top: 0 },
+			scrollLeft: withinElement.scrollLeft(),
+			scrollTop: withinElement.scrollTop(),
+			width: withinElement.outerWidth(),
+			height: withinElement.outerHeight()
+		};
+	}
+};
+
+$.fn.position = function( options ) {
+	if ( !options || !options.of ) {
+		return _position.apply( this, arguments );
+	}
+
+	// Make a copy, we don't want to modify arguments
+	options = $.extend( {}, options );
+
+	var atOffset, targetWidth, targetHeight, targetOffset, basePosition, dimensions,
+		target = $( options.of ),
+		within = $.position.getWithinInfo( options.within ),
+		scrollInfo = $.position.getScrollInfo( within ),
+		collision = ( options.collision || "flip" ).split( " " ),
+		offsets = {};
+
+	dimensions = getDimensions( target );
+	if ( target[ 0 ].preventDefault ) {
+
+		// Force left top to allow flipping
+		options.at = "left top";
+	}
+	targetWidth = dimensions.width;
+	targetHeight = dimensions.height;
+	targetOffset = dimensions.offset;
+
+	// Clone to reuse original targetOffset later
+	basePosition = $.extend( {}, targetOffset );
+
+	// Force my and at to have valid horizontal and vertical positions
+	// if a value is missing or invalid, it will be converted to center
+	$.each( [ "my", "at" ], function() {
+		var pos = ( options[ this ] || "" ).split( " " ),
+			horizontalOffset,
+			verticalOffset;
+
+		if ( pos.length === 1 ) {
+			pos = rhorizontal.test( pos[ 0 ] ) ?
+				pos.concat( [ "center" ] ) :
+				rvertical.test( pos[ 0 ] ) ?
+					[ "center" ].concat( pos ) :
+					[ "center", "center" ];
+		}
+		pos[ 0 ] = rhorizontal.test( pos[ 0 ] ) ? pos[ 0 ] : "center";
+		pos[ 1 ] = rvertical.test( pos[ 1 ] ) ? pos[ 1 ] : "center";
+
+		// Calculate offsets
+		horizontalOffset = roffset.exec( pos[ 0 ] );
+		verticalOffset = roffset.exec( pos[ 1 ] );
+		offsets[ this ] = [
+			horizontalOffset ? horizontalOffset[ 0 ] : 0,
+			verticalOffset ? verticalOffset[ 0 ] : 0
+		];
+
+		// Reduce to just the positions without the offsets
+		options[ this ] = [
+			rposition.exec( pos[ 0 ] )[ 0 ],
+			rposition.exec( pos[ 1 ] )[ 0 ]
+		];
+	} );
+
+	// Normalize collision option
+	if ( collision.length === 1 ) {
+		collision[ 1 ] = collision[ 0 ];
+	}
+
+	if ( options.at[ 0 ] === "right" ) {
+		basePosition.left += targetWidth;
+	} else if ( options.at[ 0 ] === "center" ) {
+		basePosition.left += targetWidth / 2;
+	}
+
+	if ( options.at[ 1 ] === "bottom" ) {
+		basePosition.top += targetHeight;
+	} else if ( options.at[ 1 ] === "center" ) {
+		basePosition.top += targetHeight / 2;
+	}
+
+	atOffset = getOffsets( offsets.at, targetWidth, targetHeight );
+	basePosition.left += atOffset[ 0 ];
+	basePosition.top += atOffset[ 1 ];
+
+	return this.each( function() {
+		var collisionPosition, using,
+			elem = $( this ),
+			elemWidth = elem.outerWidth(),
+			elemHeight = elem.outerHeight(),
+			marginLeft = parseCss( this, "marginLeft" ),
+			marginTop = parseCss( this, "marginTop" ),
+			collisionWidth = elemWidth + marginLeft + parseCss( this, "marginRight" ) +
+				scrollInfo.width,
+			collisionHeight = elemHeight + marginTop + parseCss( this, "marginBottom" ) +
+				scrollInfo.height,
+			position = $.extend( {}, basePosition ),
+			myOffset = getOffsets( offsets.my, elem.outerWidth(), elem.outerHeight() );
+
+		if ( options.my[ 0 ] === "right" ) {
+			position.left -= elemWidth;
+		} else if ( options.my[ 0 ] === "center" ) {
+			position.left -= elemWidth / 2;
+		}
+
+		if ( options.my[ 1 ] === "bottom" ) {
+			position.top -= elemHeight;
+		} else if ( options.my[ 1 ] === "center" ) {
+			position.top -= elemHeight / 2;
+		}
+
+		position.left += myOffset[ 0 ];
+		position.top += myOffset[ 1 ];
+
+		// If the browser doesn't support fractions, then round for consistent results
+		if ( !supportsOffsetFractions() ) {
+			position.left = round( position.left );
+			position.top = round( position.top );
+		}
+
+		collisionPosition = {
+			marginLeft: marginLeft,
+			marginTop: marginTop
+		};
+
+		$.each( [ "left", "top" ], function( i, dir ) {
+			if ( $.ui.position[ collision[ i ] ] ) {
+				$.ui.position[ collision[ i ] ][ dir ]( position, {
+					targetWidth: targetWidth,
+					targetHeight: targetHeight,
+					elemWidth: elemWidth,
+					elemHeight: elemHeight,
+					collisionPosition: collisionPosition,
+					collisionWidth: collisionWidth,
+					collisionHeight: collisionHeight,
+					offset: [ atOffset[ 0 ] + myOffset[ 0 ], atOffset [ 1 ] + myOffset[ 1 ] ],
+					my: options.my,
+					at: options.at,
+					within: within,
+					elem: elem
+				} );
+			}
+		} );
+
+		if ( options.using ) {
+
+			// Adds feedback as second argument to using callback, if present
+			using = function( props ) {
+				var left = targetOffset.left - position.left,
+					right = left + targetWidth - elemWidth,
+					top = targetOffset.top - position.top,
+					bottom = top + targetHeight - elemHeight,
+					feedback = {
+						target: {
+							element: target,
+							left: targetOffset.left,
+							top: targetOffset.top,
+							width: targetWidth,
+							height: targetHeight
+						},
+						element: {
+							element: elem,
+							left: position.left,
+							top: position.top,
+							width: elemWidth,
+							height: elemHeight
+						},
+						horizontal: right < 0 ? "left" : left > 0 ? "right" : "center",
+						vertical: bottom < 0 ? "top" : top > 0 ? "bottom" : "middle"
+					};
+				if ( targetWidth < elemWidth && abs( left + right ) < targetWidth ) {
+					feedback.horizontal = "center";
+				}
+				if ( targetHeight < elemHeight && abs( top + bottom ) < targetHeight ) {
+					feedback.vertical = "middle";
+				}
+				if ( max( abs( left ), abs( right ) ) > max( abs( top ), abs( bottom ) ) ) {
+					feedback.important = "horizontal";
+				} else {
+					feedback.important = "vertical";
+				}
+				options.using.call( this, props, feedback );
+			};
+		}
+
+		elem.offset( $.extend( position, { using: using } ) );
+	} );
+};
+
+$.ui.position = {
+	fit: {
+		left: function( position, data ) {
+			var within = data.within,
+				withinOffset = within.isWindow ? within.scrollLeft : within.offset.left,
+				outerWidth = within.width,
+				collisionPosLeft = position.left - data.collisionPosition.marginLeft,
+				overLeft = withinOffset - collisionPosLeft,
+				overRight = collisionPosLeft + data.collisionWidth - outerWidth - withinOffset,
+				newOverRight;
+
+			// Element is wider than within
+			if ( data.collisionWidth > outerWidth ) {
+
+				// Element is initially over the left side of within
+				if ( overLeft > 0 && overRight <= 0 ) {
+					newOverRight = position.left + overLeft + data.collisionWidth - outerWidth -
+						withinOffset;
+					position.left += overLeft - newOverRight;
+
+				// Element is initially over right side of within
+				} else if ( overRight > 0 && overLeft <= 0 ) {
+					position.left = withinOffset;
+
+				// Element is initially over both left and right sides of within
+				} else {
+					if ( overLeft > overRight ) {
+						position.left = withinOffset + outerWidth - data.collisionWidth;
+					} else {
+						position.left = withinOffset;
+					}
+				}
+
+			// Too far left -> align with left edge
+			} else if ( overLeft > 0 ) {
+				position.left += overLeft;
+
+			// Too far right -> align with right edge
+			} else if ( overRight > 0 ) {
+				position.left -= overRight;
+
+			// Adjust based on position and margin
+			} else {
+				position.left = max( position.left - collisionPosLeft, position.left );
+			}
+		},
+		top: function( position, data ) {
+			var within = data.within,
+				withinOffset = within.isWindow ? within.scrollTop : within.offset.top,
+				outerHeight = data.within.height,
+				collisionPosTop = position.top - data.collisionPosition.marginTop,
+				overTop = withinOffset - collisionPosTop,
+				overBottom = collisionPosTop + data.collisionHeight - outerHeight - withinOffset,
+				newOverBottom;
+
+			// Element is taller than within
+			if ( data.collisionHeight > outerHeight ) {
+
+				// Element is initially over the top of within
+				if ( overTop > 0 && overBottom <= 0 ) {
+					newOverBottom = position.top + overTop + data.collisionHeight - outerHeight -
+						withinOffset;
+					position.top += overTop - newOverBottom;
+
+				// Element is initially over bottom of within
+				} else if ( overBottom > 0 && overTop <= 0 ) {
+					position.top = withinOffset;
+
+				// Element is initially over both top and bottom of within
+				} else {
+					if ( overTop > overBottom ) {
+						position.top = withinOffset + outerHeight - data.collisionHeight;
+					} else {
+						position.top = withinOffset;
+					}
+				}
+
+			// Too far up -> align with top
+			} else if ( overTop > 0 ) {
+				position.top += overTop;
+
+			// Too far down -> align with bottom edge
+			} else if ( overBottom > 0 ) {
+				position.top -= overBottom;
+
+			// Adjust based on position and margin
+			} else {
+				position.top = max( position.top - collisionPosTop, position.top );
+			}
+		}
+	},
+	flip: {
+		left: function( position, data ) {
+			var within = data.within,
+				withinOffset = within.offset.left + within.scrollLeft,
+				outerWidth = within.width,
+				offsetLeft = within.isWindow ? within.scrollLeft : within.offset.left,
+				collisionPosLeft = position.left - data.collisionPosition.marginLeft,
+				overLeft = collisionPosLeft - offsetLeft,
+				overRight = collisionPosLeft + data.collisionWidth - outerWidth - offsetLeft,
+				myOffset = data.my[ 0 ] === "left" ?
+					-data.elemWidth :
+					data.my[ 0 ] === "right" ?
+						data.elemWidth :
+						0,
+				atOffset = data.at[ 0 ] === "left" ?
+					data.targetWidth :
+					data.at[ 0 ] === "right" ?
+						-data.targetWidth :
+						0,
+				offset = -2 * data.offset[ 0 ],
+				newOverRight,
+				newOverLeft;
+
+			if ( overLeft < 0 ) {
+				newOverRight = position.left + myOffset + atOffset + offset + data.collisionWidth -
+					outerWidth - withinOffset;
+				if ( newOverRight < 0 || newOverRight < abs( overLeft ) ) {
+					position.left += myOffset + atOffset + offset;
+				}
+			} else if ( overRight > 0 ) {
+				newOverLeft = position.left - data.collisionPosition.marginLeft + myOffset +
+					atOffset + offset - offsetLeft;
+				if ( newOverLeft > 0 || abs( newOverLeft ) < overRight ) {
+					position.left += myOffset + atOffset + offset;
+				}
+			}
+		},
+		top: function( position, data ) {
+			var within = data.within,
+				withinOffset = within.offset.top + within.scrollTop,
+				outerHeight = within.height,
+				offsetTop = within.isWindow ? within.scrollTop : within.offset.top,
+				collisionPosTop = position.top - data.collisionPosition.marginTop,
+				overTop = collisionPosTop - offsetTop,
+				overBottom = collisionPosTop + data.collisionHeight - outerHeight - offsetTop,
+				top = data.my[ 1 ] === "top",
+				myOffset = top ?
+					-data.elemHeight :
+					data.my[ 1 ] === "bottom" ?
+						data.elemHeight :
+						0,
+				atOffset = data.at[ 1 ] === "top" ?
+					data.targetHeight :
+					data.at[ 1 ] === "bottom" ?
+						-data.targetHeight :
+						0,
+				offset = -2 * data.offset[ 1 ],
+				newOverTop,
+				newOverBottom;
+			if ( overTop < 0 ) {
+				newOverBottom = position.top + myOffset + atOffset + offset + data.collisionHeight -
+					outerHeight - withinOffset;
+				if ( newOverBottom < 0 || newOverBottom < abs( overTop ) ) {
+					position.top += myOffset + atOffset + offset;
+				}
+			} else if ( overBottom > 0 ) {
+				newOverTop = position.top - data.collisionPosition.marginTop + myOffset + atOffset +
+					offset - offsetTop;
+				if ( newOverTop > 0 || abs( newOverTop ) < overBottom ) {
+					position.top += myOffset + atOffset + offset;
+				}
+			}
+		}
+	},
+	flipfit: {
+		left: function() {
+			$.ui.position.flip.left.apply( this, arguments );
+			$.ui.position.fit.left.apply( this, arguments );
+		},
+		top: function() {
+			$.ui.position.flip.top.apply( this, arguments );
+			$.ui.position.fit.top.apply( this, arguments );
+		}
+	}
+};
+
+} )();
+
+var position = $.ui.position;
+
+
+/*!
+ * jQuery UI :data 1.12.0
+ * http://jqueryui.com
+ *
+ * Copyright jQuery Foundation and other contributors
+ * Released under the MIT license.
+ * http://jquery.org/license
+ */
+
+//>>label: :data Selector
+//>>group: Core
+//>>description: Selects elements which have data stored under the specified key.
+//>>docs: http://api.jqueryui.com/data-selector/
+
+
+var data = $.extend( $.expr[ ":" ], {
+	data: $.expr.createPseudo ?
+		$.expr.createPseudo( function( dataName ) {
+			return function( elem ) {
+				return !!$.data( elem, dataName );
+			};
+		} ) :
+
+		// Support: jQuery <1.8
+		function( elem, i, match ) {
+			return !!$.data( elem, match[ 3 ] );
+		}
+} );
+
+/*!
+ * jQuery UI Disable Selection 1.12.0
+ * http://jqueryui.com
+ *
+ * Copyright jQuery Foundation and other contributors
+ * Released under the MIT license.
+ * http://jquery.org/license
+ */
+
+//>>label: disableSelection
+//>>group: Core
+//>>description: Disable selection of text content within the set of matched elements.
+//>>docs: http://api.jqueryui.com/disableSelection/
+
+// This file is deprecated
+
+
+var disableSelection = $.fn.extend( {
+	disableSelection: ( function() {
+		var eventType = "onselectstart" in document.createElement( "div" ) ?
+			"selectstart" :
+			"mousedown";
+
+		return function() {
+			return this.on( eventType + ".ui-disableSelection", function( event ) {
+				event.preventDefault();
+			} );
+		};
+	} )(),
+
+	enableSelection: function() {
+		return this.off( ".ui-disableSelection" );
+	}
+} );
+
+
+/*!
+ * jQuery UI Effects 1.12.0
+ * http://jqueryui.com
+ *
+ * Copyright jQuery Foundation and other contributors
+ * Released under the MIT license.
+ * http://jquery.org/license
+ */
+
+//>>label: Effects Core
+//>>group: Effects
+// jscs:disable maximumLineLength
+//>>description: Extends the internal jQuery effects. Includes morphing and easing. Required by all other effects.
+// jscs:enable maximumLineLength
+//>>docs: http://api.jqueryui.com/category/effects-core/
+//>>demos: http://jqueryui.com/effect/
+
+
+
+var dataSpace = "ui-effects-",
+	dataSpaceStyle = "ui-effects-style",
+	dataSpaceAnimated = "ui-effects-animated",
+
+	// Create a local jQuery because jQuery Color relies on it and the
+	// global may not exist with AMD and a custom build (#10199)
+	jQuery = $;
+
+$.effects = {
+	effect: {}
+};
+
+/*!
+ * jQuery Color Animations v2.1.2
+ * [REDACTED_PROFILE_URL]/jquery-color
+ *
+ * Copyright 2014 jQuery Foundation and other contributors
+ * Released under the MIT license.
+ * http://jquery.org/license
+ *
+ * Date: Wed Jan 16 08:47:09 2013 -0600
+ */
+( function( jQuery, undefined ) {
+
+	var stepHooks = "backgroundColor borderBottomColor borderLeftColor borderRightColor " +
+		"borderTopColor color columnRuleColor outlineColor textDecorationColor textEmphasisColor",
+
+	// Plusequals test for += 100 -= 100
+	rplusequals = /^([\-+])=\s*(\d+\.?\d*)/,
+
+	// A set of RE's that can match strings and generate color tuples.
+	stringParsers = [ {
+			re: /rgba?\(\s*(\d{1,3})\s*,\s*(\d{1,3})\s*,\s*(\d{1,3})\s*(?:,\s*(\d?(?:\.\d+)?)\s*)?\)/,
+			parse: function( execResult ) {
+				return [
+					execResult[ 1 ],
+					execResult[ 2 ],
+					execResult[ 3 ],
+					execResult[ 4 ]
+				];
+			}
+		}, {
+			re: /rgba?\(\s*(\d+(?:\.\d+)?)\%\s*,\s*(\d+(?:\.\d+)?)\%\s*,\s*(\d+(?:\.\d+)?)\%\s*(?:,\s*(\d?(?:\.\d+)?)\s*)?\)/,
+			parse: function( execResult ) {
+				return [
+					execResult[ 1 ] * 2.55,
+					execResult[ 2 ] * 2.55,
+					execResult[ 3 ] * 2.55,
+					execResult[ 4 ]
+				];
+			}
+		}, {
+
+			// This regex ignores A-F because it's compared against an already lowercased string
+			re: /#([a-f0-9]{2})([a-f0-9]{2})([a-f0-9]{2})/,
+			parse: function( execResult ) {
+				return [
+					parseInt( execResult[ 1 ], 16 ),
+					parseInt( execResult[ 2 ], 16 ),
+					parseInt( execResult[ 3 ], 16 )
+				];
+			}
+		}, {
+
+			// This regex ignores A-F because it's compared against an already lowercased string
+			re: /#([a-f0-9])([a-f0-9])([a-f0-9])/,
+			parse: function( execResult ) {
+				return [
+					parseInt( execResult[ 1 ] + execResult[ 1 ], 16 ),
+					parseInt( execResult[ 2 ] + execResult[ 2 ], 16 ),
+					parseInt( execResult[ 3 ] + execResult[ 3 ], 16 )
+				];
+			}
+		}, {
+			re: /hsla?\(\s*(\d+(?:\.\d+)?)\s*,\s*(\d+(?:\.\d+)?)\%\s*,\s*(\d+(?:\.\d+)?)\%\s*(?:,\s*(\d?(?:\.\d+)?)\s*)?\)/,
+			space: "hsla",
+			parse: function( execResult ) {
+				return [
+					execResult[ 1 ],
+					execResult[ 2 ] / 100,
+					execResult[ 3 ] / 100,
+					execResult[ 4 ]
+				];
+			}
+		} ],
+
+	// JQuery.Color( )
+	color = jQuery.Color = function( color, green, blue, alpha ) {
+		return new jQuery.Color.fn.parse( color, green, blue, alpha );
+	},
+	spaces = {
+		rgba: {
+			props: {
+				red: {
+					idx: 0,
+					type: "byte"
+				},
+				green: {
+					idx: 1,
+					type: "byte"
+				},
+				blue: {
+					idx: 2,
+					type: "byte"
+				}
+			}
+		},
+
+		hsla: {
+			props: {
+				hue: {
+					idx: 0,
+					type: "degrees"
+				},
+				saturation: {
+					idx: 1,
+					type: "percent"
+				},
+				lightness: {
+					idx: 2,
+					type: "percent"
+				}
+			}
+		}
+	},
+	propTypes = {
+		"byte": {
+			floor: true,
+			max: 255
+		},
+		"percent": {
+			max: 1
+		},
+		"degrees": {
+			mod: 360,
+			floor: true
+		}
+	},
+	support = color.support = {},
+
+	// Element for support tests
+	supportElem = jQuery( "<p>" )[ 0 ],
+
+	// Colors = jQuery.Color.names
+	colors,
+
+	// Local aliases of functions called often
+	each = jQuery.each;
+
+// Determine rgba support immediately
+supportElem.style.cssText = "background-color:rgba(1,1,1,.5)";
+support.rgba = supportElem.style.backgroundColor.indexOf( "rgba" ) > -1;
+
+// Define cache name and alpha properties
+// for rgba and hsla spaces
+each( spaces, function( spaceName, space ) {
+	space.cache = "_" + spaceName;
+	space.props.alpha = {
+		idx: 3,
+		type: "percent",
+		def: 1
+	};
+} );
+
+function clamp( value, prop, allowEmpty ) {
+	var type = propTypes[ prop.type ] || {};
+
+	if ( value == null ) {
+		return ( allowEmpty || !prop.def ) ? null : prop.def;
+	}
+
+	// ~~ is an short way of doing floor for positive numbers
+	value = type.floor ? ~~value : parseFloat( value );
+
+	// IE will pass in empty strings as value for alpha,
+	// which will hit this case
+	if ( isNaN( value ) ) {
+		return prop.def;
+	}
+
+	if ( type.mod ) {
+
+		// We add mod before modding to make sure that negatives values
+		// get converted properly: -10 -> 350
+		return ( value + type.mod ) % type.mod;
+	}
+
+	// For now all property types without mod have min and max
+	return 0 > value ? 0 : type.max < value ? type.max : value;
+}
+
+function stringParse( string ) {
+	var inst = color(),
+		rgba = inst._rgba = [];
+
+	string = string.toLowerCase();
+
+	each( stringParsers, function( i, parser ) {
+		var parsed,
+			match = parser.re.exec( string ),
+			values = match && parser.parse( match ),
+			spaceName = parser.space || "rgba";
+
+		if ( values ) {
+			parsed = inst[ spaceName ]( values );
+
+			// If this was an rgba parse the assignment might happen twice
+			// oh well....
+			inst[ spaces[ spaceName ].cache ] = parsed[ spaces[ spaceName ].cache ];
+			rgba = inst._rgba = parsed._rgba;
+
+			// Exit each( stringParsers ) here because we matched
+			return false;
+		}
+	} );
+
+	// Found a stringParser that handled it
+	if ( rgba.length ) {
+
+		// If this came from a parsed string, force "transparent" when alpha is 0
+		// chrome, (and maybe others) return "transparent" as rgba(0,0,0,0)
+		if ( rgba.join() === "0,0,0,0" ) {
+			jQuery.extend( rgba, colors.transparent );
+		}
+		return inst;
+	}
+
+	// Named colors
+	return colors[ string ];
+}
+
+color.fn = jQuery.extend( color.prototype, {
+	parse: function( red, green, blue, alpha ) {
+		if ( red === undefined ) {
+			this._rgba = [ null, null, null, null ];
+			return this;
+		}
+		if ( red.jquery || red.nodeType ) {
+			red = jQuery( red ).css( green );
+			green = undefined;
+		}
+
+		var inst = this,
+			type = jQuery.type( red ),
+			rgba = this._rgba = [];
+
+		// More than 1 argument specified - assume ( red, green, blue, alpha )
+		if ( green !== undefined ) {
+			red = [ red, green, blue, alpha ];
+			type = "array";
+		}
+
+		if ( type === "string" ) {
+			return this.parse( stringParse( red ) || colors._default );
+		}
+
+		if ( type === "array" ) {
+			each( spaces.rgba.props, function( key, prop ) {
+				rgba[ prop.idx ] = clamp( red[ prop.idx ], prop );
+			} );
+			return this;
+		}
+
+		if ( type === "object" ) {
+			if ( red instanceof color ) {
+				each( spaces, function( spaceName, space ) {
+					if ( red[ space.cache ] ) {
+						inst[ space.cache ] = red[ space.cache ].slice();
+					}
+				} );
+			} else {
+				each( spaces, function( spaceName, space ) {
+					var cache = space.cache;
+					each( space.props, function( key, prop ) {
+
+						// If the cache doesn't exist, and we know how to convert
+						if ( !inst[ cache ] && space.to ) {
+
+							// If the value was null, we don't need to copy it
+							// if the key was alpha, we don't need to copy it either
+							if ( key === "alpha" || red[ key ] == null ) {
+								return;
+							}
+							inst[ cache ] = space.to( inst._rgba );
+						}
+
+						// This is the only case where we allow nulls for ALL properties.
+						// call clamp with alwaysAllowEmpty
+						inst[ cache ][ prop.idx ] = clamp( red[ key ], prop, true );
+					} );
+
+					// Everything defined but alpha?
+					if ( inst[ cache ] &&
+							jQuery.inArray( null, inst[ cache ].slice( 0, 3 ) ) < 0 ) {
+
+						// Use the default of 1
+						inst[ cache ][ 3 ] = 1;
+						if ( space.from ) {
+							inst._rgba = space.from( inst[ cache ] );
+						}
+					}
+				} );
+			}
+			return this;
+		}
+	},
+	is: function( compare ) {
+		var is = color( compare ),
+			same = true,
+			inst = this;
+
+		each( spaces, function( _, space ) {
+			var localCache,
+				isCache = is[ space.cache ];
+			if ( isCache ) {
+				localCache = inst[ space.cache ] || space.to && space.to( inst._rgba ) || [];
+				each( space.props, function( _, prop ) {
+					if ( isCache[ prop.idx ] != null ) {
+						same = ( isCache[ prop.idx ] === localCache[ prop.idx ] );
+						return same;
+					}
+				} );
+			}
+			return same;
+		} );
+		return same;
+	},
+	_space: function() {
+		var used = [],
+			inst = this;
+		each( spaces, function( spaceName, space ) {
+			if ( inst[ space.cache ] ) {
+				used.push( spaceName );
+			}
+		} );
+		return used.pop();
+	},
+	transition: function( other, distance ) {
+		var end = color( other ),
+			spaceName = end._space(),
+			space = spaces[ spaceName ],
+			startColor = this.alpha() === 0 ? color( "transparent" ) : this,
+			start = startColor[ space.cache ] || space.to( startColor._rgba ),
+			result = start.slice();
+
+		end = end[ space.cache ];
+		each( space.props, function( key, prop ) {
+			var index = prop.idx,
+				startValue = start[ index ],
+				endValue = end[ index ],
+				type = propTypes[ prop.type ] || {};
+
+			// If null, don't override start value
+			if ( endValue === null ) {
+				return;
+			}
+
+			// If null - use end
+			if ( startValue === null ) {
+				result[ index ] = endValue;
+			} else {
+				if ( type.mod ) {
+					if ( endValue - startValue > type.mod / 2 ) {
+						startValue += type.mod;
+					} else if ( startValue - endValue > type.mod / 2 ) {
+						startValue -= type.mod;
+					}
+				}
+				result[ index ] = clamp( ( endValue - startValue ) * distance + startValue, prop );
+			}
+		} );
+		return this[ spaceName ]( result );
+	},
+	blend: function( opaque ) {
+
+		// If we are already opaque - return ourself
+		if ( this._rgba[ 3 ] === 1 ) {
+			return this;
+		}
+
+		var rgb = this._rgba.slice(),
+			a = rgb.pop(),
+			blend = color( opaque )._rgba;
+
+		return color( jQuery.map( rgb, function( v, i ) {
+			return ( 1 - a ) * blend[ i ] + a * v;
+		} ) );
+	},
+	toRgbaString: function() {
+		var prefix = "rgba(",
+			rgba = jQuery.map( this._rgba, function( v, i ) {
+				return v == null ? ( i > 2 ? 1 : 0 ) : v;
+			} );
+
+		if ( rgba[ 3 ] === 1 ) {
+			rgba.pop();
+			prefix = "rgb(";
+		}
+
+		return prefix + rgba.join() + ")";
+	},
+	toHslaString: function() {
+		var prefix = "hsla(",
+			hsla = jQuery.map( this.hsla(), function( v, i ) {
+				if ( v == null ) {
+					v = i > 2 ? 1 : 0;
+				}
+
+				// Catch 1 and 2
+				if ( i && i < 3 ) {
+					v = Math.round( v * 100 ) + "%";
+				}
+				return v;
+			} );
+
+		if ( hsla[ 3 ] === 1 ) {
+			hsla.pop();
+			prefix = "hsl(";
+		}
+		return prefix + hsla.join() + ")";
+	},
+	toHexString: function( includeAlpha ) {
+		var rgba = this._rgba.slice(),
+			alpha = rgba.pop();
+
+		if ( includeAlpha ) {
+			rgba.push( ~~( alpha * 255 ) );
+		}
+
+		return "#" + jQuery.map( rgba, function( v ) {
+
+			// Default to 0 when nulls exist
+			v = ( v || 0 ).toString( 16 );
+			return v.length === 1 ? "0" + v : v;
+		} ).join( "" );
+	},
+	toString: function() {
+		return this._rgba[ 3 ] === 0 ? "transparent" : this.toRgbaString();
+	}
+} );
+color.fn.parse.prototype = color.fn;
+
+// Hsla conversions adapted from:
+// https://code.google.com/p/maashaack/source/browse/packages/graphics/trunk/src/graphics/colors/HUE2RGB.as?r=5021
+
+function hue2rgb( p, q, h ) {
+	h = ( h + 1 ) % 1;
+	if ( h * 6 < 1 ) {
+		return p + ( q - p ) * h * 6;
+	}
+	if ( h * 2 < 1 ) {
+		return q;
+	}
+	if ( h * 3 < 2 ) {
+		return p + ( q - p ) * ( ( 2 / 3 ) - h ) * 6;
+	}
+	return p;
+}
+
+spaces.hsla.to = function( rgba ) {
+	if ( rgba[ 0 ] == null || rgba[ 1 ] == null || rgba[ 2 ] == null ) {
+		return [ null, null, null, rgba[ 3 ] ];
+	}
+	var r = rgba[ 0 ] / 255,
+		g = rgba[ 1 ] / 255,
+		b = rgba[ 2 ] / 255,
+		a = rgba[ 3 ],
+		max = Math.max( r, g, b ),
+		min = Math.min( r, g, b ),
+		diff = max - min,
+		add = max + min,
+		l = add * 0.5,
+		h, s;
+
+	if ( min === max ) {
+		h = 0;
+	} else if ( r === max ) {
+		h = ( 60 * ( g - b ) / diff ) + 360;
+	} else if ( g === max ) {
+		h = ( 60 * ( b - r ) / diff ) + 120;
+	} else {
+		h = ( 60 * ( r - g ) / diff ) + 240;
+	}
+
+	// Chroma (diff) == 0 means greyscale which, by definition, saturation = 0%
+	// otherwise, saturation is based on the ratio of chroma (diff) to lightness (add)
+	if ( diff === 0 ) {
+		s = 0;
+	} else if ( l <= 0.5 ) {
+		s = diff / add;
+	} else {
+		s = diff / ( 2 - add );
+	}
+	return [ Math.round( h ) % 360, s, l, a == null ? 1 : a ];
+};
+
+spaces.hsla.from = function( hsla ) {
+	if ( hsla[ 0 ] == null || hsla[ 1 ] == null || hsla[ 2 ] == null ) {
+		return [ null, null, null, hsla[ 3 ] ];
+	}
+	var h = hsla[ 0 ] / 360,
+		s = hsla[ 1 ],
+		l = hsla[ 2 ],
+		a = hsla[ 3 ],
+		q = l <= 0.5 ? l * ( 1 + s ) : l + s - l * s,
+		p = 2 * l - q;
+
+	return [
+		Math.round( hue2rgb( p, q, h + ( 1 / 3 ) ) * 255 ),
+		Math.round( hue2rgb( p, q, h ) * 255 ),
+		Math.round( hue2rgb( p, q, h - ( 1 / 3 ) ) * 255 ),
+		a
+	];
+};
+
+each( spaces, function( spaceName, space ) {
+	var props = space.props,
+		cache = space.cache,
+		to = space.to,
+		from = space.from;
+
+	// Makes rgba() and hsla()
+	color.fn[ spaceName ] = function( value ) {
+
+		// Generate a cache for this space if it doesn't exist
+		if ( to && !this[ cache ] ) {
+			this[ cache ] = to( this._rgba );
+		}
+		if ( value === undefined ) {
+			return this[ cache ].slice();
+		}
+
+		var ret,
+			type = jQuery.type( value ),
+			arr = ( type === "array" || type === "object" ) ? value : arguments,
+			local = this[ cache ].slice();
+
+		each( props, function( key, prop ) {
+			var val = arr[ type === "object" ? key : prop.idx ];
+			if ( val == null ) {
+				val = local[ prop.idx ];
+			}
+			local[ prop.idx ] = clamp( val, prop );
+		} );
+
+		if ( from ) {
+			ret = color( from( local ) );
+			ret[ cache ] = local;
+			return ret;
+		} else {
+			return color( local );
+		}
+	};
+
+	// Makes red() green() blue() alpha() hue() saturation() lightness()
+	each( props, function( key, prop ) {
+
+		// Alpha is included in more than one space
+		if ( color.fn[ key ] ) {
+			return;
+		}
+		color.fn[ key ] = function( value ) {
+			var vtype = jQuery.type( value ),
+				fn = ( key === "alpha" ? ( this._hsla ? "hsla" : "rgba" ) : spaceName ),
+				local = this[ fn ](),
+				cur = local[ prop.idx ],
+				match;
+
+			if ( vtype === "undefined" ) {
+				return cur;
+			}
+
+			if ( vtype === "function" ) {
+				value = value.call( this, cur );
+				vtype = jQuery.type( value );
+			}
+			if ( value == null && prop.empty ) {
+				return this;
+			}
+			if ( vtype === "string" ) {
+				match = rplusequals.exec( value );
+				if ( match ) {
+					value = cur + parseFloat( match[ 2 ] ) * ( match[ 1 ] === "+" ? 1 : -1 );
+				}
+			}
+			local[ prop.idx ] = value;
+			return this[ fn ]( local );
+		};
+	} );
+} );
+
+// Add cssHook and .fx.step function for each named hook.
+// accept a space separated string of properties
+color.hook = function( hook ) {
+	var hooks = hook.split( " " );
+	each( hooks, function( i, hook ) {
+		jQuery.cssHooks[ hook ] = {
+			set: function( elem, value ) {
+				var parsed, curElem,
+					backgroundColor = "";
+
+				if ( value !== "transparent" && ( jQuery.type( value ) !== "string" ||
+						( parsed = stringParse( value ) ) ) ) {
+					value = color( parsed || value );
+					if ( !support.rgba && value._rgba[ 3 ] !== 1 ) {
+						curElem = hook === "backgroundColor" ? elem.parentNode : elem;
+						while (
+							( backgroundColor === "" || backgroundColor === "transparent" ) &&
+							curElem && curElem.style
+						) {
+							try {
+								backgroundColor = jQuery.css( curElem, "backgroundColor" );
+								curElem = curElem.parentNode;
+							} catch ( e ) {
+							}
+						}
+
+						value = value.blend( backgroundColor && backgroundColor !== "transparent" ?
+							backgroundColor :
+							"_default" );
+					}
+
+					value = value.toRgbaString();
+				}
+				try {
+					elem.style[ hook ] = value;
+				} catch ( e ) {
+
+					// Wrapped to prevent IE from throwing errors on "invalid" values like
+					// 'auto' or 'inherit'
+				}
+			}
+		};
+		jQuery.fx.step[ hook ] = function( fx ) {
+			if ( !fx.colorInit ) {
+				fx.start = color( fx.elem, hook );
+				fx.end = color( fx.end );
+				fx.colorInit = true;
+			}
+			jQuery.cssHooks[ hook ].set( fx.elem, fx.start.transition( fx.end, fx.pos ) );
+		};
+	} );
+
+};
+
+color.hook( stepHooks );
+
+jQuery.cssHooks.borderColor = {
+	expand: function( value ) {
+		var expanded = {};
+
+		each( [ "Top", "Right", "Bottom", "Left" ], function( i, part ) {
+			expanded[ "border" + part + "Color" ] = value;
+		} );
+		return expanded;
+	}
+};
+
+// Basic color names only.
+// Usage of any of the other color names requires adding yourself or including
+// jquery.color.svg-names.js.
+colors = jQuery.Color.names = {
+
+	// 4.1. Basic color keywords
+	aqua: "#00ffff",
+	black: "#000000",
+	blue: "#0000ff",
+	fuchsia: "#ff00ff",
+	gray: "#808080",
+	green: "#008000",
+	lime: "#00ff00",
+	maroon: "#800000",
+	navy: "#000080",
+	olive: "#808000",
+	purple: "#800080",
+	red: "#ff0000",
+	silver: "#c0c0c0",
+	teal: "#008080",
+	white: "#ffffff",
+	yellow: "#ffff00",
+
+	// 4.2.3. "transparent" color keyword
+	transparent: [ null, null, null, 0 ],
+
+	_default: "#ffffff"
+};
+
+} )( jQuery );
+
+/******************************************************************************/
+/****************************** CLASS ANIMATIONS ******************************/
+/******************************************************************************/
+( function() {
+
+var classAnimationActions = [ "add", "remove", "toggle" ],
+	shorthandStyles = {
+		border: 1,
+		borderBottom: 1,
+		borderColor: 1,
+		borderLeft: 1,
+		borderRight: 1,
+		borderTop: 1,
+		borderWidth: 1,
+		margin: 1,
+		padding: 1
+	};
+
+$.each(
+	[ "borderLeftStyle", "borderRightStyle", "borderBottomStyle", "borderTopStyle" ],
+	function( _, prop ) {
+		$.fx.step[ prop ] = function( fx ) {
+			if ( fx.end !== "none" && !fx.setAttr || fx.pos === 1 && !fx.setAttr ) {
+				jQuery.style( fx.elem, prop, fx.end );
+				fx.setAttr = true;
+			}
+		};
+	}
+);
+
+function getElementStyles( elem ) {
+	var key, len,
+		style = elem.ownerDocument.defaultView ?
+			elem.ownerDocument.defaultView.getComputedStyle( elem, null ) :
+			elem.currentStyle,
+		styles = {};
+
+	if ( style && style.length && style[ 0 ] && style[ style[ 0 ] ] ) {
+		len = style.length;
+		while ( len-- ) {
+			key = style[ len ];
+			if ( typeof style[ key ] === "string" ) {
+				styles[ $.camelCase( key ) ] = style[ key ];
+			}
+		}
+
+	// Support: Opera, IE <9
+	} else {
+		for ( key in style ) {
+			if ( typeof style[ key ] === "string" ) {
+				styles[ key ] = style[ key ];
+			}
+		}
+	}
+
+	return styles;
+}
+
+function styleDifference( oldStyle, newStyle ) {
+	var diff = {},
+		name, value;
+
+	for ( name in newStyle ) {
+		value = newStyle[ name ];
+		if ( oldStyle[ name ] !== value ) {
+			if ( !shorthandStyles[ name ] ) {
+				if ( $.fx.step[ name ] || !isNaN( parseFloat( value ) ) ) {
+					diff[ name ] = value;
+				}
+			}
+		}
+	}
+
+	return diff;
+}
+
+// Support: jQuery <1.8
+if ( !$.fn.addBack ) {
+	$.fn.addBack = function( selector ) {
+		return this.add( selector == null ?
+			this.prevObject : this.prevObject.filter( selector )
+		);
+	};
+}
+
+$.effects.animateClass = function( value, duration, easing, callback ) {
+	var o = $.speed( duration, easing, callback );
+
+	return this.queue( function() {
+		var animated = $( this ),
+			baseClass = animated.attr( "class" ) || "",
+			applyClassChange,
+			allAnimations = o.children ? animated.find( "*" ).addBack() : animated;
+
+		// Map the animated objects to store the original styles.
+		allAnimations = allAnimations.map( function() {
+			var el = $( this );
+			return {
+				el: el,
+				start: getElementStyles( this )
+			};
+		} );
+
+		// Apply class change
+		applyClassChange = function() {
+			$.each( classAnimationActions, function( i, action ) {
+				if ( value[ action ] ) {
+					animated[ action + "Class" ]( value[ action ] );
+				}
+			} );
+		};
+		applyClassChange();
+
+		// Map all animated objects again - calculate new styles and diff
+		allAnimations = allAnimations.map( function() {
+			this.end = getElementStyles( this.el[ 0 ] );
+			this.diff = styleDifference( this.start, this.end );
+			return this;
+		} );
+
+		// Apply original class
+		animated.attr( "class", baseClass );
+
+		// Map all animated objects again - this time collecting a promise
+		allAnimations = allAnimations.map( function() {
+			var styleInfo = this,
+				dfd = $.Deferred(),
+				opts = $.extend( {}, o, {
+					queue: false,
+					complete: function() {
+						dfd.resolve( styleInfo );
+					}
+				} );
+
+			this.el.animate( this.diff, opts );
+			return dfd.promise();
+		} );
+
+		// Once all animations have completed:
+		$.when.apply( $, allAnimations.get() ).done( function() {
+
+			// Set the final class
+			applyClassChange();
+
+			// For each animated element,
+			// clear all css properties that were animated
+			$.each( arguments, function() {
+				var el = this.el;
+				$.each( this.diff, function( key ) {
+					el.css( key, "" );
+				} );
+			} );
+
+			// This is guarnteed to be there if you use jQuery.speed()
+			// it also handles dequeuing the next anim...
+			o.complete.call( animated[ 0 ] );
+		} );
+	} );
+};
+
+$.fn.extend( {
+	addClass: ( function( orig ) {
+		return function( classNames, speed, easing, callback ) {
+			return speed ?
+				$.effects.animateClass.call( this,
+					{ add: classNames }, speed, easing, callback ) :
+				orig.apply( this, arguments );
+		};
+	} )( $.fn.addClass ),
+
+	removeClass: ( function( orig ) {
+		return function( classNames, speed, easing, callback ) {
+			return arguments.length > 1 ?
+				$.effects.animateClass.call( this,
+					{ remove: classNames }, speed, easing, callback ) :
+				orig.apply( this, arguments );
+		};
+	} )( $.fn.removeClass ),
+
+	toggleClass: ( function( orig ) {
+		return function( classNames, force, speed, easing, callback ) {
+			if ( typeof force === "boolean" || force === undefined ) {
+				if ( !speed ) {
+
+					// Without speed parameter
+					return orig.apply( this, arguments );
+				} else {
+					return $.effects.animateClass.call( this,
+						( force ? { add: classNames } : { remove: classNames } ),
+						speed, easing, callback );
+				}
+			} else {
+
+				// Without force parameter
+				return $.effects.animateClass.call( this,
+					{ toggle: classNames }, force, speed, easing );
+			}
+		};
+	} )( $.fn.toggleClass ),
+
+	switchClass: function( remove, add, speed, easing, callback ) {
+		return $.effects.animateClass.call( this, {
+			add: add,
+			remove: remove
+		}, speed, easing, callback );
+	}
+} );
+
+} )();
+
+/******************************************************************************/
+/*********************************** EFFECTS **********************************/
+/******************************************************************************/
+
+( function() {
+
+if ( $.expr && $.expr.filters && $.expr.filters.animated ) {
+	$.expr.filters.animated = ( function( orig ) {
+		return function( elem ) {
+			return !!$( elem ).data( dataSpaceAnimated ) || orig( elem );
+		};
+	} )( $.expr.filters.animated );
+}
+
+if ( $.uiBackCompat !== false ) {
+	$.extend( $.effects, {
+
+		// Saves a set of properties in a data storage
+		save: function( element, set ) {
+			var i = 0, length = set.length;
+			for ( ; i < length; i++ ) {
+				if ( set[ i ] !== null ) {
+					element.data( dataSpace + set[ i ], element[ 0 ].style[ set[ i ] ] );
+				}
+			}
+		},
+
+		// Restores a set of previously saved properties from a data storage
+		restore: function( element, set ) {
+			var val, i = 0, length = set.length;
+			for ( ; i < length; i++ ) {
+				if ( set[ i ] !== null ) {
+					val = element.data( dataSpace + set[ i ] );
+					element.css( set[ i ], val );
+				}
+			}
+		},
+
+		setMode: function( el, mode ) {
+			if ( mode === "toggle" ) {
+				mode = el.is( ":hidden" ) ? "show" : "hide";
+			}
+			return mode;
+		},
+
+		// Wraps the element around a wrapper that copies position properties
+		createWrapper: function( element ) {
+
+			// If the element is already wrapped, return it
+			if ( element.parent().is( ".ui-effects-wrapper" ) ) {
+				return element.parent();
+			}
+
+			// Wrap the element
+			var props = {
+					width: element.outerWidth( true ),
+					height: element.outerHeight( true ),
+					"float": element.css( "float" )
+				},
+				wrapper = $( "<div></div>" )
+					.addClass( "ui-effects-wrapper" )
+					.css( {
+						fontSize: "100%",
+						background: "transparent",
+						border: "none",
+						margin: 0,
+						padding: 0
+					} ),
+
+				// Store the size in case width/height are defined in % - Fixes #5245
+				size = {
+					width: element.width(),
+					height: element.height()
+				},
+				active = document.activeElement;
+
+			// Support: Firefox
+			// Firefox incorrectly exposes anonymous content
+			// https://bugzilla.mozilla.org/show_bug.cgi?id=561664
+			try {
+				active.id;
+			} catch ( e ) {
+				active = document.body;
+			}
+
+			element.wrap( wrapper );
+
+			// Fixes #7595 - Elements lose focus when wrapped.
+			if ( element[ 0 ] === active || $.contains( element[ 0 ], active ) ) {
+				$( active ).trigger( "focus" );
+			}
+
+			// Hotfix for jQuery 1.4 since some change in wrap() seems to actually
+			// lose the reference to the wrapped element
+			wrapper = element.parent();
+
+			// Transfer positioning properties to the wrapper
+			if ( element.css( "position" ) === "static" ) {
+				wrapper.css( { position: "relative" } );
+				element.css( { position: "relative" } );
+			} else {
+				$.extend( props, {
+					position: element.css( "position" ),
+					zIndex: element.css( "z-index" )
+				} );
+				$.each( [ "top", "left", "bottom", "right" ], function( i, pos ) {
+					props[ pos ] = element.css( pos );
+					if ( isNaN( parseInt( props[ pos ], 10 ) ) ) {
+						props[ pos ] = "auto";
+					}
+				} );
+				element.css( {
+					position: "relative",
+					top: 0,
+					left: 0,
+					right: "auto",
+					bottom: "auto"
+				} );
+			}
+			element.css( size );
+
+			return wrapper.css( props ).show();
+		},
+
+		removeWrapper: function( element ) {
+			var active = document.activeElement;
+
+			if ( element.parent().is( ".ui-effects-wrapper" ) ) {
+				element.parent().replaceWith( element );
+
+				// Fixes #7595 - Elements lose focus when wrapped.
+				if ( element[ 0 ] === active || $.contains( element[ 0 ], active ) ) {
+					$( active ).trigger( "focus" );
+				}
+			}
+
+			return element;
+		}
+	} );
+}
+
+$.extend( $.effects, {
+	version: "1.12.0",
+
+	define: function( name, mode, effect ) {
+		if ( !effect ) {
+			effect = mode;
+			mode = "effect";
+		}
+
+		$.effects.effect[ name ] = effect;
+		$.effects.effect[ name ].mode = mode;
+
+		return effect;
+	},
+
+	scaledDimensions: function( element, percent, direction ) {
+		if ( percent === 0 ) {
+			return {
+				height: 0,
+				width: 0,
+				outerHeight: 0,
+				outerWidth: 0
+			};
+		}
+
+		var x = direction !== "horizontal" ? ( ( percent || 100 ) / 100 ) : 1,
+			y = direction !== "vertical" ? ( ( percent || 100 ) / 100 ) : 1;
+
+		return {
+			height: element.height() * y,
+			width: element.width() * x,
+			outerHeight: element.outerHeight() * y,
+			outerWidth: element.outerWidth() * x
+		};
+
+	},
+
+	clipToBox: function( animation ) {
+		return {
+			width: animation.clip.right - animation.clip.left,
+			height: animation.clip.bottom - animation.clip.top,
+			left: animation.clip.left,
+			top: animation.clip.top
+		};
+	},
+
+	// Injects recently queued functions to be first in line (after "inprogress")
+	unshift: function( element, queueLength, count ) {
+		var queue = element.queue();
+
+		if ( queueLength > 1 ) {
+			queue.splice.apply( queue,
+				[ 1, 0 ].concat( queue.splice( queueLength, count ) ) );
+		}
+		element.dequeue();
+	},
+
+	saveStyle: function( element ) {
+		element.data( dataSpaceStyle, element[ 0 ].style.cssText );
+	},
+
+	restoreStyle: function( element ) {
+		element[ 0 ].style.cssText = element.data( dataSpaceStyle ) || "";
+		element.removeData( dataSpaceStyle );
+	},
+
+	mode: function( element, mode ) {
+		var hidden = element.is( ":hidden" );
+
+		if ( mode === "toggle" ) {
+			mode = hidden ? "show" : "hide";
+		}
+		if ( hidden ? mode === "hide" : mode === "show" ) {
+			mode = "none";
+		}
+		return mode;
+	},
+
+	// Translates a [top,left] array into a baseline value
+	getBaseline: function( origin, original ) {
+		var y, x;
+
+		switch ( origin[ 0 ] ) {
+		case "top":
+			y = 0;
+			break;
+		case "middle":
+			y = 0.5;
+			break;
+		case "bottom":
+			y = 1;
+			break;
+		default:
+			y = origin[ 0 ] / original.height;
+		}
+
+		switch ( origin[ 1 ] ) {
+		case "left":
+			x = 0;
+			break;
+		case "center":
+			x = 0.5;
+			break;
+		case "right":
+			x = 1;
+			break;
+		default:
+			x = origin[ 1 ] / original.width;
+		}
+
+		return {
+			x: x,
+			y: y
+		};
+	},
+
+	// Creates a placeholder element so that the original element can be made absolute
+	createPlaceholder: function( element ) {
+		var placeholder,
+			cssPosition = element.css( "position" ),
+			position = element.position();
+
+		// Lock in margins first to account for form elements, which
+		// will change margin if you explicitly set height
+		// see: http://jsfiddle.net/JZSMt/3/ https://bugs.webkit.org/show_bug.cgi?id=107380
+		// Support: Safari
+		element.css( {
+			marginTop: element.css( "marginTop" ),
+			marginBottom: element.css( "marginBottom" ),
+			marginLeft: element.css( "marginLeft" ),
+			marginRight: element.css( "marginRight" )
+		} )
+		.outerWidth( element.outerWidth() )
+		.outerHeight( element.outerHeight() );
+
+		if ( /^(static|relative)/.test( cssPosition ) ) {
+			cssPosition = "absolute";
+
+			placeholder = $( "<" + element[ 0 ].nodeName + ">" ).insertAfter( element ).css( {
+
+				// Convert inline to inline block to account for inline elements
+				// that turn to inline block based on content (like img)
+				display: /^(inline|ruby)/.test( element.css( "display" ) ) ?
+					"inline-block" :
+					"block",
+				visibility: "hidden",
+
+				// Margins need to be set to account for margin collapse
+				marginTop: element.css( "marginTop" ),
+				marginBottom: element.css( "marginBottom" ),
+				marginLeft: element.css( "marginLeft" ),
+				marginRight: element.css( "marginRight" ),
+				"float": element.css( "float" )
+			} )
+			.outerWidth( element.outerWidth() )
+			.outerHeight( element.outerHeight() )
+			.addClass( "ui-effects-placeholder" );
+
+			element.data( dataSpace + "placeholder", placeholder );
+		}
+
+		element.css( {
+			position: cssPosition,
+			left: position.left,
+			top: position.top
+		} );
+
+		return placeholder;
+	},
+
+	removePlaceholder: function( element ) {
+		var dataKey = dataSpace + "placeholder",
+				placeholder = element.data( dataKey );
+
+		if ( placeholder ) {
+			placeholder.remove();
+			element.removeData( dataKey );
+		}
+	},
+
+	// Removes a placeholder if it exists and restores
+	// properties that were modified during placeholder creation
+	cleanUp: function( element ) {
+		$.effects.restoreStyle( element );
+		$.effects.removePlaceholder( element );
+	},
+
+	setTransition: function( element, list, factor, value ) {
+		value = value || {};
+		$.each( list, function( i, x ) {
+			var unit = element.cssUnit( x );
+			if ( unit[ 0 ] > 0 ) {
+				value[ x ] = unit[ 0 ] * factor + unit[ 1 ];
+			}
+		} );
+		return value;
+	}
+} );
+
+// Return an effect options object for the given parameters:
+function _normalizeArguments( effect, options, speed, callback ) {
+
+	// Allow passing all options as the first parameter
+	if ( $.isPlainObject( effect ) ) {
+		options = effect;
+		effect = effect.effect;
+	}
+
+	// Convert to an object
+	effect = { effect: effect };
+
+	// Catch (effect, null, ...)
+	if ( options == null ) {
+		options = {};
+	}
+
+	// Catch (effect, callback)
+	if ( $.isFunction( options ) ) {
+		callback = options;
+		speed = null;
+		options = {};
+	}
+
+	// Catch (effect, speed, ?)
+	if ( typeof options === "number" || $.fx.speeds[ options ] ) {
+		callback = speed;
+		speed = options;
+		options = {};
+	}
+
+	// Catch (effect, options, callback)
+	if ( $.isFunction( speed ) ) {
+		callback = speed;
+		speed = null;
+	}
+
+	// Add options to effect
+	if ( options ) {
+		$.extend( effect, options );
+	}
+
+	speed = speed || options.duration;
+	effect.duration = $.fx.off ? 0 :
+		typeof speed === "number" ? speed :
+		speed in $.fx.speeds ? $.fx.speeds[ speed ] :
+		$.fx.speeds._default;
+
+	effect.complete = callback || options.complete;
+
+	return effect;
+}
+
+function standardAnimationOption( option ) {
+
+	// Valid standard speeds (nothing, number, named speed)
+	if ( !option || typeof option === "number" || $.fx.speeds[ option ] ) {
+		return true;
+	}
+
+	// Invalid strings - treat as "normal" speed
+	if ( typeof option === "string" && !$.effects.effect[ option ] ) {
+		return true;
+	}
+
+	// Complete callback
+	if ( $.isFunction( option ) ) {
+		return true;
+	}
+
+	// Options hash (but not naming an effect)
+	if ( typeof option === "object" && !option.effect ) {
+		return true;
+	}
+
+	// Didn't match any standard API
+	return false;
+}
+
+$.fn.extend( {
+	effect: function( /* effect, options, speed, callback */ ) {
+		var args = _normalizeArguments.apply( this, arguments ),
+			effectMethod = $.effects.effect[ args.effect ],
+			defaultMode = effectMethod.mode,
+			queue = args.queue,
+			queueName = queue || "fx",
+			complete = args.complete,
+			mode = args.mode,
+			modes = [],
+			prefilter = function( next ) {
+				var el = $( this ),
+					normalizedMode = $.effects.mode( el, mode ) || defaultMode;
+
+				// Sentinel for duck-punching the :animated psuedo-selector
+				el.data( dataSpaceAnimated, true );
+
+				// Save effect mode for later use,
+				// we can't just call $.effects.mode again later,
+				// as the .show() below destroys the initial state
+				modes.push( normalizedMode );
+
+				// See $.uiBackCompat inside of run() for removal of defaultMode in 1.13
+				if ( defaultMode && ( normalizedMode === "show" ||
+						( normalizedMode === defaultMode && normalizedMode === "hide" ) ) ) {
+					el.show();
+				}
+
+				if ( !defaultMode || normalizedMode !== "none" ) {
+					$.effects.saveStyle( el );
+				}
+
+				if ( $.isFunction( next ) ) {
+					next();
+				}
+			};
+
+		if ( $.fx.off || !effectMethod ) {
+
+			// Delegate to the original method (e.g., .show()) if possible
+			if ( mode ) {
+				return this[ mode ]( args.duration, complete );
+			} else {
+				return this.each( function() {
+					if ( complete ) {
+						complete.call( this );
+					}
+				} );
+			}
+		}
+
+		function run( next ) {
+			var elem = $( this );
+
+			function cleanup() {
+				elem.removeData( dataSpaceAnimated );
+
+				$.effects.cleanUp( elem );
+
+				if ( args.mode === "hide" ) {
+					elem.hide();
+				}
+
+				done();
+			}
+
+			function done() {
+				if ( $.isFunction( complete ) ) {
+					complete.call( elem[ 0 ] );
+				}
+
+				if ( $.isFunction( next ) ) {
+					next();
+				}
+			}
+
+			// Override mode option on a per element basis,
+			// as toggle can be either show or hide depending on element state
+			args.mode = modes.shift();
+
+			if ( $.uiBackCompat !== false && !defaultMode ) {
+				if ( elem.is( ":hidden" ) ? mode === "hide" : mode === "show" ) {
+
+					// Call the core method to track "olddisplay" properly
+					elem[ mode ]();
+					done();
+				} else {
+					effectMethod.call( elem[ 0 ], args, done );
+				}
+			} else {
+				if ( args.mode === "none" ) {
+
+					// Call the core method to track "olddisplay" properly
+					elem[ mode ]();
+					done();
+				} else {
+					effectMethod.call( elem[ 0 ], args, cleanup );
+				}
+			}
+		}
+
+		// Run prefilter on all elements first to ensure that
+		// any showing or hiding happens before placeholder creation,
+		// which ensures that any layout changes are correctly captured.
+		return queue === false ?
+			this.each( prefilter ).each( run ) :
+			this.queue( queueName, prefilter ).queue( queueName, run );
+	},
+
+	show: ( function( orig ) {
+		return function( option ) {
+			if ( standardAnimationOption( option ) ) {
+				return orig.apply( this, arguments );
+			} else {
+				var args = _normalizeArguments.apply( this, arguments );
+				args.mode = "show";
+				return this.effect.call( this, args );
+			}
+		};
+	} )( $.fn.show ),
+
+	hide: ( function( orig ) {
+		return function( option ) {
+			if ( standardAnimationOption( option ) ) {
+				return orig.apply( this, arguments );
+			} else {
+				var args = _normalizeArguments.apply( this, arguments );
+				args.mode = "hide";
+				return this.effect.call( this, args );
+			}
+		};
+	} )( $.fn.hide ),
+
+	toggle: ( function( orig ) {
+		return function( option ) {
+			if ( standardAnimationOption( option ) || typeof option === "boolean" ) {
+				return orig.apply( this, arguments );
+			} else {
+				var args = _normalizeArguments.apply( this, arguments );
+				args.mode = "toggle";
+				return this.effect.call( this, args );
+			}
+		};
+	} )( $.fn.toggle ),
+
+	cssUnit: function( key ) {
+		var style = this.css( key ),
+			val = [];
+
+		$.each( [ "em", "px", "%", "pt" ], function( i, unit ) {
+			if ( style.indexOf( unit ) > 0 ) {
+				val = [ parseFloat( style ), unit ];
+			}
+		} );
+		return val;
+	},
+
+	cssClip: function( clipObj ) {
+		if ( clipObj ) {
+			return this.css( "clip", "rect(" + clipObj.top + "px " + clipObj.right + "px " +
+				clipObj.bottom + "px " + clipObj.left + "px)" );
+		}
+		return parseClip( this.css( "clip" ), this );
+	},
+
+	transfer: function( options, done ) {
+		var element = $( this ),
+			target = $( options.to ),
+			targetFixed = target.css( "position" ) === "fixed",
+			body = $( "body" ),
+			fixTop = targetFixed ? body.scrollTop() : 0,
+			fixLeft = targetFixed ? body.scrollLeft() : 0,
+			endPosition = target.offset(),
+			animation = {
+				top: endPosition.top - fixTop,
+				left: endPosition.left - fixLeft,
+				height: target.innerHeight(),
+				width: target.innerWidth()
+			},
+			startPosition = element.offset(),
+			transfer = $( "<div class='ui-effects-transfer'></div>" )
+				.appendTo( "body" )
+				.addClass( options.className )
+				.css( {
+					top: startPosition.top - fixTop,
+					left: startPosition.left - fixLeft,
+					height: element.innerHeight(),
+					width: element.innerWidth(),
+					position: targetFixed ? "fixed" : "absolute"
+				} )
+				.animate( animation, options.duration, options.easing, function() {
+					transfer.remove();
+					if ( $.isFunction( done ) ) {
+						done();
+					}
+				} );
+	}
+} );
+
+function parseClip( str, element ) {
+		var outerWidth = element.outerWidth(),
+			outerHeight = element.outerHeight(),
+			clipRegex = /^rect\((-?\d*\.?\d*px|-?\d+%|auto),?\s*(-?\d*\.?\d*px|-?\d+%|auto),?\s*(-?\d*\.?\d*px|-?\d+%|auto),?\s*(-?\d*\.?\d*px|-?\d+%|auto)\)$/,
+			values = clipRegex.exec( str ) || [ "", 0, outerWidth, outerHeight, 0 ];
+
+		return {
+			top: parseFloat( values[ 1 ] ) || 0,
+			right: values[ 2 ] === "auto" ? outerWidth : parseFloat( values[ 2 ] ),
+			bottom: values[ 3 ] === "auto" ? outerHeight : parseFloat( values[ 3 ] ),
+			left: parseFloat( values[ 4 ] ) || 0
+		};
+}
+
+$.fx.step.clip = function( fx ) {
+	if ( !fx.clipInit ) {
+		fx.start = $( fx.elem ).cssClip();
+		if ( typeof fx.end === "string" ) {
+			fx.end = parseClip( fx.end, fx.elem );
+		}
+		fx.clipInit = true;
+	}
+
+	$( fx.elem ).cssClip( {
+		top: fx.pos * ( fx.end.top - fx.start.top ) + fx.start.top,
+		right: fx.pos * ( fx.end.right - fx.start.right ) + fx.start.right,
+		bottom: fx.pos * ( fx.end.bottom - fx.start.bottom ) + fx.start.bottom,
+		left: fx.pos * ( fx.end.left - fx.start.left ) + fx.start.left
+	} );
+};
+
+} )();
+
+/******************************************************************************/
+/*********************************** EASING ***********************************/
+/******************************************************************************/
+
+( function() {
+
+// Based on easing equations from Robert Penner (http://www.robertpenner.com/easing)
+
+var baseEasings = {};
+
+$.each( [ "Quad", "Cubic", "Quart", "Quint", "Expo" ], function( i, name ) {
+	baseEasings[ name ] = function( p ) {
+		return Math.pow( p, i + 2 );
+	};
+} );
+
+$.extend( baseEasings, {
+	Sine: function( p ) {
+		return 1 - Math.cos( p * Math.PI / 2 );
+	},
+	Circ: function( p ) {
+		return 1 - Math.sqrt( 1 - p * p );
+	},
+	Elastic: function( p ) {
+		return p === 0 || p === 1 ? p :
+			-Math.pow( 2, 8 * ( p - 1 ) ) * Math.sin( ( ( p - 1 ) * 80 - 7.5 ) * Math.PI / 15 );
+	},
+	Back: function( p ) {
+		return p * p * ( 3 * p - 2 );
+	},
+	Bounce: function( p ) {
+		var pow2,
+			bounce = 4;
+
+		while ( p < ( ( pow2 = Math.pow( 2, --bounce ) ) - 1 ) / 11 ) {}
+		return 1 / Math.pow( 4, 3 - bounce ) - 7.5625 * Math.pow( ( pow2 * 3 - 2 ) / 22 - p, 2 );
+	}
+} );
+
+$.each( baseEasings, function( name, easeIn ) {
+	$.easing[ "easeIn" + name ] = easeIn;
+	$.easing[ "easeOut" + name ] = function( p ) {
+		return 1 - easeIn( 1 - p );
+	};
+	$.easing[ "easeInOut" + name ] = function( p ) {
+		return p < 0.5 ?
+			easeIn( p * 2 ) / 2 :
+			1 - easeIn( p * -2 + 2 ) / 2;
+	};
+} );
+
+} )();
+
+var effect = $.effects;
+
+
+/*!
+ * jQuery UI Effects Blind 1.12.0
+ * http://jqueryui.com
+ *
+ * Copyright jQuery Foundation and other contributors
+ * Released under the MIT license.
+ * http://jquery.org/license
+ */
+
+//>>label: Blind Effect
+//>>group: Effects
+//>>description: Blinds the element.
+//>>docs: http://api.jqueryui.com/blind-effect/
+//>>demos: http://jqueryui.com/effect/
+
+
+
+var effectsEffectBlind = $.effects.define( "blind", "hide", function( options, done ) {
+	var map = {
+			up: [ "bottom", "top" ],
+			vertical: [ "bottom", "top" ],
+			down: [ "top", "bottom" ],
+			left: [ "right", "left" ],
+			horizontal: [ "right", "left" ],
+			right: [ "left", "right" ]
+		},
+		element = $( this ),
+		direction = options.direction || "up",
+		start = element.cssClip(),
+		animate = { clip: $.extend( {}, start ) },
+		placeholder = $.effects.createPlaceholder( element );
+
+	animate.clip[ map[ direction ][ 0 ] ] = animate.clip[ map[ direction ][ 1 ] ];
+
+	if ( options.mode === "show" ) {
+		element.cssClip( animate.clip );
+		if ( placeholder ) {
+			placeholder.css( $.effects.clipToBox( animate ) );
+		}
+
+		animate.clip = start;
+	}
+
+	if ( placeholder ) {
+		placeholder.animate( $.effects.clipToBox( animate ), options.duration, options.easing );
+	}
+
+	element.animate( animate, {
+		queue: false,
+		duration: options.duration,
+		easing: options.easing,
+		complete: done
+	} );
+} );
+
+
+/*!
+ * jQuery UI Effects Bounce 1.12.0
+ * http://jqueryui.com
+ *
+ * Copyright jQuery Foundation and other contributors
+ * Released under the MIT license.
+ * http://jquery.org/license
+ */
+
+//>>label: Bounce Effect
+//>>group: Effects
+//>>description: Bounces an element horizontally or vertically n times.
+//>>docs: http://api.jqueryui.com/bounce-effect/
+//>>demos: http://jqueryui.com/effect/
+
+
+
+var effectsEffectBounce = $.effects.define( "bounce", function( options, done ) {
+	var upAnim, downAnim, refValue,
+		element = $( this ),
+
+		// Defaults:
+		mode = options.mode,
+		hide = mode === "hide",
+		show = mode === "show",
+		direction = options.direction || "up",
+		distance = options.distance,
+		times = options.times || 5,
+
+		// Number of internal animations
+		anims = times * 2 + ( show || hide ? 1 : 0 ),
+		speed = options.duration / anims,
+		easing = options.easing,
+
+		// Utility:
+		ref = ( direction === "up" || direction === "down" ) ? "top" : "left",
+		motion = ( direction === "up" || direction === "left" ),
+		i = 0,
+
+		queuelen = element.queue().length;
+
+	$.effects.createPlaceholder( element );
+
+	refValue = element.css( ref );
+
+	// Default distance for the BIGGEST bounce is the outer Distance / 3
+	if ( !distance ) {
+		distance = element[ ref === "top" ? "outerHeight" : "outerWidth" ]() / 3;
+	}
+
+	if ( show ) {
+		downAnim = { opacity: 1 };
+		downAnim[ ref ] = refValue;
+
+		// If we are showing, force opacity 0 and set the initial position
+		// then do the "first" animation
+		element
+			.css( "opacity", 0 )
+			.css( ref, motion ? -distance * 2 : distance * 2 )
+			.animate( downAnim, speed, easing );
+	}
+
+	// Start at the smallest distance if we are hiding
+	if ( hide ) {
+		distance = distance / Math.pow( 2, times - 1 );
+	}
+
+	downAnim = {};
+	downAnim[ ref ] = refValue;
+
+	// Bounces up/down/left/right then back to 0 -- times * 2 animations happen here
+	for ( ; i < times; i++ ) {
+		upAnim = {};
+		upAnim[ ref ] = ( motion ? "-=" : "+=" ) + distance;
+
+		element
+			.animate( upAnim, speed, easing )
+			.animate( downAnim, speed, easing );
+
+		distance = hide ? distance * 2 : distance / 2;
+	}
+
+	// Last Bounce when Hiding
+	if ( hide ) {
+		upAnim = { opacity: 0 };
+		upAnim[ ref ] = ( motion ? "-=" : "+=" ) + distance;
+
+		element.animate( upAnim, speed, easing );
+	}
+
+	element.queue( done );
+
+	$.effects.unshift( element, queuelen, anims + 1 );
+} );
+
+
+/*!
+ * jQuery UI Effects Clip 1.12.0
+ * http://jqueryui.com
+ *
+ * Copyright jQuery Foundation and other contributors
+ * Released under the MIT license.
+ * http://jquery.org/license
+ */
+
+//>>label: Clip Effect
+//>>group: Effects
+//>>description: Clips the element on and off like an old TV.
+//>>docs: http://api.jqueryui.com/clip-effect/
+//>>demos: http://jqueryui.com/effect/
+
+
+
+var effectsEffectClip = $.effects.define( "clip", "hide", function( options, done ) {
+	var start,
+		animate = {},
+		element = $( this ),
+		direction = options.direction || "vertical",
+		both = direction === "both",
+		horizontal = both || direction === "horizontal",
+		vertical = both || direction === "vertical";
+
+	start = element.cssClip();
+	animate.clip = {
+		top: vertical ? ( start.bottom - start.top ) / 2 : start.top,
+		right: horizontal ? ( start.right - start.left ) / 2 : start.right,
+		bottom: vertical ? ( start.bottom - start.top ) / 2 : start.bottom,
+		left: horizontal ? ( start.right - start.left ) / 2 : start.left
+	};
+
+	$.effects.createPlaceholder( element );
+
+	if ( options.mode === "show" ) {
+		element.cssClip( animate.clip );
+		animate.clip = start;
+	}
+
+	element.animate( animate, {
+		queue: false,
+		duration: options.duration,
+		easing: options.easing,
+		complete: done
+	} );
+
+} );
+
+
+/*!
+ * jQuery UI Effects Drop 1.12.0
+ * http://jqueryui.com
+ *
+ * Copyright jQuery Foundation and other contributors
+ * Released under the MIT license.
+ * http://jquery.org/license
+ */
+
+//>>label: Drop Effect
+//>>group: Effects
+//>>description: Moves an element in one direction and hides it at the same time.
+//>>docs: http://api.jqueryui.com/drop-effect/
+//>>demos: http://jqueryui.com/effect/
+
+
+
+var effectsEffectDrop = $.effects.define( "drop", "hide", function( options, done ) {
+
+	var distance,
+		element = $( this ),
+		mode = options.mode,
+		show = mode === "show",
+		direction = options.direction || "left",
+		ref = ( direction === "up" || direction === "down" ) ? "top" : "left",
+		motion = ( direction === "up" || direction === "left" ) ? "-=" : "+=",
+		oppositeMotion = ( motion === "+=" ) ? "-=" : "+=",
+		animation = {
+			opacity: 0
+		};
+
+	$.effects.createPlaceholder( element );
+
+	distance = options.distance ||
+		element[ ref === "top" ? "outerHeight" : "outerWidth" ]( true ) / 2;
+
+	animation[ ref ] = motion + distance;
+
+	if ( show ) {
+		element.css( animation );
+
+		animation[ ref ] = oppositeMotion + distance;
+		animation.opacity = 1;
+	}
+
+	// Animate
+	element.animate( animation, {
+		queue: false,
+		duration: options.duration,
+		easing: options.easing,
+		complete: done
+	} );
+} );
+
+
+/*!
+ * jQuery UI Effects Explode 1.12.0
+ * http://jqueryui.com
+ *
+ * Copyright jQuery Foundation and other contributors
+ * Released under the MIT license.
+ * http://jquery.org/license
+ */
+
+//>>label: Explode Effect
+//>>group: Effects
+// jscs:disable maximumLineLength
+//>>description: Explodes an element in all directions into n pieces. Implodes an element to its original wholeness.
+// jscs:enable maximumLineLength
+//>>docs: http://api.jqueryui.com/explode-effect/
+//>>demos: http://jqueryui.com/effect/
+
+
+
+var effectsEffectExplode = $.effects.define( "explode", "hide", function( options, done ) {
+
+	var i, j, left, top, mx, my,
+		rows = options.pieces ? Math.round( Math.sqrt( options.pieces ) ) : 3,
+		cells = rows,
+		element = $( this ),
+		mode = options.mode,
+		show = mode === "show",
+
+		// Show and then visibility:hidden the element before calculating offset
+		offset = element.show().css( "visibility", "hidden" ).offset(),
+
+		// Width and height of a piece
+		width = Math.ceil( element.outerWidth() / cells ),
+		height = Math.ceil( element.outerHeight() / rows ),
+		pieces = [];
+
+	// Children animate complete:
+	function childComplete() {
+		pieces.push( this );
+		if ( pieces.length === rows * cells ) {
+			animComplete();
+		}
+	}
+
+	// Clone the element for each row and cell.
+	for ( i = 0; i < rows; i++ ) { // ===>
+		top = offset.top + i * height;
+		my = i - ( rows - 1 ) / 2;
+
+		for ( j = 0; j < cells; j++ ) { // |||
+			left = offset.left + j * width;
+			mx = j - ( cells - 1 ) / 2;
+
+			// Create a clone of the now hidden main element that will be absolute positioned
+			// within a wrapper div off the -left and -top equal to size of our pieces
+			element
+				.clone()
+				.appendTo( "body" )
+				.wrap( "<div></div>" )
+				.css( {
+					position: "absolute",
+					visibility: "visible",
+					left: -j * width,
+					top: -i * height
+				} )
+
+				// Select the wrapper - make it overflow: hidden and absolute positioned based on
+				// where the original was located +left and +top equal to the size of pieces
+				.parent()
+					.addClass( "ui-effects-explode" )
+					.css( {
+						position: "absolute",
+						overflow: "hidden",
+						width: width,
+						height: height,
+						left: left + ( show ? mx * width : 0 ),
+						top: top + ( show ? my * height : 0 ),
+						opacity: show ? 0 : 1
+					} )
+					.animate( {
+						left: left + ( show ? 0 : mx * width ),
+						top: top + ( show ? 0 : my * height ),
+						opacity: show ? 1 : 0
+					}, options.duration || 500, options.easing, childComplete );
+		}
+	}
+
+	function animComplete() {
+		element.css( {
+			visibility: "visible"
+		} );
+		$( pieces ).remove();
+		done();
+	}
+} );
+
+
+/*!
+ * jQuery UI Effects Fade 1.12.0
+ * http://jqueryui.com
+ *
+ * Copyright jQuery Foundation and other contributors
+ * Released under the MIT license.
+ * http://jquery.org/license
+ */
+
+//>>label: Fade Effect
+//>>group: Effects
+//>>description: Fades the element.
+//>>docs: http://api.jqueryui.com/fade-effect/
+//>>demos: http://jqueryui.com/effect/
+
+
+
+var effectsEffectFade = $.effects.define( "fade", "toggle", function( options, done ) {
+	var show = options.mode === "show";
+
+	$( this )
+		.css( "opacity", show ? 0 : 1 )
+		.animate( {
+			opacity: show ? 1 : 0
+		}, {
+			queue: false,
+			duration: options.duration,
+			easing: options.easing,
+			complete: done
+		} );
+} );
+
+
+/*!
+ * jQuery UI Effects Fold 1.12.0
+ * http://jqueryui.com
+ *
+ * Copyright jQuery Foundation and other contributors
+ * Released under the MIT license.
+ * http://jquery.org/license
+ */
+
+//>>label: Fold Effect
+//>>group: Effects
+//>>description: Folds an element first horizontally and then vertically.
+//>>docs: http://api.jqueryui.com/fold-effect/
+//>>demos: http://jqueryui.com/effect/
+
+
+
+var effectsEffectFold = $.effects.define( "fold", "hide", function( options, done ) {
+
+	// Create element
+	var element = $( this ),
+		mode = options.mode,
+		show = mode === "show",
+		hide = mode === "hide",
+		size = options.size || 15,
+		percent = /([0-9]+)%/.exec( size ),
+		horizFirst = !!options.horizFirst,
+		ref = horizFirst ? [ "right", "bottom" ] : [ "bottom", "right" ],
+		duration = options.duration / 2,
+
+		placeholder = $.effects.createPlaceholder( element ),
+
+		start = element.cssClip(),
+		animation1 = { clip: $.extend( {}, start ) },
+		animation2 = { clip: $.extend( {}, start ) },
+
+		distance = [ start[ ref[ 0 ] ], start[ ref[ 1 ] ] ],
+
+		queuelen = element.queue().length;
+
+	if ( percent ) {
+		size = parseInt( percent[ 1 ], 10 ) / 100 * distance[ hide ? 0 : 1 ];
+	}
+	animation1.clip[ ref[ 0 ] ] = size;
+	animation2.clip[ ref[ 0 ] ] = size;
+	animation2.clip[ ref[ 1 ] ] = 0;
+
+	if ( show ) {
+		element.cssClip( animation2.clip );
+		if ( placeholder ) {
+			placeholder.css( $.effects.clipToBox( animation2 ) );
+		}
+
+		animation2.clip = start;
+	}
+
+	// Animate
+	element
+		.queue( function( next ) {
+			if ( placeholder ) {
+				placeholder
+					.animate( $.effects.clipToBox( animation1 ), duration, options.easing )
+					.animate( $.effects.clipToBox( animation2 ), duration, options.easing );
+			}
+
+			next();
+		} )
+		.animate( animation1, duration, options.easing )
+		.animate( animation2, duration, options.easing )
+		.queue( done );
+
+	$.effects.unshift( element, queuelen, 4 );
+} );
+
+
+/*!
+ * jQuery UI Effects Highlight 1.12.0
+ * http://jqueryui.com
+ *
+ * Copyright jQuery Foundation and other contributors
+ * Released under the MIT license.
+ * http://jquery.org/license
+ */
+
+//>>label: Highlight Effect
+//>>group: Effects
+//>>description: Highlights the background of an element in a defined color for a custom duration.
+//>>docs: http://api.jqueryui.com/highlight-effect/
+//>>demos: http://jqueryui.com/effect/
+
+
+
+var effectsEffectHighlight = $.effects.define( "highlight", "show", function( options, done ) {
+	var element = $( this ),
+		animation = {
+			backgroundColor: element.css( "backgroundColor" )
+		};
+
+	if ( options.mode === "hide" ) {
+		animation.opacity = 0;
+	}
+
+	$.effects.saveStyle( element );
+
+	element
+		.css( {
+			backgroundImage: "none",
+			backgroundColor: options.color || "#ffff99"
+		} )
+		.animate( animation, {
+			queue: false,
+			duration: options.duration,
+			easing: options.easing,
+			complete: done
+		} );
+} );
+
+
+/*!
+ * jQuery UI Effects Size 1.12.0
+ * http://jqueryui.com
+ *
+ * Copyright jQuery Foundation and other contributors
+ * Released under the MIT license.
+ * http://jquery.org/license
+ */
+
+//>>label: Size Effect
+//>>group: Effects
+//>>description: Resize an element to a specified width and height.
+//>>docs: http://api.jqueryui.com/size-effect/
+//>>demos: http://jqueryui.com/effect/
+
+
+
+var effectsEffectSize = $.effects.define( "size", function( options, done ) {
+
+	// Create element
+	var baseline, factor, temp,
+		element = $( this ),
+
+		// Copy for children
+		cProps = [ "fontSize" ],
+		vProps = [ "borderTopWidth", "borderBottomWidth", "paddingTop", "paddingBottom" ],
+		hProps = [ "borderLeftWidth", "borderRightWidth", "paddingLeft", "paddingRight" ],
+
+		// Set options
+		mode = options.mode,
+		restore = mode !== "effect",
+		scale = options.scale || "both",
+		origin = options.origin || [ "middle", "center" ],
+		position = element.css( "position" ),
+		pos = element.position(),
+		original = $.effects.scaledDimensions( element ),
+		from = options.from || original,
+		to = options.to || $.effects.scaledDimensions( element, 0 );
+
+	$.effects.createPlaceholder( element );
+
+	if ( mode === "show" ) {
+		temp = from;
+		from = to;
+		to = temp;
+	}
+
+	// Set scaling factor
+	factor = {
+		from: {
+			y: from.height / original.height,
+			x: from.width / original.width
+		},
+		to: {
+			y: to.height / original.height,
+			x: to.width / original.width
+		}
+	};
+
+	// Scale the css box
+	if ( scale === "box" || scale === "both" ) {
+
+		// Vertical props scaling
+		if ( factor.from.y !== factor.to.y ) {
+			from = $.effects.setTransition( element, vProps, factor.from.y, from );
+			to = $.effects.setTransition( element, vProps, factor.to.y, to );
+		}
+
+		// Horizontal props scaling
+		if ( factor.from.x !== factor.to.x ) {
+			from = $.effects.setTransition( element, hProps, factor.from.x, from );
+			to = $.effects.setTransition( element, hProps, factor.to.x, to );
+		}
+	}
+
+	// Scale the content
+	if ( scale === "content" || scale === "both" ) {
+
+		// Vertical props scaling
+		if ( factor.from.y !== factor.to.y ) {
+			from = $.effects.setTransition( element, cProps, factor.from.y, from );
+			to = $.effects.setTransition( element, cProps, factor.to.y, to );
+		}
+	}
+
+	// Adjust the position properties based on the provided origin points
+	if ( origin ) {
+		baseline = $.effects.getBaseline( origin, original );
+		from.top = ( original.outerHeight - from.outerHeight ) * baseline.y + pos.top;
+		from.left = ( original.outerWidth - from.outerWidth ) * baseline.x + pos.left;
+		to.top = ( original.outerHeight - to.outerHeight ) * baseline.y + pos.top;
+		to.left = ( original.outerWidth - to.outerWidth ) * baseline.x + pos.left;
+	}
+	element.css( from );
+
+	// Animate the children if desired
+	if ( scale === "content" || scale === "both" ) {
+
+		vProps = vProps.concat( [ "marginTop", "marginBottom" ] ).concat( cProps );
+		hProps = hProps.concat( [ "marginLeft", "marginRight" ] );
+
+		// Only animate children with width attributes specified
+		// TODO: is this right? should we include anything with css width specified as well
+		element.find( "*[width]" ).each( function() {
+			var child = $( this ),
+				childOriginal = $.effects.scaledDimensions( child ),
+				childFrom = {
+					height: childOriginal.height * factor.from.y,
+					width: childOriginal.width * factor.from.x,
+					outerHeight: childOriginal.outerHeight * factor.from.y,
+					outerWidth: childOriginal.outerWidth * factor.from.x
+				},
+				childTo = {
+					height: childOriginal.height * factor.to.y,
+					width: childOriginal.width * factor.to.x,
+					outerHeight: childOriginal.height * factor.to.y,
+					outerWidth: childOriginal.width * factor.to.x
+				};
+
+			// Vertical props scaling
+			if ( factor.from.y !== factor.to.y ) {
+				childFrom = $.effects.setTransition( child, vProps, factor.from.y, childFrom );
+				childTo = $.effects.setTransition( child, vProps, factor.to.y, childTo );
+			}
+
+			// Horizontal props scaling
+			if ( factor.from.x !== factor.to.x ) {
+				childFrom = $.effects.setTransition( child, hProps, factor.from.x, childFrom );
+				childTo = $.effects.setTransition( child, hProps, factor.to.x, childTo );
+			}
+
+			if ( restore ) {
+				$.effects.saveStyle( child );
+			}
+
+			// Animate children
+			child.css( childFrom );
+			child.animate( childTo, options.duration, options.easing, function() {
+
+				// Restore children
+				if ( restore ) {
+					$.effects.restoreStyle( child );
+				}
+			} );
+		} );
+	}
+
+	// Animate
+	element.animate( to, {
+		queue: false,
+		duration: options.duration,
+		easing: options.easing,
+		complete: function() {
+
+			var offset = element.offset();
+
+			if ( to.opacity === 0 ) {
+				element.css( "opacity", from.opacity );
+			}
+
+			if ( !restore ) {
+				element
+					.css( "position", position === "static" ? "relative" : position )
+					.offset( offset );
+
+				// Need to save style here so that automatic style restoration
+				// doesn't restore to the original styles from before the animation.
+				$.effects.saveStyle( element );
+			}
+
+			done();
+		}
+	} );
+
+} );
+
+
+/*!
+ * jQuery UI Effects Scale 1.12.0
+ * http://jqueryui.com
+ *
+ * Copyright jQuery Foundation and other contributors
+ * Released under the MIT license.
+ * http://jquery.org/license
+ */
+
+//>>label: Scale Effect
+//>>group: Effects
+//>>description: Grows or shrinks an element and its content.
+//>>docs: http://api.jqueryui.com/scale-effect/
+//>>demos: http://jqueryui.com/effect/
+
+
+
+var effectsEffectScale = $.effects.define( "scale", function( options, done ) {
+
+	// Create element
+	var el = $( this ),
+		mode = options.mode,
+		percent = parseInt( options.percent, 10 ) ||
+			( parseInt( options.percent, 10 ) === 0 ? 0 : ( mode !== "effect" ? 0 : 100 ) ),
+
+		newOptions = $.extend( true, {
+			from: $.effects.scaledDimensions( el ),
+			to: $.effects.scaledDimensions( el, percent, options.direction || "both" ),
+			origin: options.origin || [ "middle", "center" ]
+		}, options );
+
+	// Fade option to support puff
+	if ( options.fade ) {
+		newOptions.from.opacity = 1;
+		newOptions.to.opacity = 0;
+	}
+
+	$.effects.effect.size.call( this, newOptions, done );
+} );
+
+
+/*!
+ * jQuery UI Effects Puff 1.12.0
+ * http://jqueryui.com
+ *
+ * Copyright jQuery Foundation and other contributors
+ * Released under the MIT license.
+ * http://jquery.org/license
+ */
+
+//>>label: Puff Effect
+//>>group: Effects
+//>>description: Creates a puff effect by scaling the element up and hiding it at the same time.
+//>>docs: http://api.jqueryui.com/puff-effect/
+//>>demos: http://jqueryui.com/effect/
+
+
+
+var effectsEffectPuff = $.effects.define( "puff", "hide", function( options, done ) {
+	var newOptions = $.extend( true, {}, options, {
+		fade: true,
+		percent: parseInt( options.percent, 10 ) || 150
+	} );
+
+	$.effects.effect.scale.call( this, newOptions, done );
+} );
+
+
+/*!
+ * jQuery UI Effects Pulsate 1.12.0
+ * http://jqueryui.com
+ *
+ * Copyright jQuery Foundation and other contributors
+ * Released under the MIT license.
+ * http://jquery.org/license
+ */
+
+//>>label: Pulsate Effect
+//>>group: Effects
+//>>description: Pulsates an element n times by changing the opacity to zero and back.
+//>>docs: http://api.jqueryui.com/pulsate-effect/
+//>>demos: http://jqueryui.com/effect/
+
+
+
+var effectsEffectPulsate = $.effects.define( "pulsate", "show", function( options, done ) {
+	var element = $( this ),
+		mode = options.mode,
+		show = mode === "show",
+		hide = mode === "hide",
+		showhide = show || hide,
+
+		// Showing or hiding leaves off the "last" animation
+		anims = ( ( options.times || 5 ) * 2 ) + ( showhide ? 1 : 0 ),
+		duration = options.duration / anims,
+		animateTo = 0,
+		i = 1,
+		queuelen = element.queue().length;
+
+	if ( show || !element.is( ":visible" ) ) {
+		element.css( "opacity", 0 ).show();
+		animateTo = 1;
+	}
+
+	// Anims - 1 opacity "toggles"
+	for ( ; i < anims; i++ ) {
+		element.animate( { opacity: animateTo }, duration, options.easing );
+		animateTo = 1 - animateTo;
+	}
+
+	element.animate( { opacity: animateTo }, duration, options.easing );
+
+	element.queue( done );
+
+	$.effects.unshift( element, queuelen, anims + 1 );
+} );
+
+
+/*!
+ * jQuery UI Effects Shake 1.12.0
+ * http://jqueryui.com
+ *
+ * Copyright jQuery Foundation and other contributors
+ * Released under the MIT license.
+ * http://jquery.org/license
+ */
+
+//>>label: Shake Effect
+//>>group: Effects
+//>>description: Shakes an element horizontally or vertically n times.
+//>>docs: http://api.jqueryui.com/shake-effect/
+//>>demos: http://jqueryui.com/effect/
+
+
+
+var effectsEffectShake = $.effects.define( "shake", function( options, done ) {
+
+	var i = 1,
+		element = $( this ),
+		direction = options.direction || "left",
+		distance = options.distance || 20,
+		times = options.times || 3,
+		anims = times * 2 + 1,
+		speed = Math.round( options.duration / anims ),
+		ref = ( direction === "up" || direction === "down" ) ? "top" : "left",
+		positiveMotion = ( direction === "up" || direction === "left" ),
+		animation = {},
+		animation1 = {},
+		animation2 = {},
+
+		queuelen = element.queue().length;
+
+	$.effects.createPlaceholder( element );
+
+	// Animation
+	animation[ ref ] = ( positiveMotion ? "-=" : "+=" ) + distance;
+	animation1[ ref ] = ( positiveMotion ? "+=" : "-=" ) + distance * 2;
+	animation2[ ref ] = ( positiveMotion ? "-=" : "+=" ) + distance * 2;
+
+	// Animate
+	element.animate( animation, speed, options.easing );
+
+	// Shakes
+	for ( ; i < times; i++ ) {
+		element
+			.animate( animation1, speed, options.easing )
+			.animate( animation2, speed, options.easing );
+	}
+
+	element
+		.animate( animation1, speed, options.easing )
+		.animate( animation, speed / 2, options.easing )
+		.queue( done );
+
+	$.effects.unshift( element, queuelen, anims + 1 );
+} );
+
+
+/*!
+ * jQuery UI Effects Slide 1.12.0
+ * http://jqueryui.com
+ *
+ * Copyright jQuery Foundation and other contributors
+ * Released under the MIT license.
+ * http://jquery.org/license
+ */
+
+//>>label: Slide Effect
+//>>group: Effects
+//>>description: Slides an element in and out of the viewport.
+//>>docs: http://api.jqueryui.com/slide-effect/
+//>>demos: http://jqueryui.com/effect/
+
+
+
+var effectsEffectSlide = $.effects.define( "slide", "show", function( options, done ) {
+	var startClip, startRef,
+		element = $( this ),
+		map = {
+			up: [ "bottom", "top" ],
+			down: [ "top", "bottom" ],
+			left: [ "right", "left" ],
+			right: [ "left", "right" ]
+		},
+		mode = options.mode,
+		direction = options.direction || "left",
+		ref = ( direction === "up" || direction === "down" ) ? "top" : "left",
+		positiveMotion = ( direction === "up" || direction === "left" ),
+		distance = options.distance ||
+			element[ ref === "top" ? "outerHeight" : "outerWidth" ]( true ),
+		animation = {};
+
+	$.effects.createPlaceholder( element );
+
+	startClip = element.cssClip();
+	startRef = element.position()[ ref ];
+
+	// Define hide animation
+	animation[ ref ] = ( positiveMotion ? -1 : 1 ) * distance + startRef;
+	animation.clip = element.cssClip();
+	animation.clip[ map[ direction ][ 1 ] ] = animation.clip[ map[ direction ][ 0 ] ];
+
+	// Reverse the animation if we're showing
+	if ( mode === "show" ) {
+		element.cssClip( animation.clip );
+		element.css( ref, animation[ ref ] );
+		animation.clip = startClip;
+		animation[ ref ] = startRef;
+	}
+
+	// Actually animate
+	element.animate( animation, {
+		queue: false,
+		duration: options.duration,
+		easing: options.easing,
+		complete: done
+	} );
+} );
+
+
+/*!
+ * jQuery UI Effects Transfer 1.12.0
+ * http://jqueryui.com
+ *
+ * Copyright jQuery Foundation and other contributors
+ * Released under the MIT license.
+ * http://jquery.org/license
+ */
+
+//>>label: Transfer Effect
+//>>group: Effects
+//>>description: Displays a transfer effect from one element to another.
+//>>docs: http://api.jqueryui.com/transfer-effect/
+//>>demos: http://jqueryui.com/effect/
+
+
+
+var effect;
+if ( $.uiBackCompat !== false ) {
+	effect = $.effects.define( "transfer", function( options, done ) {
+		$( this ).transfer( options, done );
+	} );
+}
+var effectsEffectTransfer = effect;
+
+
+/*!
+ * jQuery UI Focusable 1.12.0
+ * http://jqueryui.com
+ *
+ * Copyright jQuery Foundation and other contributors
+ * Released under the MIT license.
+ * http://jquery.org/license
+ */
+
+//>>label: :focusable Selector
+//>>group: Core
+//>>description: Selects elements which can be focused.
+//>>docs: http://api.jqueryui.com/focusable-selector/
+
+
+
+// Selectors
+$.ui.focusable = function( element, hasTabindex ) {
+	var map, mapName, img, focusableIfVisible, fieldset,
+		nodeName = element.nodeName.toLowerCase();
+
+	if ( "area" === nodeName ) {
+		map = element.parentNode;
+		mapName = map.name;
+		if ( !element.href || !mapName || map.nodeName.toLowerCase() !== "map" ) {
+			return false;
+		}
+		img = $( "img[usemap='#" + mapName + "']" );
+		return img.length > 0 && img.is( ":visible" );
+	}
+
+	if ( /^(input|select|textarea|button|object)$/.test( nodeName ) ) {
+		focusableIfVisible = !element.disabled;
+
+		if ( focusableIfVisible ) {
+
+			// Form controls within a disabled fieldset are disabled.
+			// However, controls within the fieldset's legend do not get disabled.
+			// Since controls generally aren't placed inside legends, we skip
+			// this portion of the check.
+			fieldset = $( element ).closest( "fieldset" )[ 0 ];
+			if ( fieldset ) {
+				focusableIfVisible = !fieldset.disabled;
+			}
+		}
+	} else if ( "a" === nodeName ) {
+		focusableIfVisible = element.href || hasTabindex;
+	} else {
+		focusableIfVisible = hasTabindex;
+	}
+
+	return focusableIfVisible && $( element ).is( ":visible" ) && visible( $( element ) );
+};
+
+// Support: IE 8 only
+// IE 8 doesn't resolve inherit to visible/hidden for computed values
+function visible( element ) {
+	var visibility = element.css( "visibility" );
+	while ( visibility === "inherit" ) {
+		element = element.parent();
+		visibility = element.css( "visibility" );
+	}
+	return visibility !== "hidden";
+}
+
+$.extend( $.expr[ ":" ], {
+	focusable: function( element ) {
+		return $.ui.focusable( element, $.attr( element, "tabindex" ) != null );
+	}
+} );
+
+var focusable = $.ui.focusable;
+
+
+
+
+// Support: IE8 Only
+// IE8 does not support the form attribute and when it is supplied. It overwrites the form prop
+// with a string, so we need to find the proper form.
+var form = $.fn.form = function() {
+	return typeof this[ 0 ].form === "string" ? this.closest( "form" ) : $( this[ 0 ].form );
+};
+
+
+/*!
+ * jQuery UI Form Reset Mixin 1.12.0
+ * http://jqueryui.com
+ *
+ * Copyright jQuery Foundation and other contributors
+ * Released under the MIT license.
+ * http://jquery.org/license
+ */
+
+//>>label: Form Reset Mixin
+//>>group: Core
+//>>description: Refresh input widgets when their form is reset
+//>>docs: http://api.jqueryui.com/form-reset-mixin/
+
+
+
+var formResetMixin = $.ui.formResetMixin = {
+	_formResetHandler: function() {
+		var form = $( this );
+
+		// Wait for the form reset to actually happen before refreshing
+		setTimeout( function() {
+			var instances = form.data( "ui-form-reset-instances" );
+			$.each( instances, function() {
+				this.refresh();
+			} );
+		} );
+	},
+
+	_bindFormResetHandler: function() {
+		this.form = this.element.form();
+		if ( !this.form.length ) {
+			return;
+		}
+
+		var instances = this.form.data( "ui-form-reset-instances" ) || [];
+		if ( !instances.length ) {
+
+			// We don't use _on() here because we use a single event handler per form
+			this.form.on( "reset.ui-form-reset", this._formResetHandler );
+		}
+		instances.push( this );
+		this.form.data( "ui-form-reset-instances", instances );
+	},
+
+	_unbindFormResetHandler: function() {
+		if ( !this.form.length ) {
+			return;
+		}
+
+		var instances = this.form.data( "ui-form-reset-instances" );
+		instances.splice( $.inArray( this, instances ), 1 );
+		if ( instances.length ) {
+			this.form.data( "ui-form-reset-instances", instances );
+		} else {
+			this.form
+				.removeData( "ui-form-reset-instances" )
+				.off( "reset.ui-form-reset" );
+		}
+	}
+};
+
+
+/*!
+ * jQuery UI Support for jQuery core 1.7.x 1.12.0
+ * http://jqueryui.com
+ *
+ * Copyright jQuery Foundation and other contributors
+ * Released under the MIT license.
+ * http://jquery.org/license
+ *
+ */
+
+//>>label: jQuery 1.7 Support
+//>>group: Core
+//>>description: Support version 1.7.x of jQuery core
+
+
+
+// Support: jQuery 1.7 only
+// Not a great way to check versions, but since we only support 1.7+ and only
+// need to detect <1.8, this is a simple check that should suffice. Checking
+// for "1.7." would be a bit safer, but the version string is 1.7, not 1.7.0
+// and we'll never reach 1.70.0 (if we do, we certainly won't be supporting
+// 1.7 anymore). See #11197 for why we're not using feature detection.
+if ( $.fn.jquery.substring( 0, 3 ) === "1.7" ) {
+
+	// Setters for .innerWidth(), .innerHeight(), .outerWidth(), .outerHeight()
+	// Unlike jQuery Core 1.8+, these only support numeric values to set the
+	// dimensions in pixels
+	$.each( [ "Width", "Height" ], function( i, name ) {
+		var side = name === "Width" ? [ "Left", "Right" ] : [ "Top", "Bottom" ],
+			type = name.toLowerCase(),
+			orig = {
+				innerWidth: $.fn.innerWidth,
+				innerHeight: $.fn.innerHeight,
+				outerWidth: $.fn.outerWidth,
+				outerHeight: $.fn.outerHeight
+			};
+
+		function reduce( elem, size, border, margin ) {
+			$.each( side, function() {
+				size -= parseFloat( $.css( elem, "padding" + this ) ) || 0;
+				if ( border ) {
+					size -= parseFloat( $.css( elem, "border" + this + "Width" ) ) || 0;
+				}
+				if ( margin ) {
+					size -= parseFloat( $.css( elem, "margin" + this ) ) || 0;
+				}
+			} );
+			return size;
+		}
+
+		$.fn[ "inner" + name ] = function( size ) {
+			if ( size === undefined ) {
+				return orig[ "inner" + name ].call( this );
+			}
+
+			return this.each( function() {
+				$( this ).css( type, reduce( this, size ) + "px" );
+			} );
+		};
+
+		$.fn[ "outer" + name ] = function( size, margin ) {
+			if ( typeof size !== "number" ) {
+				return orig[ "outer" + name ].call( this, size );
+			}
+
+			return this.each( function() {
+				$( this ).css( type, reduce( this, size, true, margin ) + "px" );
+			} );
+		};
+	} );
+
+	$.fn.addBack = function( selector ) {
+		return this.add( selector == null ?
+			this.prevObject : this.prevObject.filter( selector )
+		);
+	};
+}
+
+;
+/*!
+ * jQuery UI Keycode 1.12.0
+ * http://jqueryui.com
+ *
+ * Copyright jQuery Foundation and other contributors
+ * Released under the MIT license.
+ * http://jquery.org/license
+ */
+
+//>>label: Keycode
+//>>group: Core
+//>>description: Provide keycodes as keynames
+//>>docs: http://api.jqueryui.com/jQuery.ui.keyCode/
+
+
+var keycode = $.ui.keyCode = {
+	BACKSPACE: 8,
+	COMMA: 188,
+	DELETE: 46,
+	DOWN: 40,
+	END: 35,
+	ENTER: 13,
+	ESCAPE: 27,
+	HOME: 36,
+	LEFT: 37,
+	PAGE_DOWN: 34,
+	PAGE_UP: 33,
+	PERIOD: 190,
+	RIGHT: 39,
+	SPACE: 32,
+	TAB: 9,
+	UP: 38
+};
+
+
+
+
+// Internal use only
+var escapeSelector = $.ui.escapeSelector = ( function() {
+	var selectorEscape = /([!"#$%&'()*+,./:;<=>?@[\]^`{|}~])/g;
+	return function( selector ) {
+		return selector.replace( selectorEscape, "\\$1" );
+	};
+} )();
+
+
+/*!
+ * jQuery UI Labels 1.12.0
+ * http://jqueryui.com
+ *
+ * Copyright jQuery Foundation and other contributors
+ * Released under the MIT license.
+ * http://jquery.org/license
+ */
+
+//>>label: labels
+//>>group: Core
+//>>description: Find all the labels associated with a given input
+//>>docs: http://api.jqueryui.com/labels/
+
+
+
+var labels = $.fn.labels = function() {
+	var ancestor, selector, id, labels, ancestors;
+
+	// Check control.labels first
+	if ( this[ 0 ].labels && this[ 0 ].labels.length ) {
+		return this.pushStack( this[ 0 ].labels );
+	}
+
+	// Support: IE <= 11, FF <= 37, Android <= 2.3 only
+	// Above browsers do not support control.labels. Everything below is to support them
+	// as well as document fragments. control.labels does not work on document fragments
+	labels = this.eq( 0 ).parents( "label" );
+
+	// Look for the label based on the id
+	id = this.attr( "id" );
+	if ( id ) {
+
+		// We don't search against the document in case the element
+		// is disconnected from the DOM
+		ancestor = this.eq( 0 ).parents().last();
+
+		// Get a full set of top level ancestors
+		ancestors = ancestor.add( ancestor.length ? ancestor.siblings() : this.siblings() );
+
+		// Create a selector for the label based on the id
+		selector = "label[for='" + $.ui.escapeSelector( id ) + "']";
+
+		labels = labels.add( ancestors.find( selector ).addBack( selector ) );
+
+	}
+
+	// Return whatever we have found for labels
+	return this.pushStack( labels );
+};
+
+
+/*!
+ * jQuery UI Scroll Parent 1.12.0
+ * http://jqueryui.com
+ *
+ * Copyright jQuery Foundation and other contributors
+ * Released under the MIT license.
+ * http://jquery.org/license
+ */
+
+//>>label: scrollParent
+//>>group: Core
+//>>description: Get the closest ancestor element that is scrollable.
+//>>docs: http://api.jqueryui.com/scrollParent/
+
+
+
+var scrollParent = $.fn.scrollParent = function( includeHidden ) {
+	var position = this.css( "position" ),
+		excludeStaticParent = position === "absolute",
+		overflowRegex = includeHidden ? /(auto|scroll|hidden)/ : /(auto|scroll)/,
+		scrollParent = this.parents().filter( function() {
+			var parent = $( this );
+			if ( excludeStaticParent && parent.css( "position" ) === "static" ) {
+				return false;
+			}
+			return overflowRegex.test( parent.css( "overflow" ) + parent.css( "overflow-y" ) +
+				parent.css( "overflow-x" ) );
+		} ).eq( 0 );
+
+	return position === "fixed" || !scrollParent.length ?
+		$( this[ 0 ].ownerDocument || document ) :
+		scrollParent;
+};
+
+
+/*!
+ * jQuery UI Tabbable 1.12.0
+ * http://jqueryui.com
+ *
+ * Copyright jQuery Foundation and other contributors
+ * Released under the MIT license.
+ * http://jquery.org/license
+ */
+
+//>>label: :tabbable Selector
+//>>group: Core
+//>>description: Selects elements which can be tabbed to.
+//>>docs: http://api.jqueryui.com/tabbable-selector/
+
+
+
+var tabbable = $.extend( $.expr[ ":" ], {
+	tabbable: function( element ) {
+		var tabIndex = $.attr( element, "tabindex" ),
+			hasTabindex = tabIndex != null;
+		return ( !hasTabindex || tabIndex >= 0 ) && $.ui.focusable( element, hasTabindex );
+	}
+} );
+
+
+/*!
+ * jQuery UI Unique ID 1.12.0
+ * http://jqueryui.com
+ *
+ * Copyright jQuery Foundation and other contributors
+ * Released under the MIT license.
+ * http://jquery.org/license
+ */
+
+//>>label: uniqueId
+//>>group: Core
+//>>description: Functions to generate and remove uniqueId's
+//>>docs: http://api.jqueryui.com/uniqueId/
+
+
+
+var uniqueId = $.fn.extend( {
+	uniqueId: ( function() {
+		var uuid = 0;
+
+		return function() {
+			return this.each( function() {
+				if ( !this.id ) {
+					this.id = "ui-id-" + ( ++uuid );
+				}
+			} );
+		};
+	} )(),
+
+	removeUniqueId: function() {
+		return this.each( function() {
+			if ( /^ui-id-\d+$/.test( this.id ) ) {
+				$( this ).removeAttr( "id" );
+			}
+		} );
+	}
+} );
+
+
+/*!
+ * jQuery UI Accordion 1.12.0
+ * http://jqueryui.com
+ *
+ * Copyright jQuery Foundation and other contributors
+ * Released under the MIT license.
+ * http://jquery.org/license
+ */
+
+//>>label: Accordion
+//>>group: Widgets
+// jscs:disable maximumLineLength
+//>>description: Displays collapsible content panels for presenting information in a limited amount of space.
+// jscs:enable maximumLineLength
+//>>docs: http://api.jqueryui.com/accordion/
+//>>demos: http://jqueryui.com/accordion/
+//>>css.structure: ../../themes/base/core.css
+//>>css.structure: ../../themes/base/accordion.css
+//>>css.theme: ../../themes/base/theme.css
+
+
+
+var widgetsAccordion = $.widget( "ui.accordion", {
+	version: "1.12.0",
+	options: {
+		active: 0,
+		animate: {},
+		classes: {
+			"ui-accordion-header": "ui-corner-top",
+			"ui-accordion-header-collapsed": "ui-corner-all",
+			"ui-accordion-content": "ui-corner-bottom"
+		},
+		collapsible: false,
+		event: "click",
+		header: "> li > :first-child, > :not(li):even",
+		heightStyle: "auto",
+		icons: {
+			activeHeader: "ui-icon-triangle-1-s",
+			header: "ui-icon-triangle-1-e"
+		},
+
+		// Callbacks
+		activate: null,
+		beforeActivate: null
+	},
+
+	hideProps: {
+		borderTopWidth: "hide",
+		borderBottomWidth: "hide",
+		paddingTop: "hide",
+		paddingBottom: "hide",
+		height: "hide"
+	},
+
+	showProps: {
+		borderTopWidth: "show",
+		borderBottomWidth: "show",
+		paddingTop: "show",
+		paddingBottom: "show",
+		height: "show"
+	},
+
+	_create: function() {
+		var options = this.options;
+
+		this.prevShow = this.prevHide = $();
+		this._addClass( "ui-accordion", "ui-widget ui-helper-reset" );
+		this.element.attr( "role", "tablist" );
+
+		// Don't allow collapsible: false and active: false / null
+		if ( !options.collapsible && ( options.active === false || options.active == null ) ) {
+			options.active = 0;
+		}
+
+		this._processPanels();
+
+		// handle negative values
+		if ( options.active < 0 ) {
+			options.active += this.headers.length;
+		}
+		this._refresh();
+	},
+
+	_getCreateEventData: function() {
+		return {
+			header: this.active,
+			panel: !this.active.length ? $() : this.active.next()
+		};
+	},
+
+	_createIcons: function() {
+		var icon, children,
+			icons = this.options.icons;
+
+		if ( icons ) {
+			icon = $( "<span>" );
+			this._addClass( icon, "ui-accordion-header-icon", "ui-icon " + icons.header );
+			icon.prependTo( this.headers );
+			children = this.active.children( ".ui-accordion-header-icon" );
+			this._removeClass( children, icons.header )
+				._addClass( children, null, icons.activeHeader )
+				._addClass( this.headers, "ui-accordion-icons" );
+		}
+	},
+
+	_destroyIcons: function() {
+		this._removeClass( this.headers, "ui-accordion-icons" );
+		this.headers.children( ".ui-accordion-header-icon" ).remove();
+	},
+
+	_destroy: function() {
+		var contents;
+
+		// Clean up main element
+		this.element.removeAttr( "role" );
+
+		// Clean up headers
+		this.headers
+			.removeAttr( "role aria-expanded aria-selected aria-controls tabIndex" )
+			.removeUniqueId();
+
+		this._destroyIcons();
+
+		// Clean up content panels
+		contents = this.headers.next()
+			.css( "display", "" )
+			.removeAttr( "role aria-hidden aria-labelledby" )
+			.removeUniqueId();
+
+		if ( this.options.heightStyle !== "content" ) {
+			contents.css( "height", "" );
+		}
+	},
+
+	_setOption: function( key, value ) {
+		if ( key === "active" ) {
+
+			// _activate() will handle invalid values and update this.options
+			this._activate( value );
+			return;
+		}
+
+		if ( key === "event" ) {
+			if ( this.options.event ) {
+				this._off( this.headers, this.options.event );
+			}
+			this._setupEvents( value );
+		}
+
+		this._super( key, value );
+
+		// Setting collapsible: false while collapsed; open first panel
+		if ( key === "collapsible" && !value && this.options.active === false ) {
+			this._activate( 0 );
+		}
+
+		if ( key === "icons" ) {
+			this._destroyIcons();
+			if ( value ) {
+				this._createIcons();
+			}
+		}
+	},
+
+	_setOptionDisabled: function( value ) {
+		this._super( value );
+
+		this.element.attr( "aria-disabled", value );
+
+		// Support: IE8 Only
+		// #5332 / #6059 - opacity doesn't cascade to positioned elements in IE
+		// so we need to add the disabled class to the headers and panels
+		this._toggleClass( null, "ui-state-disabled", !!value );
+		this._toggleClass( this.headers.add( this.headers.next() ), null, "ui-state-disabled",
+			!!value );
+	},
+
+	_keydown: function( event ) {
+		if ( event.altKey || event.ctrlKey ) {
+			return;
+		}
+
+		var keyCode = $.ui.keyCode,
+			length = this.headers.length,
+			currentIndex = this.headers.index( event.target ),
+			toFocus = false;
+
+		switch ( event.keyCode ) {
+		case keyCode.RIGHT:
+		case keyCode.DOWN:
+			toFocus = this.headers[ ( currentIndex + 1 ) % length ];
+			break;
+		case keyCode.LEFT:
+		case keyCode.UP:
+			toFocus = this.headers[ ( currentIndex - 1 + length ) % length ];
+			break;
+		case keyCode.SPACE:
+		case keyCode.ENTER:
+			this._eventHandler( event );
+			break;
+		case keyCode.HOME:
+			toFocus = this.headers[ 0 ];
+			break;
+		case keyCode.END:
+			toFocus = this.headers[ length - 1 ];
+			break;
+		}
+
+		if ( toFocus ) {
+			$( event.target ).attr( "tabIndex", -1 );
+			$( toFocus ).attr( "tabIndex", 0 );
+			$( toFocus ).trigger( "focus" );
+			event.preventDefault();
+		}
+	},
+
+	_panelKeyDown: function( event ) {
+		if ( event.keyCode === $.ui.keyCode.UP && event.ctrlKey ) {
+			$( event.currentTarget ).prev().trigger( "focus" );
+		}
+	},
+
+	refresh: function() {
+		var options = this.options;
+		this._processPanels();
+
+		// Was collapsed or no panel
+		if ( ( options.active === false && options.collapsible === true ) ||
+				!this.headers.length ) {
+			options.active = false;
+			this.active = $();
+
+		// active false only when collapsible is true
+		} else if ( options.active === false ) {
+			this._activate( 0 );
+
+		// was active, but active panel is gone
+		} else if ( this.active.length && !$.contains( this.element[ 0 ], this.active[ 0 ] ) ) {
+
+			// all remaining panel are disabled
+			if ( this.headers.length === this.headers.find( ".ui-state-disabled" ).length ) {
+				options.active = false;
+				this.active = $();
+
+			// activate previous panel
+			} else {
+				this._activate( Math.max( 0, options.active - 1 ) );
+			}
+
+		// was active, active panel still exists
+		} else {
+
+			// make sure active index is correct
+			options.active = this.headers.index( this.active );
+		}
+
+		this._destroyIcons();
+
+		this._refresh();
+	},
+
+	_processPanels: function() {
+		var prevHeaders = this.headers,
+			prevPanels = this.panels;
+
+		this.headers = this.element.find( this.options.header );
+		this._addClass( this.headers, "ui-accordion-header ui-accordion-header-collapsed",
+			"ui-state-default" );
+
+		this.panels = this.headers.next().filter( ":not(.ui-accordion-content-active)" ).hide();
+		this._addClass( this.panels, "ui-accordion-content", "ui-helper-reset ui-widget-content" );
+
+		// Avoid memory leaks (#10056)
+		if ( prevPanels ) {
+			this._off( prevHeaders.not( this.headers ) );
+			this._off( prevPanels.not( this.panels ) );
+		}
+	},
+
+	_refresh: function() {
+		var maxHeight,
+			options = this.options,
+			heightStyle = options.heightStyle,
+			parent = this.element.parent();
+
+		this.active = this._findActive( options.active );
+		this._addClass( this.active, "ui-accordion-header-active", "ui-state-active" )
+			._removeClass( this.active, "ui-accordion-header-collapsed" );
+		this._addClass( this.active.next(), "ui-accordion-content-active" );
+		this.active.next().show();
+
+		this.headers
+			.attr( "role", "tab" )
+			.each( function() {
+				var header = $( this ),
+					headerId = header.uniqueId().attr( "id" ),
+					panel = header.next(),
+					panelId = panel.uniqueId().attr( "id" );
+				header.attr( "aria-controls", panelId );
+				panel.attr( "aria-labelledby", headerId );
+			} )
+			.next()
+				.attr( "role", "tabpanel" );
+
+		this.headers
+			.not( this.active )
+				.attr( {
+					"aria-selected": "false",
+					"aria-expanded": "false",
+					tabIndex: -1
+				} )
+				.next()
+					.attr( {
+						"aria-hidden": "true"
+					} )
+					.hide();
+
+		// Make sure at least one header is in the tab order
+		if ( !this.active.length ) {
+			this.headers.eq( 0 ).attr( "tabIndex", 0 );
+		} else {
+			this.active.attr( {
+				"aria-selected": "true",
+				"aria-expanded": "true",
+				tabIndex: 0
+			} )
+				.next()
+					.attr( {
+						"aria-hidden": "false"
+					} );
+		}
+
+		this._createIcons();
+
+		this._setupEvents( options.event );
+
+		if ( heightStyle === "fill" ) {
+			maxHeight = parent.height();
+			this.element.siblings( ":visible" ).each( function() {
+				var elem = $( this ),
+					position = elem.css( "position" );
+
+				if ( position === "absolute" || position === "fixed" ) {
+					return;
+				}
+				maxHeight -= elem.outerHeight( true );
+			} );
+
+			this.headers.each( function() {
+				maxHeight -= $( this ).outerHeight( true );
+			} );
+
+			this.headers.next()
+				.each( function() {
+					$( this ).height( Math.max( 0, maxHeight -
+						$( this ).innerHeight() + $( this ).height() ) );
+				} )
+				.css( "overflow", "auto" );
+		} else if ( heightStyle === "auto" ) {
+			maxHeight = 0;
+			this.headers.next()
+				.each( function() {
+					var isVisible = $( this ).is( ":visible" );
+					if ( !isVisible ) {
+						$( this ).show();
+					}
+					maxHeight = Math.max( maxHeight, $( this ).css( "height", "" ).height() );
+					if ( !isVisible ) {
+						$( this ).hide();
+					}
+				} )
+				.height( maxHeight );
+		}
+	},
+
+	_activate: function( index ) {
+		var active = this._findActive( index )[ 0 ];
+
+		// Trying to activate the already active panel
+		if ( active === this.active[ 0 ] ) {
+			return;
+		}
+
+		// Trying to collapse, simulate a click on the currently active header
+		active = active || this.active[ 0 ];
+
+		this._eventHandler( {
+			target: active,
+			currentTarget: active,
+			preventDefault: $.noop
+		} );
+	},
+
+	_findActive: function( selector ) {
+		return typeof selector === "number" ? this.headers.eq( selector ) : $();
+	},
+
+	_setupEvents: function( event ) {
+		var events = {
+			keydown: "_keydown"
+		};
+		if ( event ) {
+			$.each( event.split( " " ), function( index, eventName ) {
+				events[ eventName ] = "_eventHandler";
+			} );
+		}
+
+		this._off( this.headers.add( this.headers.next() ) );
+		this._on( this.headers, events );
+		this._on( this.headers.next(), { keydown: "_panelKeyDown" } );
+		this._hoverable( this.headers );
+		this._focusable( this.headers );
+	},
+
+	_eventHandler: function( event ) {
+		var activeChildren, clickedChildren,
+			options = this.options,
+			active = this.active,
+			clicked = $( event.currentTarget ),
+			clickedIsActive = clicked[ 0 ] === active[ 0 ],
+			collapsing = clickedIsActive && options.collapsible,
+			toShow = collapsing ? $() : clicked.next(),
+			toHide = active.next(),
+			eventData = {
+				oldHeader: active,
+				oldPanel: toHide,
+				newHeader: collapsing ? $() : clicked,
+				newPanel: toShow
+			};
+
+		event.preventDefault();
+
+		if (
+
+				// click on active header, but not collapsible
+				( clickedIsActive && !options.collapsible ) ||
+
+				// allow canceling activation
+				( this._trigger( "beforeActivate", event, eventData ) === false ) ) {
+			return;
+		}
+
+		options.active = collapsing ? false : this.headers.index( clicked );
+
+		// When the call to ._toggle() comes after the class changes
+		// it causes a very odd bug in IE 8 (see #6720)
+		this.active = clickedIsActive ? $() : clicked;
+		this._toggle( eventData );
+
+		// Switch classes
+		// corner classes on the previously active header stay after the animation
+		this._removeClass( active, "ui-accordion-header-active", "ui-state-active" );
+		if ( options.icons ) {
+			activeChildren = active.children( ".ui-accordion-header-icon" );
+			this._removeClass( activeChildren, null, options.icons.activeHeader )
+				._addClass( activeChildren, null, options.icons.header );
+		}
+
+		if ( !clickedIsActive ) {
+			this._removeClass( clicked, "ui-accordion-header-collapsed" )
+				._addClass( clicked, "ui-accordion-header-active", "ui-state-active" );
+			if ( options.icons ) {
+				clickedChildren = clicked.children( ".ui-accordion-header-icon" );
+				this._removeClass( clickedChildren, null, options.icons.header )
+					._addClass( clickedChildren, null, options.icons.activeHeader );
+			}
+
+			this._addClass( clicked.next(), "ui-accordion-content-active" );
+		}
+	},
+
+	_toggle: function( data ) {
+		var toShow = data.newPanel,
+			toHide = this.prevShow.length ? this.prevShow : data.oldPanel;
+
+		// Handle activating a panel during the animation for another activation
+		this.prevShow.add( this.prevHide ).stop( true, true );
+		this.prevShow = toShow;
+		this.prevHide = toHide;
+
+		if ( this.options.animate ) {
+			this._animate( toShow, toHide, data );
+		} else {
+			toHide.hide();
+			toShow.show();
+			this._toggleComplete( data );
+		}
+
+		toHide.attr( {
+			"aria-hidden": "true"
+		} );
+		toHide.prev().attr( {
+			"aria-selected": "false",
+			"aria-expanded": "false"
+		} );
+
+		// if we're switching panels, remove the old header from the tab order
+		// if we're opening from collapsed state, remove the previous header from the tab order
+		// if we're collapsing, then keep the collapsing header in the tab order
+		if ( toShow.length && toHide.length ) {
+			toHide.prev().attr( {
+				"tabIndex": -1,
+				"aria-expanded": "false"
+			} );
+		} else if ( toShow.length ) {
+			this.headers.filter( function() {
+				return parseInt( $( this ).attr( "tabIndex" ), 10 ) === 0;
+			} )
+				.attr( "tabIndex", -1 );
+		}
+
+		toShow
+			.attr( "aria-hidden", "false" )
+			.prev()
+				.attr( {
+					"aria-selected": "true",
+					"aria-expanded": "true",
+					tabIndex: 0
+				} );
+	},
+
+	_animate: function( toShow, toHide, data ) {
+		var total, easing, duration,
+			that = this,
+			adjust = 0,
+			boxSizing = toShow.css( "box-sizing" ),
+			down = toShow.length &&
+				( !toHide.length || ( toShow.index() < toHide.index() ) ),
+			animate = this.options.animate || {},
+			options = down && animate.down || animate,
+			complete = function() {
+				that._toggleComplete( data );
+			};
+
+		if ( typeof options === "number" ) {
+			duration = options;
+		}
+		if ( typeof options === "string" ) {
+			easing = options;
+		}
+
+		// fall back from options to animation in case of partial down settings
+		easing = easing || options.easing || animate.easing;
+		duration = duration || options.duration || animate.duration;
+
+		if ( !toHide.length ) {
+			return toShow.animate( this.showProps, duration, easing, complete );
+		}
+		if ( !toShow.length ) {
+			return toHide.animate( this.hideProps, duration, easing, complete );
+		}
+
+		total = toShow.show().outerHeight();
+		toHide.animate( this.hideProps, {
+			duration: duration,
+			easing: easing,
+			step: function( now, fx ) {
+				fx.now = Math.round( now );
+			}
+		} );
+		toShow
+			.hide()
+			.animate( this.showProps, {
+				duration: duration,
+				easing: easing,
+				complete: complete,
+				step: function( now, fx ) {
+					fx.now = Math.round( now );
+					if ( fx.prop !== "height" ) {
+						if ( boxSizing === "content-box" ) {
+							adjust += fx.now;
+						}
+					} else if ( that.options.heightStyle !== "content" ) {
+						fx.now = Math.round( total - toHide.outerHeight() - adjust );
+						adjust = 0;
+					}
+				}
+			} );
+	},
+
+	_toggleComplete: function( data ) {
+		var toHide = data.oldPanel,
+			prev = toHide.prev();
+
+		this._removeClass( toHide, "ui-accordion-content-active" );
+		this._removeClass( prev, "ui-accordion-header-active" )
+			._addClass( prev, "ui-accordion-header-collapsed" );
+
+		// Work around for rendering bug in IE (#5421)
+		if ( toHide.length ) {
+			toHide.parent()[ 0 ].className = toHide.parent()[ 0 ].className;
+		}
+		this._trigger( "activate", null, data );
+	}
+} );
+
+
+
+var safeActiveElement = $.ui.safeActiveElement = function( document ) {
+	var activeElement;
+
+	// Support: IE 9 only
+	// IE9 throws an "Unspecified error" accessing document.activeElement from an <iframe>
+	try {
+		activeElement = document.activeElement;
+	} catch ( error ) {
+		activeElement = document.body;
+	}
+
+	// Support: IE 9 - 11 only
+	// IE may return null instead of an element
+	// Interestingly, this only seems to occur when NOT in an iframe
+	if ( !activeElement ) {
+		activeElement = document.body;
+	}
+
+	// Support: IE 11 only
+	// IE11 returns a seemingly empty object in some cases when accessing
+	// document.activeElement from an <iframe>
+	if ( !activeElement.nodeName ) {
+		activeElement = document.body;
+	}
+
+	return activeElement;
+};
+
+
+/*!
+ * jQuery UI Menu 1.12.0
+ * http://jqueryui.com
+ *
+ * Copyright jQuery Foundation and other contributors
+ * Released under the MIT license.
+ * http://jquery.org/license
+ */
+
+//>>label: Menu
+//>>group: Widgets
+//>>description: Creates nestable menus.
+//>>docs: http://api.jqueryui.com/menu/
+//>>demos: http://jqueryui.com/menu/
+//>>css.structure: ../../themes/base/core.css
+//>>css.structure: ../../themes/base/menu.css
+//>>css.theme: ../../themes/base/theme.css
+
+
+
+var widgetsMenu = $.widget( "ui.menu", {
+	version: "1.12.0",
+	defaultElement: "<ul>",
+	delay: 300,
+	options: {
+		icons: {
+			submenu: "ui-icon-caret-1-e"
+		},
+		items: "> *",
+		menus: "ul",
+		position: {
+			my: "left top",
+			at: "right top"
+		},
+		role: "menu",
+
+		// Callbacks
+		blur: null,
+		focus: null,
+		select: null
+	},
+
+	_create: function() {
+		this.activeMenu = this.element;
+
+		// Flag used to prevent firing of the click handler
+		// as the event bubbles up through nested menus
+		this.mouseHandled = false;
+		this.element
+			.uniqueId()
+			.attr( {
+				role: this.options.role,
+				tabIndex: 0
+			} );
+
+		this._addClass( "ui-menu", "ui-widget ui-widget-content" );
+		this._on( {
+
+			// Prevent focus from sticking to links inside menu after clicking
+			// them (focus should always stay on UL during navigation).
+			"mousedown .ui-menu-item": function( event ) {
+				event.preventDefault();
+			},
+			"click .ui-menu-item": function( event ) {
+				var target = $( event.target );
+				var active = $( $.ui.safeActiveElement( this.document[ 0 ] ) );
+				if ( !this.mouseHandled && target.not( ".ui-state-disabled" ).length ) {
+					this.select( event );
+
+					// Only set the mouseHandled flag if the event will bubble, see #9469.
+					if ( !event.isPropagationStopped() ) {
+						this.mouseHandled = true;
+					}
+
+					// Open submenu on click
+					if ( target.has( ".ui-menu" ).length ) {
+						this.expand( event );
+					} else if ( !this.element.is( ":focus" ) &&
+							active.closest( ".ui-menu" ).length ) {
+
+						// Redirect focus to the menu
+						this.element.trigger( "focus", [ true ] );
+
+						// If the active item is on the top level, let it stay active.
+						// Otherwise, blur the active item since it is no longer visible.
+						if ( this.active && this.active.parents( ".ui-menu" ).length === 1 ) {
+							clearTimeout( this.timer );
+						}
+					}
+				}
+			},
+			"mouseenter .ui-menu-item": function( event ) {
+
+				// Ignore mouse events while typeahead is active, see #10458.
+				// Prevents focusing the wrong item when typeahead causes a scroll while the mouse
+				// is over an item in the menu
+				if ( this.previousFilter ) {
+					return;
+				}
+
+				var actualTarget = $( event.target ).closest( ".ui-menu-item" ),
+					target = $( event.currentTarget );
+
+				// Ignore bubbled events on parent items, see #11641
+				if ( actualTarget[ 0 ] !== target[ 0 ] ) {
+					return;
+				}
+
+				// Remove ui-state-active class from siblings of the newly focused menu item
+				// to avoid a jump caused by adjacent elements both having a class with a border
+				this._removeClass( target.siblings().children( ".ui-state-active" ),
+					null, "ui-state-active" );
+				this.focus( event, target );
+			},
+			mouseleave: "collapseAll",
+			"mouseleave .ui-menu": "collapseAll",
+			focus: function( event, keepActiveItem ) {
+
+				// If there's already an active item, keep it active
+				// If not, activate the first item
+				var item = this.active || this.element.find( this.options.items ).eq( 0 );
+
+				if ( !keepActiveItem ) {
+					this.focus( event, item );
+				}
+			},
+			blur: function( event ) {
+				this._delay( function() {
+					var notContained = !$.contains(
+						this.element[ 0 ],
+						$.ui.safeActiveElement( this.document[ 0 ] )
+					);
+					if ( notContained ) {
+						this.collapseAll( event );
+					}
+				} );
+			},
+			keydown: "_keydown"
+		} );
+
+		this.refresh();
+
+		// Clicks outside of a menu collapse any open menus
+		this._on( this.document, {
+			click: function( event ) {
+				if ( this._closeOnDocumentClick( event ) ) {
+					this.collapseAll( event );
+				}
+
+				// Reset the mouseHandled flag
+				this.mouseHandled = false;
+			}
+		} );
+	},
+
+	_destroy: function() {
+		var items = this.element.find( ".ui-menu-item" )
+				.removeAttr( "role aria-disabled" ),
+			submenus = items.children( ".ui-menu-item-wrapper" )
+				.removeUniqueId()
+				.removeAttr( "tabIndex role aria-haspopup" );
+
+		// Destroy (sub)menus
+		this.element
+			.removeAttr( "aria-activedescendant" )
+			.find( ".ui-menu" ).addBack()
+				.removeAttr( "role aria-labelledby aria-expanded aria-hidden aria-disabled " +
+					"tabIndex" )
+				.removeUniqueId()
+				.show();
+
+		submenus.children().each( function() {
+			var elem = $( this );
+			if ( elem.data( "ui-menu-submenu-caret" ) ) {
+				elem.remove();
+			}
+		} );
+	},
+
+	_keydown: function( event ) {
+		var match, prev, character, skip,
+			preventDefault = true;
+
+		switch ( event.keyCode ) {
+		case $.ui.keyCode.PAGE_UP:
+			this.previousPage( event );
+			break;
+		case $.ui.keyCode.PAGE_DOWN:
+			this.nextPage( event );
+			break;
+		case $.ui.keyCode.HOME:
+			this._move( "first", "first", event );
+			break;
+		case $.ui.keyCode.END:
+			this._move( "last", "last", event );
+			break;
+		case $.ui.keyCode.UP:
+			this.previous( event );
+			break;
+		case $.ui.keyCode.DOWN:
+			this.next( event );
+			break;
+		case $.ui.keyCode.LEFT:
+			this.collapse( event );
+			break;
+		case $.ui.keyCode.RIGHT:
+			if ( this.active && !this.active.is( ".ui-state-disabled" ) ) {
+				this.expand( event );
+			}
+			break;
+		case $.ui.keyCode.ENTER:
+		case $.ui.keyCode.SPACE:
+			this._activate( event );
+			break;
+		case $.ui.keyCode.ESCAPE:
+			this.collapse( event );
+			break;
+		default:
+			preventDefault = false;
+			prev = this.previousFilter || "";
+			character = String.fromCharCode( event.keyCode );
+			skip = false;
+
+			clearTimeout( this.filterTimer );
+
+			if ( character === prev ) {
+				skip = true;
+			} else {
+				character = prev + character;
+			}
+
+			match = this._filterMenuItems( character );
+			match = skip && match.index( this.active.next() ) !== -1 ?
+				this.active.nextAll( ".ui-menu-item" ) :
+				match;
+
+			// If no matches on the current filter, reset to the last character pressed
+			// to move down the menu to the first item that starts with that character
+			if ( !match.length ) {
+				character = String.fromCharCode( event.keyCode );
+				match = this._filterMenuItems( character );
+			}
+
+			if ( match.length ) {
+				this.focus( event, match );
+				this.previousFilter = character;
+				this.filterTimer = this._delay( function() {
+					delete this.previousFilter;
+				}, 1000 );
+			} else {
+				delete this.previousFilter;
+			}
+		}
+
+		if ( preventDefault ) {
+			event.preventDefault();
+		}
+	},
+
+	_activate: function( event ) {
+		if ( this.active && !this.active.is( ".ui-state-disabled" ) ) {
+			if ( this.active.children( "[aria-haspopup='true']" ).length ) {
+				this.expand( event );
+			} else {
+				this.select( event );
+			}
+		}
+	},
+
+	refresh: function() {
+		var menus, items, newSubmenus, newItems, newWrappers,
+			that = this,
+			icon = this.options.icons.submenu,
+			submenus = this.element.find( this.options.menus );
+
+		this._toggleClass( "ui-menu-icons", null, !!this.element.find( ".ui-icon" ).length );
+
+		// Initialize nested menus
+		newSubmenus = submenus.filter( ":not(.ui-menu)" )
+			.hide()
+			.attr( {
+				role: this.options.role,
+				"aria-hidden": "true",
+				"aria-expanded": "false"
+			} )
+			.each( function() {
+				var menu = $( this ),
+					item = menu.prev(),
+					submenuCaret = $( "<span>" ).data( "ui-menu-submenu-caret", true );
+
+				that._addClass( submenuCaret, "ui-menu-icon", "ui-icon " + icon );
+				item
+					.attr( "aria-haspopup", "true" )
+					.prepend( submenuCaret );
+				menu.attr( "aria-labelledby", item.attr( "id" ) );
+			} );
+
+		this._addClass( newSubmenus, "ui-menu", "ui-widget ui-widget-content ui-front" );
+
+		menus = submenus.add( this.element );
+		items = menus.find( this.options.items );
+
+		// Initialize menu-items containing spaces and/or dashes only as dividers
+		items.not( ".ui-menu-item" ).each( function() {
+			var item = $( this );
+			if ( that._isDivider( item ) ) {
+				that._addClass( item, "ui-menu-divider", "ui-widget-content" );
+			}
+		} );
+
+		// Don't refresh list items that are already adapted
+		newItems = items.not( ".ui-menu-item, .ui-menu-divider" );
+		newWrappers = newItems.children()
+			.not( ".ui-menu" )
+				.uniqueId()
+				.attr( {
+					tabIndex: -1,
+					role: this._itemRole()
+				} );
+		this._addClass( newItems, "ui-menu-item" )
+			._addClass( newWrappers, "ui-menu-item-wrapper" );
+
+		// Add aria-disabled attribute to any disabled menu item
+		items.filter( ".ui-state-disabled" ).attr( "aria-disabled", "true" );
+
+		// If the active item has been removed, blur the menu
+		if ( this.active && !$.contains( this.element[ 0 ], this.active[ 0 ] ) ) {
+			this.blur();
+		}
+	},
+
+	_itemRole: function() {
+		return {
+			menu: "menuitem",
+			listbox: "option"
+		}[ this.options.role ];
+	},
+
+	_setOption: function( key, value ) {
+		if ( key === "icons" ) {
+			var icons = this.element.find( ".ui-menu-icon" );
+			this._removeClass( icons, null, this.options.icons.submenu )
+				._addClass( icons, null, value.submenu );
+		}
+		this._super( key, value );
+	},
+
+	_setOptionDisabled: function( value ) {
+		this._super( value );
+
+		this.element.attr( "aria-disabled", String( value ) );
+		this._toggleClass( null, "ui-state-disabled", !!value );
+	},
+
+	focus: function( event, item ) {
+		var nested, focused, activeParent;
+		this.blur( event, event && event.type === "focus" );
+
+		this._scrollIntoView( item );
+
+		this.active = item.first();
+
+		focused = this.active.children( ".ui-menu-item-wrapper" );
+		this._addClass( focused, null, "ui-state-active" );
+
+		// Only update aria-activedescendant if there's a role
+		// otherwise we assume focus is managed elsewhere
+		if ( this.options.role ) {
+			this.element.attr( "aria-activedescendant", focused.attr( "id" ) );
+		}
+
+		// Highlight active parent menu item, if any
+		activeParent = this.active
+			.parent()
+				.closest( ".ui-menu-item" )
+					.children( ".ui-menu-item-wrapper" );
+		this._addClass( activeParent, null, "ui-state-active" );
+
+		if ( event && event.type === "keydown" ) {
+			this._close();
+		} else {
+			this.timer = this._delay( function() {
+				this._close();
+			}, this.delay );
+		}
+
+		nested = item.children( ".ui-menu" );
+		if ( nested.length && event && ( /^mouse/.test( event.type ) ) ) {
+			this._startOpening( nested );
+		}
+		this.activeMenu = item.parent();
+
+		this._trigger( "focus", event, { item: item } );
+	},
+
+	_scrollIntoView: function( item ) {
+		var borderTop, paddingTop, offset, scroll, elementHeight, itemHeight;
+		if ( this._hasScroll() ) {
+			borderTop = parseFloat( $.css( this.activeMenu[ 0 ], "borderTopWidth" ) ) || 0;
+			paddingTop = parseFloat( $.css( this.activeMenu[ 0 ], "paddingTop" ) ) || 0;
+			offset = item.offset().top - this.activeMenu.offset().top - borderTop - paddingTop;
+			scroll = this.activeMenu.scrollTop();
+			elementHeight = this.activeMenu.height();
+			itemHeight = item.outerHeight();
+
+			if ( offset < 0 ) {
+				this.activeMenu.scrollTop( scroll + offset );
+			} else if ( offset + itemHeight > elementHeight ) {
+				this.activeMenu.scrollTop( scroll + offset - elementHeight + itemHeight );
+			}
+		}
+	},
+
+	blur: function( event, fromFocus ) {
+		if ( !fromFocus ) {
+			clearTimeout( this.timer );
+		}
+
+		if ( !this.active ) {
+			return;
+		}
+
+		this._removeClass( this.active.children( ".ui-menu-item-wrapper" ),
+			null, "ui-state-active" );
+
+		this._trigger( "blur", event, { item: this.active } );
+		this.active = null;
+	},
+
+	_startOpening: function( submenu ) {
+		clearTimeout( this.timer );
+
+		// Don't open if already open fixes a Firefox bug that caused a .5 pixel
+		// shift in the submenu position when mousing over the caret icon
+		if ( submenu.attr( "aria-hidden" ) !== "true" ) {
+			return;
+		}
+
+		this.timer = this._delay( function() {
+			this._close();
+			this._open( submenu );
+		}, this.delay );
+	},
+
+	_open: function( submenu ) {
+		var position = $.extend( {
+			of: this.active
+		}, this.options.position );
+
+		clearTimeout( this.timer );
+		this.element.find( ".ui-menu" ).not( submenu.parents( ".ui-menu" ) )
+			.hide()
+			.attr( "aria-hidden", "true" );
+
+		submenu
+			.show()
+			.removeAttr( "aria-hidden" )
+			.attr( "aria-expanded", "true" )
+			.position( position );
+	},
+
+	collapseAll: function( event, all ) {
+		clearTimeout( this.timer );
+		this.timer = this._delay( function() {
+
+			// If we were passed an event, look for the submenu that contains the event
+			var currentMenu = all ? this.element :
+				$( event && event.target ).closest( this.element.find( ".ui-menu" ) );
+
+			// If we found no valid submenu ancestor, use the main menu to close all
+			// sub menus anyway
+			if ( !currentMenu.length ) {
+				currentMenu = this.element;
+			}
+
+			this._close( currentMenu );
+
+			this.blur( event );
+
+			// Work around active item staying active after menu is blurred
+			this._removeClass( currentMenu.find( ".ui-state-active" ), null, "ui-state-active" );
+
+			this.activeMenu = currentMenu;
+		}, this.delay );
+	},
+
+	// With no arguments, closes the currently active menu - if nothing is active
+	// it closes all menus.  If passed an argument, it will search for menus BELOW
+	_close: function( startMenu ) {
+		if ( !startMenu ) {
+			startMenu = this.active ? this.active.parent() : this.element;
+		}
+
+		startMenu.find( ".ui-menu" )
+			.hide()
+			.attr( "aria-hidden", "true" )
+			.attr( "aria-expanded", "false" );
+	},
+
+	_closeOnDocumentClick: function( event ) {
+		return !$( event.target ).closest( ".ui-menu" ).length;
+	},
+
+	_isDivider: function( item ) {
+
+		// Match hyphen, em dash, en dash
+		return !/[^\-\u2014\u2013\s]/.test( item.text() );
+	},
+
+	collapse: function( event ) {
+		var newItem = this.active &&
+			this.active.parent().closest( ".ui-menu-item", this.element );
+		if ( newItem && newItem.length ) {
+			this._close();
+			this.focus( event, newItem );
+		}
+	},
+
+	expand: function( event ) {
+		var newItem = this.active &&
+			this.active
+				.children( ".ui-menu " )
+					.find( this.options.items )
+						.first();
+
+		if ( newItem && newItem.length ) {
+			this._open( newItem.parent() );
+
+			// Delay so Firefox will not hide activedescendant change in expanding submenu from AT
+			this._delay( function() {
+				this.focus( event, newItem );
+			} );
+		}
+	},
+
+	next: function( event ) {
+		this._move( "next", "first", event );
+	},
+
+	previous: function( event ) {
+		this._move( "prev", "last", event );
+	},
+
+	isFirstItem: function() {
+		return this.active && !this.active.prevAll( ".ui-menu-item" ).length;
+	},
+
+	isLastItem: function() {
+		return this.active && !this.active.nextAll( ".ui-menu-item" ).length;
+	},
+
+	_move: function( direction, filter, event ) {
+		var next;
+		if ( this.active ) {
+			if ( direction === "first" || direction === "last" ) {
+				next = this.active
+					[ direction === "first" ? "prevAll" : "nextAll" ]( ".ui-menu-item" )
+					.eq( -1 );
+			} else {
+				next = this.active
+					[ direction + "All" ]( ".ui-menu-item" )
+					.eq( 0 );
+			}
+		}
+		if ( !next || !next.length || !this.active ) {
+			next = this.activeMenu.find( this.options.items )[ filter ]();
+		}
+
+		this.focus( event, next );
+	},
+
+	nextPage: function( event ) {
+		var item, base, height;
+
+		if ( !this.active ) {
+			this.next( event );
+			return;
+		}
+		if ( this.isLastItem() ) {
+			return;
+		}
+		if ( this._hasScroll() ) {
+			base = this.active.offset().top;
+			height = this.element.height();
+			this.active.nextAll( ".ui-menu-item" ).each( function() {
+				item = $( this );
+				return item.offset().top - base - height < 0;
+			} );
+
+			this.focus( event, item );
+		} else {
+			this.focus( event, this.activeMenu.find( this.options.items )
+				[ !this.active ? "first" : "last" ]() );
+		}
+	},
+
+	previousPage: function( event ) {
+		var item, base, height;
+		if ( !this.active ) {
+			this.next( event );
+			return;
+		}
+		if ( this.isFirstItem() ) {
+			return;
+		}
+		if ( this._hasScroll() ) {
+			base = this.active.offset().top;
+			height = this.element.height();
+			this.active.prevAll( ".ui-menu-item" ).each( function() {
+				item = $( this );
+				return item.offset().top - base + height > 0;
+			} );
+
+			this.focus( event, item );
+		} else {
+			this.focus( event, this.activeMenu.find( this.options.items ).first() );
+		}
+	},
+
+	_hasScroll: function() {
+		return this.element.outerHeight() < this.element.prop( "scrollHeight" );
+	},
+
+	select: function( event ) {
+
+		// TODO: It should never be possible to not have an active item at this
+		// point, but the tests don't trigger mouseenter before click.
+		this.active = this.active || $( event.target ).closest( ".ui-menu-item" );
+		var ui = { item: this.active };
+		if ( !this.active.has( ".ui-menu" ).length ) {
+			this.collapseAll( event, true );
+		}
+		this._trigger( "select", event, ui );
+	},
+
+	_filterMenuItems: function( character ) {
+		var escapedCharacter = character.replace( /[\-\[\]{}()*+?.,\\\^$|#\s]/g, "\\$&" ),
+			regex = new RegExp( "^" + escapedCharacter, "i" );
+
+		return this.activeMenu
+			.find( this.options.items )
+
+				// Only match on items, not dividers or other content (#10571)
+				.filter( ".ui-menu-item" )
+					.filter( function() {
+						return regex.test(
+							$.trim( $( this ).children( ".ui-menu-item-wrapper" ).text() ) );
+					} );
+	}
+} );
+
+
+/*!
+ * jQuery UI Autocomplete 1.12.0
+ * http://jqueryui.com
+ *
+ * Copyright jQuery Foundation and other contributors
+ * Released under the MIT license.
+ * http://jquery.org/license
+ */
+
+//>>label: Autocomplete
+//>>group: Widgets
+//>>description: Lists suggested words as the user is typing.
+//>>docs: http://api.jqueryui.com/autocomplete/
+//>>demos: http://jqueryui.com/autocomplete/
+//>>css.structure: ../../themes/base/core.css
+//>>css.structure: ../../themes/base/autocomplete.css
+//>>css.theme: ../../themes/base/theme.css
+
+
+
+$.widget( "ui.autocomplete", {
+	version: "1.12.0",
+	defaultElement: "<input>",
+	options: {
+		appendTo: null,
+		autoFocus: false,
+		delay: 300,
+		minLength: 1,
+		position: {
+			my: "left top",
+			at: "left bottom",
+			collision: "none"
+		},
+		source: null,
+
+		// Callbacks
+		change: null,
+		close: null,
+		focus: null,
+		open: null,
+		response: null,
+		search: null,
+		select: null
+	},
+
+	requestIndex: 0,
+	pending: 0,
+
+	_create: function() {
+
+		// Some browsers only repeat keydown events, not keypress events,
+		// so we use the suppressKeyPress flag to determine if we've already
+		// handled the keydown event. #7269
+		// Unfortunately the code for & in keypress is the same as the up arrow,
+		// so we use the suppressKeyPressRepeat flag to avoid handling keypress
+		// events when we know the keydown event was used to modify the
+		// search term. #7799
+		var suppressKeyPress, suppressKeyPressRepeat, suppressInput,
+			nodeName = this.element[ 0 ].nodeName.toLowerCase(),
+			isTextarea = nodeName === "textarea",
+			isInput = nodeName === "input";
+
+		// Textareas are always multi-line
+		// Inputs are always single-line, even if inside a contentEditable element
+		// IE also treats inputs as contentEditable
+		// All other element types are determined by whether or not they're contentEditable
+		this.isMultiLine = isTextarea || !isInput && this._isContentEditable( this.element );
+
+		this.valueMethod = this.element[ isTextarea || isInput ? "val" : "text" ];
+		this.isNewMenu = true;
+
+		this._addClass( "ui-autocomplete-input" );
+		this.element.attr( "autocomplete", "off" );
+
+		this._on( this.element, {
+			keydown: function( event ) {
+				if ( this.element.prop( "readOnly" ) ) {
+					suppressKeyPress = true;
+					suppressInput = true;
+					suppressKeyPressRepeat = true;
+					return;
+				}
+
+				suppressKeyPress = false;
+				suppressInput = false;
+				suppressKeyPressRepeat = false;
+				var keyCode = $.ui.keyCode;
+				switch ( event.keyCode ) {
+				case keyCode.PAGE_UP:
+					suppressKeyPress = true;
+					this._move( "previousPage", event );
+					break;
+				case keyCode.PAGE_DOWN:
+					suppressKeyPress = true;
+					this._move( "nextPage", event );
+					break;
+				case keyCode.UP:
+					suppressKeyPress = true;
+					this._keyEvent( "previous", event );
+					break;
+				case keyCode.DOWN:
+					suppressKeyPress = true;
+					this._keyEvent( "next", event );
+					break;
+				case keyCode.ENTER:
+
+					// when menu is open and has focus
+					if ( this.menu.active ) {
+
+						// #6055 - Opera still allows the keypress to occur
+						// which causes forms to submit
+						suppressKeyPress = true;
+						event.preventDefault();
+						this.menu.select( event );
+					}
+					break;
+				case keyCode.TAB:
+					if ( this.menu.active ) {
+						this.menu.select( event );
+					}
+					break;
+				case keyCode.ESCAPE:
+					if ( this.menu.element.is( ":visible" ) ) {
+						if ( !this.isMultiLine ) {
+							this._value( this.term );
+						}
+						this.close( event );
+
+						// Different browsers have different default behavior for escape
+						// Single press can mean undo or clear
+						// Double press in IE means clear the whole form
+						event.preventDefault();
+					}
+					break;
+				default:
+					suppressKeyPressRepeat = true;
+
+					// search timeout should be triggered before the input value is changed
+					this._searchTimeout( event );
+					break;
+				}
+			},
+			keypress: function( event ) {
+				if ( suppressKeyPress ) {
+					suppressKeyPress = false;
+					if ( !this.isMultiLine || this.menu.element.is( ":visible" ) ) {
+						event.preventDefault();
+					}
+					return;
+				}
+				if ( suppressKeyPressRepeat ) {
+					return;
+				}
+
+				// Replicate some key handlers to allow them to repeat in Firefox and Opera
+				var keyCode = $.ui.keyCode;
+				switch ( event.keyCode ) {
+				case keyCode.PAGE_UP:
+					this._move( "previousPage", event );
+					break;
+				case keyCode.PAGE_DOWN:
+					this._move( "nextPage", event );
+					break;
+				case keyCode.UP:
+					this._keyEvent( "previous", event );
+					break;
+				case keyCode.DOWN:
+					this._keyEvent( "next", event );
+					break;
+				}
+			},
+			input: function( event ) {
+				if ( suppressInput ) {
+					suppressInput = false;
+					event.preventDefault();
+					return;
+				}
+				this._searchTimeout( event );
+			},
+			focus: function() {
+				this.selectedItem = null;
+				this.previous = this._value();
+			},
+			blur: function( event ) {
+				if ( this.cancelBlur ) {
+					delete this.cancelBlur;
+					return;
+				}
+
+				clearTimeout( this.searching );
+				this.close( event );
+				this._change( event );
+			}
+		} );
+
+		this._initSource();
+		this.menu = $( "<ul>" )
+			.appendTo( this._appendTo() )
+			.menu( {
+
+				// disable ARIA support, the live region takes care of that
+				role: null
+			} )
+			.hide()
+			.menu( "instance" );
+
+		this._addClass( this.menu.element, "ui-autocomplete", "ui-front" );
+		this._on( this.menu.element, {
+			mousedown: function( event ) {
+
+				// prevent moving focus out of the text field
+				event.preventDefault();
+
+				// IE doesn't prevent moving focus even with event.preventDefault()
+				// so we set a flag to know when we should ignore the blur event
+				this.cancelBlur = true;
+				this._delay( function() {
+					delete this.cancelBlur;
+
+					// Support: IE 8 only
+					// Right clicking a menu item or selecting text from the menu items will
+					// result in focus moving out of the input. However, we've already received
+					// and ignored the blur event because of the cancelBlur flag set above. So
+					// we restore focus to ensure that the menu closes properly based on the user's
+					// next actions.
+					if ( this.element[ 0 ] !== $.ui.safeActiveElement( this.document[ 0 ] ) ) {
+						this.element.trigger( "focus" );
+					}
+				} );
+			},
+			menufocus: function( event, ui ) {
+				var label, item;
+
+				// support: Firefox
+				// Prevent accidental activation of menu items in Firefox (#7024 #9118)
+				if ( this.isNewMenu ) {
+					this.isNewMenu = false;
+					if ( event.originalEvent && /^mouse/.test( event.originalEvent.type ) ) {
+						this.menu.blur();
+
+						this.document.one( "mousemove", function() {
+							$( event.target ).trigger( event.originalEvent );
+						} );
+
+						return;
+					}
+				}
+
+				item = ui.item.data( "ui-autocomplete-item" );
+				if ( false !== this._trigger( "focus", event, { item: item } ) ) {
+
+					// use value to match what will end up in the input, if it was a key event
+					if ( event.originalEvent && /^key/.test( event.originalEvent.type ) ) {
+						this._value( item.value );
+					}
+				}
+
+				// Announce the value in the liveRegion
+				label = ui.item.attr( "aria-label" ) || item.value;
+				if ( label && $.trim( label ).length ) {
+					this.liveRegion.children().hide();
+					$( "<div>" ).text( label ).appendTo( this.liveRegion );
+				}
+			},
+			menuselect: function( event, ui ) {
+				var item = ui.item.data( "ui-autocomplete-item" ),
+					previous = this.previous;
+
+				// Only trigger when focus was lost (click on menu)
+				if ( this.element[ 0 ] !== $.ui.safeActiveElement( this.document[ 0 ] ) ) {
+					this.element.trigger( "focus" );
+					this.previous = previous;
+
+					// #6109 - IE triggers two focus events and the second
+					// is asynchronous, so we need to reset the previous
+					// term synchronously and asynchronously :-(
+					this._delay( function() {
+						this.previous = previous;
+						this.selectedItem = item;
+					} );
+				}
+
+				if ( false !== this._trigger( "select", event, { item: item } ) ) {
+					this._value( item.value );
+				}
+
+				// reset the term after the select event
+				// this allows custom select handling to work properly
+				this.term = this._value();
+
+				this.close( event );
+				this.selectedItem = item;
+			}
+		} );
+
+		this.liveRegion = $( "<div>", {
+			role: "status",
+			"aria-live": "assertive",
+			"aria-relevant": "additions"
+		} )
+			.appendTo( this.document[ 0 ].body );
+
+		this._addClass( this.liveRegion, null, "ui-helper-hidden-accessible" );
+
+		// Turning off autocomplete prevents the browser from remembering the
+		// value when navigating through history, so we re-enable autocomplete
+		// if the page is unloaded before the widget is destroyed. #7790
+		this._on( this.window, {
+			beforeunload: function() {
+				this.element.removeAttr( "autocomplete" );
+			}
+		} );
+	},
+
+	_destroy: function() {
+		clearTimeout( this.searching );
+		this.element.removeAttr( "autocomplete" );
+		this.menu.element.remove();
+		this.liveRegion.remove();
+	},
+
+	_setOption: function( key, value ) {
+		this._super( key, value );
+		if ( key === "source" ) {
+			this._initSource();
+		}
+		if ( key === "appendTo" ) {
+			this.menu.element.appendTo( this._appendTo() );
+		}
+		if ( key === "disabled" && value && this.xhr ) {
+			this.xhr.abort();
+		}
+	},
+
+	_isEventTargetInWidget: function( event ) {
+		var menuElement = this.menu.element[ 0 ];
+
+		return event.target === this.element[ 0 ] ||
+			event.target === menuElement ||
+			$.contains( menuElement, event.target );
+	},
+
+	_closeOnClickOutside: function( event ) {
+		if ( !this._isEventTargetInWidget( event ) ) {
+			this.close();
+		}
+	},
+
+	_appendTo: function() {
+		var element = this.options.appendTo;
+
+		if ( element ) {
+			element = element.jquery || element.nodeType ?
+				$( element ) :
+				this.document.find( element ).eq( 0 );
+		}
+
+		if ( !element || !element[ 0 ] ) {
+			element = this.element.closest( ".ui-front, dialog" );
+		}
+
+		if ( !element.length ) {
+			element = this.document[ 0 ].body;
+		}
+
+		return element;
+	},
+
+	_initSource: function() {
+		var array, url,
+			that = this;
+		if ( $.isArray( this.options.source ) ) {
+			array = this.options.source;
+			this.source = function( request, response ) {
+				response( $.ui.autocomplete.filter( array, request.term ) );
+			};
+		} else if ( typeof this.options.source === "string" ) {
+			url = this.options.source;
+			this.source = function( request, response ) {
+				if ( that.xhr ) {
+					that.xhr.abort();
+				}
+				that.xhr = $.ajax( {
+					url: url,
+					data: request,
+					dataType: "json",
+					success: function( data ) {
+						response( data );
+					},
+					error: function() {
+						response( [] );
+					}
+				} );
+			};
+		} else {
+			this.source = this.options.source;
+		}
+	},
+
+	_searchTimeout: function( event ) {
+		clearTimeout( this.searching );
+		this.searching = this._delay( function() {
+
+			// Search if the value has changed, or if the user retypes the same value (see #7434)
+			var equalValues = this.term === this._value(),
+				menuVisible = this.menu.element.is( ":visible" ),
+				modifierKey = event.altKey || event.ctrlKey || event.metaKey || event.shiftKey;
+
+			if ( !equalValues || ( equalValues && !menuVisible && !modifierKey ) ) {
+				this.selectedItem = null;
+				this.search( null, event );
+			}
+		}, this.options.delay );
+	},
+
+	search: function( value, event ) {
+		value = value != null ? value : this._value();
+
+		// Always save the actual value, not the one passed as an argument
+		this.term = this._value();
+
+		if ( value.length < this.options.minLength ) {
+			return this.close( event );
+		}
+
+		if ( this._trigger( "search", event ) === false ) {
+			return;
+		}
+
+		return this._search( value );
+	},
+
+	_search: function( value ) {
+		this.pending++;
+		this._addClass( "ui-autocomplete-loading" );
+		this.cancelSearch = false;
+
+		this.source( { term: value }, this._response() );
+	},
+
+	_response: function() {
+		var index = ++this.requestIndex;
+
+		return $.proxy( function( content ) {
+			if ( index === this.requestIndex ) {
+				this.__response( content );
+			}
+
+			this.pending--;
+			if ( !this.pending ) {
+				this._removeClass( "ui-autocomplete-loading" );
+			}
+		}, this );
+	},
+
+	__response: function( content ) {
+		if ( content ) {
+			content = this._normalize( content );
+		}
+		this._trigger( "response", null, { content: content } );
+		if ( !this.options.disabled && content && content.length && !this.cancelSearch ) {
+			this._suggest( content );
+			this._trigger( "open" );
+		} else {
+
+			// use ._close() instead of .close() so we don't cancel future searches
+			this._close();
+		}
+	},
+
+	close: function( event ) {
+		this.cancelSearch = true;
+		this._close( event );
+	},
+
+	_close: function( event ) {
+
+		// Remove the handler that closes the menu on outside clicks
+		this._off( this.document, "mousedown" );
+
+		if ( this.menu.element.is( ":visible" ) ) {
+			this.menu.element.hide();
+			this.menu.blur();
+			this.isNewMenu = true;
+			this._trigger( "close", event );
+		}
+	},
+
+	_change: function( event ) {
+		if ( this.previous !== this._value() ) {
+			this._trigger( "change", event, { item: this.selectedItem } );
+		}
+	},
+
+	_normalize: function( items ) {
+
+		// assume all items have the right format when the first item is complete
+		if ( items.length && items[ 0 ].label && items[ 0 ].value ) {
+			return items;
+		}
+		return $.map( items, function( item ) {
+			if ( typeof item === "string" ) {
+				return {
+					label: item,
+					value: item
+				};
+			}
+			return $.extend( {}, item, {
+				label: item.label || item.value,
+				value: item.value || item.label
+			} );
+		} );
+	},
+
+	_suggest: function( items ) {
+		var ul = this.menu.element.empty();
+		this._renderMenu( ul, items );
+		this.isNewMenu = true;
+		this.menu.refresh();
+
+		// Size and position menu
+		ul.show();
+		this._resizeMenu();
+		ul.position( $.extend( {
+			of: this.element
+		}, this.options.position ) );
+
+		if ( this.options.autoFocus ) {
+			this.menu.next();
+		}
+
+		// Listen for interactions outside of the widget (#6642)
+		this._on( this.document, {
+			mousedown: "_closeOnClickOutside"
+		} );
+	},
+
+	_resizeMenu: function() {
+		var ul = this.menu.element;
+		ul.outerWidth( Math.max(
+
+			// Firefox wraps long text (possibly a rounding bug)
+			// so we add 1px to avoid the wrapping (#7513)
+			ul.width( "" ).outerWidth() + 1,
+			this.element.outerWidth()
+		) );
+	},
+
+	_renderMenu: function( ul, items ) {
+		var that = this;
+		$.each( items, function( index, item ) {
+			that._renderItemData( ul, item );
+		} );
+	},
+
+	_renderItemData: function( ul, item ) {
+		return this._renderItem( ul, item ).data( "ui-autocomplete-item", item );
+	},
+
+	_renderItem: function( ul, item ) {
+		return $( "<li>" )
+			.append( $( "<div>" ).text( item.label ) )
+			.appendTo( ul );
+	},
+
+	_move: function( direction, event ) {
+		if ( !this.menu.element.is( ":visible" ) ) {
+			this.search( null, event );
+			return;
+		}
+		if ( this.menu.isFirstItem() && /^previous/.test( direction ) ||
+				this.menu.isLastItem() && /^next/.test( direction ) ) {
+
+			if ( !this.isMultiLine ) {
+				this._value( this.term );
+			}
+
+			this.menu.blur();
+			return;
+		}
+		this.menu[ direction ]( event );
+	},
+
+	widget: function() {
+		return this.menu.element;
+	},
+
+	_value: function() {
+		return this.valueMethod.apply( this.element, arguments );
+	},
+
+	_keyEvent: function( keyEvent, event ) {
+		if ( !this.isMultiLine || this.menu.element.is( ":visible" ) ) {
+			this._move( keyEvent, event );
+
+			// Prevents moving cursor to beginning/end of the text field in some browsers
+			event.preventDefault();
+		}
+	},
+
+	// Support: Chrome <=50
+	// We should be able to just use this.element.prop( "isContentEditable" )
+	// but hidden elements always report false in Chrome.
+	// https://code.google.com/p/chromium/issues/detail?id=313082
+	_isContentEditable: function( element ) {
+		if ( !element.length ) {
+			return false;
+		}
+
+		var editable = element.prop( "contentEditable" );
+
+		if ( editable === "inherit" ) {
+		  return this._isContentEditable( element.parent() );
+		}
+
+		return editable === "true";
+	}
+} );
+
+$.extend( $.ui.autocomplete, {
+	escapeRegex: function( value ) {
+		return value.replace( /[\-\[\]{}()*+?.,\\\^$|#\s]/g, "\\$&" );
+	},
+	filter: function( array, term ) {
+		var matcher = new RegExp( $.ui.autocomplete.escapeRegex( term ), "i" );
+		return $.grep( array, function( value ) {
+			return matcher.test( value.label || value.value || value );
+		} );
+	}
+} );
+
+// Live region extension, adding a `messages` option
+// NOTE: This is an experimental API. We are still investigating
+// a full solution for string manipulation and internationalization.
+$.widget( "ui.autocomplete", $.ui.autocomplete, {
+	options: {
+		messages: {
+			noResults: "No search results.",
+			results: function( amount ) {
+				return amount + ( amount > 1 ? " results are" : " result is" ) +
+					" available, use up and down arrow keys to navigate.";
+			}
+		}
+	},
+
+	__response: function( content ) {
+		var message;
+		this._superApply( arguments );
+		if ( this.options.disabled || this.cancelSearch ) {
+			return;
+		}
+		if ( content && content.length ) {
+			message = this.options.messages.results( content.length );
+		} else {
+			message = this.options.messages.noResults;
+		}
+		this.liveRegion.children().hide();
+		$( "<div>" ).text( message ).appendTo( this.liveRegion );
+	}
+} );
+
+var widgetsAutocomplete = $.ui.autocomplete;
+
+
+/*!
+ * jQuery UI Controlgroup 1.12.0
+ * http://jqueryui.com
+ *
+ * Copyright jQuery Foundation and other contributors
+ * Released under the MIT license.
+ * http://jquery.org/license
+ */
+
+//>>label: Controlgroup
+//>>group: Widgets
+//>>description: Visually groups form control widgets
+//>>docs: http://api.jqueryui.com/controlgroup/
+//>>demos: http://jqueryui.com/controlgroup/
+//>>css.structure: ../../themes/base/core.css
+//>>css.structure: ../../themes/base/controlgroup.css
+//>>css.theme: ../../themes/base/theme.css
+
+
+var controlgroupCornerRegex = /ui-corner-([a-z]){2,6}/g;
+
+var widgetsControlgroup = $.widget( "ui.controlgroup", {
+	version: "1.12.0",
+	defaultElement: "<div>",
+	options: {
+		direction: "horizontal",
+		disabled: null,
+		onlyVisible: true,
+		items: {
+			"button": "input[type=button], input[type=submit], input[type=reset], button, a",
+			"controlgroupLabel": ".ui-controlgroup-label",
+			"checkboxradio": "input[type='checkbox'], input[type='radio']",
+			"selectmenu": "select",
+			"spinner": ".ui-spinner-input"
+		}
+	},
+
+	_create: function() {
+		this._enhance();
+	},
+
+	// To support the enhanced option in jQuery Mobile, we isolate DOM manipulation
+	_enhance: function() {
+		this.element.attr( "role", "toolbar" );
+		this.refresh();
+	},
+
+	_destroy: function() {
+		this._callChildMethod( "destroy" );
+		this.childWidgets.removeData( "ui-controlgroup-data" );
+		this.element.removeAttr( "role" );
+		if ( this.options.items.controlgroupLabel ) {
+			this.element
+				.find( this.options.items.controlgroupLabel )
+				.find( ".ui-controlgroup-label-contents" )
+				.contents().unwrap();
+		}
+	},
+
+	_initWidgets: function() {
+		var that = this,
+			childWidgets = [];
+
+		// First we iterate over each of the items options
+		$.each( this.options.items, function( widget, selector ) {
+			var labels;
+			var options = {};
+
+			// Make sure the widget has a selector set
+			if ( !selector ) {
+				return;
+			}
+
+			if ( widget === "controlgroupLabel" ) {
+				labels = that.element.find( selector );
+				labels.each( function() {
+					var element = $( this );
+
+					if ( element.children( ".ui-controlgroup-label-contents" ).length ) {
+						return;
+					}
+					element.contents()
+						.wrapAll( "<span class='ui-controlgroup-label-contents'></span>" );
+				} );
+				that._addClass( labels, null, "ui-widget ui-widget-content ui-state-default" );
+				childWidgets = childWidgets.concat( labels.get() );
+				return;
+			}
+
+			// Make sure the widget actually exists
+			if ( !$.fn[ widget ] ) {
+				return;
+			}
+
+			// We assume everything is in the middle to start because we can't determine
+			// first / last elements until all enhancments are done.
+			if ( that[ "_" + widget + "Options" ] ) {
+				options = that[ "_" + widget + "Options" ]( "middle" );
+			}
+
+			// Find instances of this widget inside controlgroup and init them
+			that.element
+				.find( selector )
+				.each( function() {
+					var element = $( this );
+					var instance = element[ widget ]( "instance" );
+
+					// We need to clone the default options for this type of widget to avoid
+					// polluting the variable options which has a wider scope than a single widget.
+					var instanceOptions = $.widget.extend( {}, options );
+
+					// If the button is the child of a spinner ignore it
+					// TODO: Find a more generic solution
+					if ( widget === "button" && element.parent( ".ui-spinner" ).length ) {
+						return;
+					}
+
+					// Create the widget if it doesn't exist
+					if ( !instance ) {
+						instance = element[ widget ]()[ widget ]( "instance" );
+					}
+					if ( instance ) {
+						instanceOptions.classes =
+							that._resolveClassesValues( instanceOptions.classes, instance );
+					}
+					element[ widget ]( instanceOptions );
+
+					// Store an instance of the controlgroup to be able to reference
+					// from the outermost element for changing options and refresh
+					var widgetElement = element[ widget ]( "widget" );
+					$.data( widgetElement[ 0 ], "ui-controlgroup-data",
+						instance ? instance : element[ widget ]( "instance" ) );
+
+					childWidgets.push( widgetElement[ 0 ] );
+				} );
+		} );
+
+		this.childWidgets = $( $.unique( childWidgets ) );
+		this._addClass( this.childWidgets, "ui-controlgroup-item" );
+	},
+
+	_callChildMethod: function( method ) {
+		this.childWidgets.each( function() {
+			var element = $( this ),
+				data = element.data( "ui-controlgroup-data" );
+			if ( data && data[ method ] ) {
+				data[ method ]();
+			}
+		} );
+	},
+
+	_updateCornerClass: function( element, position ) {
+		var remove = "ui-corner-top ui-corner-bottom ui-corner-left ui-corner-right ui-corner-all";
+		var add = this._buildSimpleOptions( position, "label" ).classes.label;
+
+		this._removeClass( element, null, remove );
+		this._addClass( element, null, add );
+	},
+
+	_buildSimpleOptions: function( position, key ) {
+		var direction = this.options.direction === "vertical";
+		var result = {
+			classes: {}
+		};
+		result.classes[ key ] = {
+			"middle": "",
+			"first": "ui-corner-" + ( direction ? "top" : "left" ),
+			"last": "ui-corner-" + ( direction ? "bottom" : "right" ),
+			"only": "ui-corner-all"
+		}[ position ];
+
+		return result;
+	},
+
+	_spinnerOptions: function( position ) {
+		var options = this._buildSimpleOptions( position, "ui-spinner" );
+
+		options.classes[ "ui-spinner-up" ] = "";
+		options.classes[ "ui-spinner-down" ] = "";
+
+		return options;
+	},
+
+	_buttonOptions: function( position ) {
+		return this._buildSimpleOptions( position, "ui-button" );
+	},
+
+	_checkboxradioOptions: function( position ) {
+		return this._buildSimpleOptions( position, "ui-checkboxradio-label" );
+	},
+
+	_selectmenuOptions: function( position ) {
+		var direction = this.options.direction === "vertical";
+		return {
+			width: direction ? "auto" : false,
+			classes: {
+				middle: {
+					"ui-selectmenu-button-open": "",
+					"ui-selectmenu-button-closed": ""
+				},
+				first: {
+					"ui-selectmenu-button-open": "ui-corner-" + ( direction ? "top" : "tl" ),
+					"ui-selectmenu-button-closed": "ui-corner-" + ( direction ? "top" : "left" )
+				},
+				last: {
+					"ui-selectmenu-button-open": direction ? "" : "ui-corner-tr",
+					"ui-selectmenu-button-closed": "ui-corner-" + ( direction ? "bottom" : "right" )
+				},
+				only: {
+					"ui-selectmenu-button-open": "ui-corner-top",
+					"ui-selectmenu-button-closed": "ui-corner-all"
+				}
+
+			}[ position ]
+		};
+	},
+
+	_resolveClassesValues: function( classes, instance ) {
+		var result = {};
+		$.each( classes, function( key ) {
+			var current = instance.options.classes[ key ] || "";
+			current = current.replace( controlgroupCornerRegex, "" ).trim();
+			result[ key ] = ( current + " " + classes[ key ] ).replace( /\s+/g, " " );
+		} );
+		return result;
+	},
+
+	_setOption: function( key, value ) {
+		if ( key === "direction" ) {
+			this._removeClass( "ui-controlgroup-" + this.options.direction );
+		}
+
+		this._super( key, value );
+		if ( key === "disabled" ) {
+			this._callChildMethod( value ? "disable" : "enable" );
+			return;
+		}
+
+		this.refresh();
+	},
+
+	refresh: function() {
+		var children,
+			that = this;
+
+		this._addClass( "ui-controlgroup ui-controlgroup-" + this.options.direction );
+
+		if ( this.options.direction === "horizontal" ) {
+			this._addClass( null, "ui-helper-clearfix" );
+		}
+		this._initWidgets();
+
+		children = this.childWidgets;
+
+		// We filter here because we need to track all childWidgets not just the visible ones
+		if ( this.options.onlyVisible ) {
+			children = children.filter( ":visible" );
+		}
+
+		if ( children.length ) {
+
+			// We do this last because we need to make sure all enhancment is done
+			// before determining first and last
+			$.each( [ "first", "last" ], function( index, value ) {
+				var instance = children[ value ]().data( "ui-controlgroup-data" );
+
+				if ( instance && that[ "_" + instance.widgetName + "Options" ] ) {
+					var options = that[ "_" + instance.widgetName + "Options" ](
+						children.length === 1 ? "only" : value
+					);
+					options.classes = that._resolveClassesValues( options.classes, instance );
+					instance.element[ instance.widgetName ]( options );
+				} else {
+					that._updateCornerClass( children[ value ](), value );
+				}
+			} );
+
+			// Finally call the refresh method on each of the child widgets.
+			this._callChildMethod( "refresh" );
+		}
+	}
+} );
+
+/*!
+ * jQuery UI Checkboxradio 1.12.0
+ * http://jqueryui.com
+ *
+ * Copyright jQuery Foundation and other contributors
+ * Released under the MIT license.
+ * http://jquery.org/license
+ */
+
+//>>label: Checkboxradio
+//>>group: Widgets
+//>>description: Enhances a form with multiple themeable checkboxes or radio buttons.
+//>>docs: http://api.jqueryui.com/checkboxradio/
+//>>demos: http://jqueryui.com/checkboxradio/
+//>>css.structure: ../../themes/base/core.css
+//>>css.structure: ../../themes/base/button.css
+//>>css.structure: ../../themes/base/checkboxradio.css
+//>>css.theme: ../../themes/base/theme.css
+
+
+
+$.widget( "ui.checkboxradio", [ $.ui.formResetMixin, {
+	version: "1.12.0",
+	options: {
+		disabled: null,
+		label: null,
+		icon: true,
+		classes: {
+			"ui-checkboxradio-label": "ui-corner-all",
+			"ui-checkboxradio-icon": "ui-corner-all"
+		}
+	},
+
+	_getCreateOptions: function() {
+		var disabled, labels;
+		var that = this;
+		var options = this._super() || {};
+
+		// We read the type here, because it makes more sense to throw a element type error first,
+		// rather then the error for lack of a label. Often if its the wrong type, it
+		// won't have a label (e.g. calling on a div, btn, etc)
+		this._readType();
+
+		labels = this.element.labels();
+
+		// If there are multiple labels, use the last one
+		this.label = $( labels[ labels.length - 1 ] );
+		if ( !this.label.length ) {
+			$.error( "No label found for checkboxradio widget" );
+		}
+
+		this.originalLabel = "";
+
+		// We need to get the label text but this may also need to make sure it does not contain the
+		// input itself.
+		this.label.contents().not( this.element ).each( function() {
+
+			// The label contents could be text, html, or a mix. We concat each element to get a
+			// string representation of the label, without the input as part of it.
+			that.originalLabel += this.nodeType === 3 ? $( this ).text() : this.outerHTML;
+		} );
+
+		// Set the label option if we found label text
+		if ( this.originalLabel ) {
+			options.label = this.originalLabel;
+		}
+
+		disabled = this.element[ 0 ].disabled;
+		if ( disabled != null ) {
+			options.disabled = disabled;
+		}
+		return options;
+	},
+
+	_create: function() {
+		var checked = this.element[ 0 ].checked;
+
+		this._bindFormResetHandler();
+
+		if ( this.options.disabled == null ) {
+			this.options.disabled = this.element[ 0 ].disabled;
+		}
+
+		this._setOption( "disabled", this.options.disabled );
+		this._addClass( "ui-checkboxradio", "ui-helper-hidden-accessible" );
+		this._addClass( this.label, "ui-checkboxradio-label", "ui-button ui-widget" );
+
+		if ( this.type === "radio" ) {
+			this._addClass( this.label, "ui-checkboxradio-radio-label" );
+		}
+
+		if ( this.options.label && this.options.label !== this.originalLabel ) {
+			this._updateLabel();
+		} else if ( this.originalLabel ) {
+			this.options.label = this.originalLabel;
+		}
+
+		this._enhance();
+
+		if ( checked ) {
+			this._addClass( this.label, "ui-checkboxradio-checked", "ui-state-active" );
+			if ( this.icon ) {
+				this._addClass( this.icon, null, "ui-state-hover" );
+			}
+		}
+
+		this._on( {
+			change: "_toggleClasses",
+			focus: function() {
+				this._addClass( this.label, null, "ui-state-focus ui-visual-focus" );
+			},
+			blur: function() {
+				this._removeClass( this.label, null, "ui-state-focus ui-visual-focus" );
+			}
+		} );
+	},
+
+	_readType: function() {
+		var nodeName = this.element[ 0 ].nodeName.toLowerCase();
+		this.type = this.element[ 0 ].type;
+		if ( nodeName !== "input" || !/radio|checkbox/.test( this.type ) ) {
+			$.error( "Can't create checkboxradio on element.nodeName=" + nodeName +
+				" and element.type=" + this.type );
+		}
+	},
+
+	// Support jQuery Mobile enhanced option
+	_enhance: function() {
+		this._updateIcon( this.element[ 0 ].checked );
+	},
+
+	widget: function() {
+		return this.label;
+	},
+
+	_getRadioGroup: function() {
+		var group;
+		var name = this.element[ 0 ].name;
+		var nameSelector = "input[name='" + $.ui.escapeSelector( name ) + "']";
+
+		if ( !name ) {
+			return $( [] );
+		}
+
+		if ( this.form.length ) {
+			group = $( this.form[ 0 ].elements ).filter( nameSelector );
+		} else {
+
+			// Not inside a form, check all inputs that also are not inside a form
+			group = $( nameSelector ).filter( function() {
+				return $( this ).form().length === 0;
+			} );
+		}
+
+		return group.not( this.element );
+	},
+
+	_toggleClasses: function() {
+		var checked = this.element[ 0 ].checked;
+		this._toggleClass( this.label, "ui-checkboxradio-checked", "ui-state-active", checked );
+
+		if ( this.options.icon && this.type === "checkbox" ) {
+			this._toggleClass( this.icon, null, "ui-icon-check ui-state-checked", checked )
+				._toggleClass( this.icon, null, "ui-icon-blank", !checked );
+		}
+
+		if ( this.type === "radio" ) {
+			this._getRadioGroup()
+				.each( function() {
+					var instance = $( this ).checkboxradio( "instance" );
+
+					if ( instance ) {
+						instance._removeClass( instance.label,
+							"ui-checkboxradio-checked", "ui-state-active" );
+					}
+				} );
+		}
+	},
+
+	_destroy: function() {
+		this._unbindFormResetHandler();
+
+		if ( this.icon ) {
+			this.icon.remove();
+			this.iconSpace.remove();
+		}
+	},
+
+	_setOption: function( key, value ) {
+
+		// We don't allow the value to be set to nothing
+		if ( key === "label" && !value ) {
+			return;
+		}
+
+		this._super( key, value );
+
+		if ( key === "disabled" ) {
+			this._toggleClass( this.label, null, "ui-state-disabled", value );
+			this.element[ 0 ].disabled = value;
+
+			// Don't refresh when setting disabled
+			return;
+		}
+		this.refresh();
+	},
+
+	_updateIcon: function( checked ) {
+		var toAdd = "ui-icon ui-icon-background ";
+
+		if ( this.options.icon ) {
+			if ( !this.icon ) {
+				this.icon = $( "<span>" );
+				this.iconSpace = $( "<span> </span>" );
+				this._addClass( this.iconSpace, "ui-checkboxradio-icon-space" );
+			}
+
+			if ( this.type === "checkbox" ) {
+				toAdd += checked ? "ui-icon-check ui-state-checked" : "ui-icon-blank";
+				this._removeClass( this.icon, null, checked ? "ui-icon-blank" : "ui-icon-check" );
+			} else {
+				toAdd += "ui-icon-blank";
+			}
+			this._addClass( this.icon, "ui-checkboxradio-icon", toAdd );
+			if ( !checked ) {
+				this._removeClass( this.icon, null, "ui-icon-check ui-state-checked" );
+			}
+			this.icon.prependTo( this.label ).after( this.iconSpace );
+		} else if ( this.icon !== undefined ) {
+			this.icon.remove();
+			this.iconSpace.remove();
+			delete this.icon;
+		}
+	},
+
+	_updateLabel: function() {
+
+		// Remove the contents of the label ( minus the icon, icon space, and input )
+		this.label.contents().not( this.element.add( this.icon ).add( this.iconSpace ) ).remove();
+		this.label.append( this.options.label );
+	},
+
+	refresh: function() {
+		var checked = this.element[ 0 ].checked,
+			isDisabled = this.element[ 0 ].disabled;
+
+		this._updateIcon( checked );
+		this._toggleClass( this.label, "ui-checkboxradio-checked", "ui-state-active", checked );
+		if ( this.options.label !== null ) {
+			this._updateLabel();
+		}
+
+		if ( isDisabled !== this.options.disabled ) {
+			this._setOptions( { "disabled": isDisabled } );
+		}
+	}
+
+} ] );
+
+var widgetsCheckboxradio = $.ui.checkboxradio;
+
+
+/*!
+ * jQuery UI Button 1.12.0
+ * http://jqueryui.com
+ *
+ * Copyright jQuery Foundation and other contributors
+ * Released under the MIT license.
+ * http://jquery.org/license
+ */
+
+//>>label: Button
+//>>group: Widgets
+//>>description: Enhances a form with themeable buttons.
+//>>docs: http://api.jqueryui.com/button/
+//>>demos: http://jqueryui.com/button/
+//>>css.structure: ../../themes/base/core.css
+//>>css.structure: ../../themes/base/button.css
+//>>css.theme: ../../themes/base/theme.css
+
+
+
+$.widget( "ui.button", {
+	version: "1.12.0",
+	defaultElement: "<button>",
+	options: {
+		classes: {
+			"ui-button": "ui-corner-all"
+		},
+		disabled: null,
+		icon: null,
+		iconPosition: "beginning",
+		label: null,
+		showLabel: true
+	},
+
+	_getCreateOptions: function() {
+		var disabled,
+
+			// This is to support cases like in jQuery Mobile where the base widget does have
+			// an implementation of _getCreateOptions
+			options = this._super() || {};
+
+		this.isInput = this.element.is( "input" );
+
+		disabled = this.element[ 0 ].disabled;
+		if ( disabled != null ) {
+			options.disabled = disabled;
+		}
+
+		this.originalLabel = this.isInput ? this.element.val() : this.element.html();
+		if ( this.originalLabel ) {
+			options.label = this.originalLabel;
+		}
+
+		return options;
+	},
+
+	_create: function() {
+		if ( !this.option.showLabel & !this.options.icon ) {
+			this.options.showLabel = true;
+		}
+
+		// We have to check the option again here even though we did in _getCreateOptions,
+		// because null may have been passed on init which would override what was set in
+		// _getCreateOptions
+		if ( this.options.disabled == null ) {
+			this.options.disabled = this.element[ 0 ].disabled || false;
+		}
+
+		this.hasTitle = !!this.element.attr( "title" );
+
+		// Check to see if the label needs to be set or if its already correct
+		if ( this.options.label && this.options.label !== this.originalLabel ) {
+			if ( this.isInput ) {
+				this.element.val( this.options.label );
+			} else {
+				this.element.html( this.options.label );
+			}
+		}
+		this._addClass( "ui-button", "ui-widget" );
+		this._setOption( "disabled", this.options.disabled );
+		this._enhance();
+
+		if ( this.element.is( "a" ) ) {
+			this._on( {
+				"keyup": function( event ) {
+					if ( event.keyCode === $.ui.keyCode.SPACE ) {
+						event.preventDefault();
+
+						// Support: PhantomJS <= 1.9, IE 8 Only
+						// If a native click is available use it so we actually cause navigation
+						// otherwise just trigger a click event
+						if ( this.element[ 0 ].click ) {
+							this.element[ 0 ].click();
+						} else {
+							this.element.trigger( "click" );
+						}
+					}
+				}
+			} );
+		}
+	},
+
+	_enhance: function() {
+		if ( !this.element.is( "button" ) ) {
+			this.element.attr( "role", "button" );
+		}
+
+		if ( this.options.icon ) {
+			this._updateIcon( "icon", this.options.icon );
+			this._updateTooltip();
+		}
+	},
+
+	_updateTooltip: function() {
+		this.title = this.element.attr( "title" );
+
+		if ( !this.options.showLabel && !this.title ) {
+			this.element.attr( "title", this.options.label );
+		}
+	},
+
+	_updateIcon: function( option, value ) {
+		var icon = option !== "iconPosition",
+			position = icon ? this.options.iconPosition : value,
+			displayBlock = position === "top" || position === "bottom";
+
+		// Create icon
+		if ( !this.icon ) {
+			this.icon = $( "<span>" );
+
+			this._addClass( this.icon, "ui-button-icon", "ui-icon" );
+
+			if ( !this.options.showLabel ) {
+				this._addClass( "ui-button-icon-only" );
+			}
+		} else if ( icon ) {
+
+			// If we are updating the icon remove the old icon class
+			this._removeClass( this.icon, null, this.options.icon );
+		}
+
+		// If we are updating the icon add the new icon class
+		if ( icon ) {
+			this._addClass( this.icon, null, value );
+		}
+
+		this._attachIcon( position );
+
+		// If the icon is on top or bottom we need to add the ui-widget-icon-block class and remove
+		// the iconSpace if there is one.
+		if ( displayBlock ) {
+			this._addClass( this.icon, null, "ui-widget-icon-block" );
+			if ( this.iconSpace ) {
+				this.iconSpace.remove();
+			}
+		} else {
+
+			// Position is beginning or end so remove the ui-widget-icon-block class and add the
+			// space if it does not exist
+			if ( !this.iconSpace ) {
+				this.iconSpace = $( "<span> </span>" );
+				this._addClass( this.iconSpace, "ui-button-icon-space" );
+			}
+			this._removeClass( this.icon, null, "ui-wiget-icon-block" );
+			this._attachIconSpace( position );
+		}
+	},
+
+	_destroy: function() {
+		this.element.removeAttr( "role" );
+
+		if ( this.icon ) {
+			this.icon.remove();
+		}
+		if ( this.iconSpace ) {
+			this.iconSpace.remove();
+		}
+		if ( !this.hasTitle ) {
+			this.element.removeAttr( "title" );
+		}
+	},
+
+	_attachIconSpace: function( iconPosition ) {
+		this.icon[ /^(?:end|bottom)/.test( iconPosition ) ? "before" : "after" ]( this.iconSpace );
+	},
+
+	_attachIcon: function( iconPosition ) {
+		this.element[ /^(?:end|bottom)/.test( iconPosition ) ? "append" : "prepend" ]( this.icon );
+	},
+
+	_setOptions: function( options ) {
+		var newShowLabel = options.showLabel === undefined ?
+				this.options.showLabel :
+				options.showLabel,
+			newIcon = options.icon === undefined ? this.options.icon : options.icon;
+
+		if ( !newShowLabel && !newIcon ) {
+			options.showLabel = true;
+		}
+		this._super( options );
+	},
+
+	_setOption: function( key, value ) {
+		if ( key === "icon" ) {
+			if ( value ) {
+				this._updateIcon( key, value );
+			} else if ( this.icon ) {
+				this.icon.remove();
+				if ( this.iconSpace ) {
+					this.iconSpace.remove();
+				}
+			}
+		}
+
+		if ( key === "iconPosition" ) {
+			this._updateIcon( key, value );
+		}
+
+		// Make sure we can't end up with a button that has neither text nor icon
+		if ( key === "showLabel" ) {
+				this._toggleClass( "ui-button-icon-only", null, !value );
+				this._updateTooltip();
+		}
+
+		if ( key === "label" ) {
+			if ( this.isInput ) {
+				this.element.val( value );
+			} else {
+
+				// If there is an icon, append it, else nothing then append the value
+				// this avoids removal of the icon when setting label text
+				this.element.html( value );
+				if ( this.icon ) {
+					this._attachIcon( this.options.iconPosition );
+					this._attachIconSpace( this.options.iconPosition );
+				}
+			}
+		}
+
+		this._super( key, value );
+
+		if ( key === "disabled" ) {
+			this._toggleClass( null, "ui-state-disabled", value );
+			this.element[ 0 ].disabled = value;
+			if ( value ) {
+				this.element.blur();
+			}
+		}
+	},
+
+	refresh: function() {
+
+		// Make sure to only check disabled if its an element that supports this otherwise
+		// check for the disabled class to determine state
+		var isDisabled = this.element.is( "input, button" ) ?
+			this.element[ 0 ].disabled : this.element.hasClass( "ui-button-disabled" );
+
+		if ( isDisabled !== this.options.disabled ) {
+			this._setOptions( { disabled: isDisabled } );
+		}
+
+		this._updateTooltip();
+	}
+} );
+
+// DEPRECATED
+if ( $.uiBackCompat !== false ) {
+
+	// Text and Icons options
+	$.widget( "ui.button", $.ui.button, {
+		options: {
+			text: true,
+			icons: {
+				primary: null,
+				secondary: null
+			}
+		},
+
+		_create: function() {
+			if ( this.options.showLabel && !this.options.text ) {
+				this.options.showLabel = this.options.text;
+			}
+			if ( !this.options.showLabel && this.options.text ) {
+				this.options.text = this.options.showLabel;
+			}
+			if ( !this.options.icon && ( this.options.icons.primary ||
+					this.options.icons.secondary ) ) {
+				if ( this.options.icons.primary ) {
+					this.options.icon = this.options.icons.primary;
+				} else {
+					this.options.icon = this.options.icons.secondary;
+					this.options.iconPosition = "end";
+				}
+			} else if ( this.options.icon ) {
+				this.options.icons.primary = this.options.icon;
+			}
+			this._super();
+		},
+
+		_setOption: function( key, value ) {
+			if ( key === "text" ) {
+				this._super( "showLabel", value );
+				return;
+			}
+			if ( key === "showLabel" ) {
+				this.options.text = value;
+			}
+			if ( key === "icon" ) {
+				this.options.icons.primary = value;
+			}
+			if ( key === "icons" ) {
+				if ( value.primary ) {
+					this._super( "icon", value.primary );
+					this._super( "iconPosition", "beginning" );
+				} else if ( value.secondary ) {
+					this._super( "icon", value.secondary );
+					this._super( "iconPosition", "end" );
+				}
+			}
+			this._superApply( arguments );
+		}
+	} );
+
+	$.fn.button = ( function( orig ) {
+		return function() {
+			if ( !this.length || ( this.length && this[ 0 ].tagName !== "INPUT" ) ||
+					( this.length && this[ 0 ].tagName === "INPUT" && (
+						this.attr( "type" ) !== "checkbox" && this.attr( "type" ) !== "radio"
+					) ) ) {
+				return orig.apply( this, arguments );
+			}
+			if ( !$.ui.checkboxradio ) {
+				$.error( "Checkboxradio widget missing" );
+			}
+			if ( arguments.length === 0 ) {
+				return this.checkboxradio( {
+					"icon": false
+				} );
+			}
+			return this.checkboxradio.apply( this, arguments );
+		};
+	} )( $.fn.button );
+
+	$.fn.buttonset = function() {
+		if ( !$.ui.controlgroup ) {
+			$.error( "Controlgroup widget missing" );
+		}
+		if ( arguments[ 0 ] === "option" && arguments[ 1 ] === "items" && arguments[ 2 ] ) {
+			return this.controlgroup.apply( this,
+				[ arguments[ 0 ], "items.button", arguments[ 2 ] ] );
+		}
+		if ( arguments[ 0 ] === "option" && arguments[ 1 ] === "items" ) {
+			return this.controlgroup.apply( this, [ arguments[ 0 ], "items.button" ] );
+		}
+		if ( typeof arguments[ 0 ] === "object" && arguments[ 0 ].items ) {
+			arguments[ 0 ].items = {
+				button: arguments[ 0 ].items
+			};
+		}
+		return this.controlgroup.apply( this, arguments );
+	};
+}
+
+var widgetsButton = $.ui.button;
+
+
+// jscs:disable maximumLineLength
+/* jscs:disable requireCamelCaseOrUpperCaseIdentifiers */
+/*!
+ * jQuery UI Datepicker 1.12.0
+ * http://jqueryui.com
+ *
+ * Copyright jQuery Foundation and other contributors
+ * Released under the MIT license.
+ * http://jquery.org/license
+ */
+
+//>>label: Datepicker
+//>>group: Widgets
+//>>description: Displays a calendar from an input or inline for selecting dates.
+//>>docs: http://api.jqueryui.com/datepicker/
+//>>demos: http://jqueryui.com/datepicker/
+//>>css.structure: ../../themes/base/core.css
+//>>css.structure: ../../themes/base/datepicker.css
+//>>css.theme: ../../themes/base/theme.css
+
+
+
+$.extend( $.ui, { datepicker: { version: "1.12.0" } } );
+
+var datepicker_instActive;
+
+function datepicker_getZindex( elem ) {
+	var position, value;
+	while ( elem.length && elem[ 0 ] !== document ) {
+
+		// Ignore z-index if position is set to a value where z-index is ignored by the browser
+		// This makes behavior of this function consistent across browsers
+		// WebKit always returns auto if the element is positioned
+		position = elem.css( "position" );
+		if ( position === "absolute" || position === "relative" || position === "fixed" ) {
+
+			// IE returns 0 when zIndex is not specified
+			// other browsers return a string
+			// we ignore the case of nested elements with an explicit value of 0
+			// <div style="z-index: -10;"><div style="z-index: 0;"></div></div>
+			value = parseInt( elem.css( "zIndex" ), 10 );
+			if ( !isNaN( value ) && value !== 0 ) {
+				return value;
+			}
+		}
+		elem = elem.parent();
+	}
+
+	return 0;
+}
+/* Date picker manager.
+   Use the singleton instance of this class, $.datepicker, to interact with the date picker.
+   Settings for (groups of) date pickers are maintained in an instance object,
+   allowing multiple different settings on the same page. */
+
+function Datepicker() {
+	this._curInst = null; // The current instance in use
+	this._keyEvent = false; // If the last event was a key event
+	this._disabledInputs = []; // List of date picker inputs that have been disabled
+	this._datepickerShowing = false; // True if the popup picker is showing , false if not
+	this._inDialog = false; // True if showing within a "dialog", false if not
+	this._mainDivId = "ui-datepicker-div"; // The ID of the main datepicker division
+	this._inlineClass = "ui-datepicker-inline"; // The name of the inline marker class
+	this._appendClass = "ui-datepicker-append"; // The name of the append marker class
+	this._triggerClass = "ui-datepicker-trigger"; // The name of the trigger marker class
+	this._dialogClass = "ui-datepicker-dialog"; // The name of the dialog marker class
+	this._disableClass = "ui-datepicker-disabled"; // The name of the disabled covering marker class
+	this._unselectableClass = "ui-datepicker-unselectable"; // The name of the unselectable cell marker class
+	this._currentClass = "ui-datepicker-current-day"; // The name of the current day marker class
+	this._dayOverClass = "ui-datepicker-days-cell-over"; // The name of the day hover marker class
+	this.regional = []; // Available regional settings, indexed by language code
+	this.regional[ "" ] = { // Default regional settings
+		closeText: "Done", // Display text for close link
+		prevText: "Prev", // Display text for previous month link
+		nextText: "Next", // Display text for next month link
+		currentText: "Today", // Display text for current month link
+		monthNames: [ "January","February","March","April","May","June",
+			"July","August","September","October","November","December" ], // Names of months for drop-down and formatting
+		monthNamesShort: [ "Jan", "Feb", "Mar", "Apr", "May", "Jun", "Jul", "Aug", "Sep", "Oct", "Nov", "Dec" ], // For formatting
+		dayNames: [ "Sunday", "Monday", "Tuesday", "Wednesday", "Thursday", "Friday", "Saturday" ], // For formatting
+		dayNamesShort: [ "Sun", "Mon", "Tue", "Wed", "Thu", "Fri", "Sat" ], // For formatting
+		dayNamesMin: [ "Su","Mo","Tu","We","Th","Fr","Sa" ], // Column headings for days starting at Sunday
+		weekHeader: "Wk", // Column header for week of the year
+		dateFormat: "mm/dd/yy", // See format options on parseDate
+		firstDay: 0, // The first day of the week, Sun = 0, Mon = 1, ...
+		isRTL: false, // True if right-to-left language, false if left-to-right
+		showMonthAfterYear: false, // True if the year select precedes month, false for month then year
+		yearSuffix: "" // Additional text to append to the year in the month headers
+	};
+	this._defaults = { // Global defaults for all the date picker instances
+		showOn: "focus", // "focus" for popup on focus,
+			// "button" for trigger button, or "both" for either
+		showAnim: "fadeIn", // Name of jQuery animation for popup
+		showOptions: {}, // Options for enhanced animations
+		defaultDate: null, // Used when field is blank: actual date,
+			// +/-number for offset from today, null for today
+		appendText: "", // Display text following the input box, e.g. showing the format
+		buttonText: "...", // Text for trigger button
+		buttonImage: "", // URL for trigger button image
+		buttonImageOnly: false, // True if the image appears alone, false if it appears on a button
+		hideIfNoPrevNext: false, // True to hide next/previous month links
+			// if not applicable, false to just disable them
+		navigationAsDateFormat: false, // True if date formatting applied to prev/today/next links
+		gotoCurrent: false, // True if today link goes back to current selection instead
+		changeMonth: false, // True if month can be selected directly, false if only prev/next
+		changeYear: false, // True if year can be selected directly, false if only prev/next
+		yearRange: "c-10:c+10", // Range of years to display in drop-down,
+			// either relative to today's year (-nn:+nn), relative to currently displayed year
+			// (c-nn:c+nn), absolute (nnnn:nnnn), or a combination of the above (nnnn:-n)
+		showOtherMonths: false, // True to show dates in other months, false to leave blank
+		selectOtherMonths: false, // True to allow selection of dates in other months, false for unselectable
+		showWeek: false, // True to show week of the year, false to not show it
+		calculateWeek: this.iso8601Week, // How to calculate the week of the year,
+			// takes a Date and returns the number of the week for it
+		shortYearCutoff: "+10", // Short year values < this are in the current century,
+			// > this are in the previous century,
+			// string value starting with "+" for current year + value
+		minDate: null, // The earliest selectable date, or null for no limit
+		maxDate: null, // The latest selectable date, or null for no limit
+		duration: "fast", // Duration of display/closure
+		beforeShowDay: null, // Function that takes a date and returns an array with
+			// [0] = true if selectable, false if not, [1] = custom CSS class name(s) or "",
+			// [2] = cell title (optional), e.g. $.datepicker.noWeekends
+		beforeShow: null, // Function that takes an input field and
+			// returns a set of custom settings for the date picker
+		onSelect: null, // Define a callback function when a date is selected
+		onChangeMonthYear: null, // Define a callback function when the month or year is changed
+		onClose: null, // Define a callback function when the datepicker is closed
+		numberOfMonths: 1, // Number of months to show at a time
+		showCurrentAtPos: 0, // The position in multipe months at which to show the current month (starting at 0)
+		stepMonths: 1, // Number of months to step back/forward
+		stepBigMonths: 12, // Number of months to step back/forward for the big links
+		altField: "", // Selector for an alternate field to store selected dates into
+		altFormat: "", // The date format to use for the alternate field
+		constrainInput: true, // The input is constrained by the current date format
+		showButtonPanel: false, // True to show button panel, false to not show it
+		autoSize: false, // True to size the input for the date format, false to leave as is
+		disabled: false // The initial disabled state
+	};
+	$.extend( this._defaults, this.regional[ "" ] );
+	this.regional.en = $.extend( true, {}, this.regional[ "" ] );
+	this.regional[ "en-US" ] = $.extend( true, {}, this.regional.en );
+	this.dpDiv = datepicker_bindHover( $( "<div id='" + this._mainDivId + "' class='ui-datepicker ui-widget ui-widget-content ui-helper-clearfix ui-corner-all'></div>" ) );
+}
+
+$.extend( Datepicker.prototype, {
+	/* Class name added to elements to indicate already configured with a date picker. */
+	markerClassName: "hasDatepicker",
+
+	//Keep track of the maximum number of rows displayed (see #7043)
+	maxRows: 4,
+
+	// TODO rename to "widget" when switching to widget factory
+	_widgetDatepicker: function() {
+		return this.dpDiv;
+	},
+
+	/* Override the default settings for all instances of the date picker.
+	 * @param  settings  object - the new settings to use as defaults (anonymous object)
+	 * @return the manager object
+	 */
+	setDefaults: function( settings ) {
+		datepicker_extendRemove( this._defaults, settings || {} );
+		return this;
+	},
+
+	/* Attach the date picker to a jQuery selection.
+	 * @param  target	element - the target input field or division or span
+	 * @param  settings  object - the new settings to use for this date picker instance (anonymous)
+	 */
+	_attachDatepicker: function( target, settings ) {
+		var nodeName, inline, inst;
+		nodeName = target.nodeName.toLowerCase();
+		inline = ( nodeName === "div" || nodeName === "span" );
+		if ( !target.id ) {
+			this.uuid += 1;
+			target.id = "dp" + this.uuid;
+		}
+		inst = this._newInst( $( target ), inline );
+		inst.settings = $.extend( {}, settings || {} );
+		if ( nodeName === "input" ) {
+			this._connectDatepicker( target, inst );
+		} else if ( inline ) {
+			this._inlineDatepicker( target, inst );
+		}
+	},
+
+	/* Create a new instance object. */
+	_newInst: function( target, inline ) {
+		var id = target[ 0 ].id.replace( /([^A-Za-z0-9_\-])/g, "\\\\$1" ); // escape jQuery meta chars
+		return { id: id, input: target, // associated target
+			selectedDay: 0, selectedMonth: 0, selectedYear: 0, // current selection
+			drawMonth: 0, drawYear: 0, // month being drawn
+			inline: inline, // is datepicker inline or not
+			dpDiv: ( !inline ? this.dpDiv : // presentation div
+			datepicker_bindHover( $( "<div class='" + this._inlineClass + " ui-datepicker ui-widget ui-widget-content ui-helper-clearfix ui-corner-all'></div>" ) ) ) };
+	},
+
+	/* Attach the date picker to an input field. */
+	_connectDatepicker: function( target, inst ) {
+		var input = $( target );
+		inst.append = $( [] );
+		inst.trigger = $( [] );
+		if ( input.hasClass( this.markerClassName ) ) {
+			return;
+		}
+		this._attachments( input, inst );
+		input.addClass( this.markerClassName ).on( "keydown", this._doKeyDown ).
+			on( "keypress", this._doKeyPress ).on( "keyup", this._doKeyUp );
+		this._autoSize( inst );
+		$.data( target, "datepicker", inst );
+
+		//If disabled option is true, disable the datepicker once it has been attached to the input (see ticket #5665)
+		if ( inst.settings.disabled ) {
+			this._disableDatepicker( target );
+		}
+	},
+
+	/* Make attachments based on settings. */
+	_attachments: function( input, inst ) {
+		var showOn, buttonText, buttonImage,
+			appendText = this._get( inst, "appendText" ),
+			isRTL = this._get( inst, "isRTL" );
+
+		if ( inst.append ) {
+			inst.append.remove();
+		}
+		if ( appendText ) {
+			inst.append = $( "<span class='" + this._appendClass + "'>" + appendText + "</span>" );
+			input[ isRTL ? "before" : "after" ]( inst.append );
+		}
+
+		input.off( "focus", this._showDatepicker );
+
+		if ( inst.trigger ) {
+			inst.trigger.remove();
+		}
+
+		showOn = this._get( inst, "showOn" );
+		if ( showOn === "focus" || showOn === "both" ) { // pop-up date picker when in the marked field
+			input.on( "focus", this._showDatepicker );
+		}
+		if ( showOn === "button" || showOn === "both" ) { // pop-up date picker when button clicked
+			buttonText = this._get( inst, "buttonText" );
+			buttonImage = this._get( inst, "buttonImage" );
+			inst.trigger = $( this._get( inst, "buttonImageOnly" ) ?
+				$( "<img/>" ).addClass( this._triggerClass ).
+					attr( { src: buttonImage, alt: buttonText, title: buttonText } ) :
+				$( "<button type='button'></button>" ).addClass( this._triggerClass ).
+					html( !buttonImage ? buttonText : $( "<img/>" ).attr(
+					{ src:buttonImage, alt:buttonText, title:buttonText } ) ) );
+			input[ isRTL ? "before" : "after" ]( inst.trigger );
+			inst.trigger.on( "click", function() {
+				if ( $.datepicker._datepickerShowing && $.datepicker._lastInput === input[ 0 ] ) {
+					$.datepicker._hideDatepicker();
+				} else if ( $.datepicker._datepickerShowing && $.datepicker._lastInput !== input[ 0 ] ) {
+					$.datepicker._hideDatepicker();
+					$.datepicker._showDatepicker( input[ 0 ] );
+				} else {
+					$.datepicker._showDatepicker( input[ 0 ] );
+				}
+				return false;
+			} );
+		}
+	},
+
+	/* Apply the maximum length for the date format. */
+	_autoSize: function( inst ) {
+		if ( this._get( inst, "autoSize" ) && !inst.inline ) {
+			var findMax, max, maxI, i,
+				date = new Date( 2009, 12 - 1, 20 ), // Ensure double digits
+				dateFormat = this._get( inst, "dateFormat" );
+
+			if ( dateFormat.match( /[DM]/ ) ) {
+				findMax = function( names ) {
+					max = 0;
+					maxI = 0;
+					for ( i = 0; i < names.length; i++ ) {
+						if ( names[ i ].length > max ) {
+							max = names[ i ].length;
+							maxI = i;
+						}
+					}
+					return maxI;
+				};
+				date.setMonth( findMax( this._get( inst, ( dateFormat.match( /MM/ ) ?
+					"monthNames" : "monthNamesShort" ) ) ) );
+				date.setDate( findMax( this._get( inst, ( dateFormat.match( /DD/ ) ?
+					"dayNames" : "dayNamesShort" ) ) ) + 20 - date.getDay() );
+			}
+			inst.input.attr( "size", this._formatDate( inst, date ).length );
+		}
+	},
+
+	/* Attach an inline date picker to a div. */
+	_inlineDatepicker: function( target, inst ) {
+		var divSpan = $( target );
+		if ( divSpan.hasClass( this.markerClassName ) ) {
+			return;
+		}
+		divSpan.addClass( this.markerClassName ).append( inst.dpDiv );
+		$.data( target, "datepicker", inst );
+		this._setDate( inst, this._getDefaultDate( inst ), true );
+		this._updateDatepicker( inst );
+		this._updateAlternate( inst );
+
+		//If disabled option is true, disable the datepicker before showing it (see ticket #5665)
+		if ( inst.settings.disabled ) {
+			this._disableDatepicker( target );
+		}
+
+		// Set display:block in place of inst.dpDiv.show() which won't work on disconnected elements
+		// http://bugs.jqueryui.com/ticket/7552 - A Datepicker created on a detached div has zero height
+		inst.dpDiv.css( "display", "block" );
+	},
+
+	/* Pop-up the date picker in a "dialog" box.
+	 * @param  input element - ignored
+	 * @param  date	string or Date - the initial date to display
+	 * @param  onSelect  function - the function to call when a date is selected
+	 * @param  settings  object - update the dialog date picker instance's settings (anonymous object)
+	 * @param  pos int[2] - coordinates for the dialog's position within the screen or
+	 *					event - with x/y coordinates or
+	 *					leave empty for default (screen centre)
+	 * @return the manager object
+	 */
+	_dialogDatepicker: function( input, date, onSelect, settings, pos ) {
+		var id, browserWidth, browserHeight, scrollX, scrollY,
+			inst = this._dialogInst; // internal instance
+
+		if ( !inst ) {
+			this.uuid += 1;
+			id = "dp" + this.uuid;
+			this._dialogInput = $( "<input type='text' id='" + id +
+				"' style='position: absolute; top: -100px; width: 0px;'/>" );
+			this._dialogInput.on( "keydown", this._doKeyDown );
+			$( "body" ).append( this._dialogInput );
+			inst = this._dialogInst = this._newInst( this._dialogInput, false );
+			inst.settings = {};
+			$.data( this._dialogInput[ 0 ], "datepicker", inst );
+		}
+		datepicker_extendRemove( inst.settings, settings || {} );
+		date = ( date && date.constructor === Date ? this._formatDate( inst, date ) : date );
+		this._dialogInput.val( date );
+
+		this._pos = ( pos ? ( pos.length ? pos : [ pos.pageX, pos.pageY ] ) : null );
+		if ( !this._pos ) {
+			browserWidth = document.documentElement.clientWidth;
+			browserHeight = document.documentElement.clientHeight;
+			scrollX = document.documentElement.scrollLeft || document.body.scrollLeft;
+			scrollY = document.documentElement.scrollTop || document.body.scrollTop;
+			this._pos = // should use actual width/height below
+				[ ( browserWidth / 2 ) - 100 + scrollX, ( browserHeight / 2 ) - 150 + scrollY ];
+		}
+
+		// Move input on screen for focus, but hidden behind dialog
+		this._dialogInput.css( "left", ( this._pos[ 0 ] + 20 ) + "px" ).css( "top", this._pos[ 1 ] + "px" );
+		inst.settings.onSelect = onSelect;
+		this._inDialog = true;
+		this.dpDiv.addClass( this._dialogClass );
+		this._showDatepicker( this._dialogInput[ 0 ] );
+		if ( $.blockUI ) {
+			$.blockUI( this.dpDiv );
+		}
+		$.data( this._dialogInput[ 0 ], "datepicker", inst );
+		return this;
+	},
+
+	/* Detach a datepicker from its control.
+	 * @param  target	element - the target input field or division or span
+	 */
+	_destroyDatepicker: function( target ) {
+		var nodeName,
+			$target = $( target ),
+			inst = $.data( target, "datepicker" );
+
+		if ( !$target.hasClass( this.markerClassName ) ) {
+			return;
+		}
+
+		nodeName = target.nodeName.toLowerCase();
+		$.removeData( target, "datepicker" );
+		if ( nodeName === "input" ) {
+			inst.append.remove();
+			inst.trigger.remove();
+			$target.removeClass( this.markerClassName ).
+				off( "focus", this._showDatepicker ).
+				off( "keydown", this._doKeyDown ).
+				off( "keypress", this._doKeyPress ).
+				off( "keyup", this._doKeyUp );
+		} else if ( nodeName === "div" || nodeName === "span" ) {
+			$target.removeClass( this.markerClassName ).empty();
+		}
+
+		if ( datepicker_instActive === inst ) {
+			datepicker_instActive = null;
+		}
+	},
+
+	/* Enable the date picker to a jQuery selection.
+	 * @param  target	element - the target input field or division or span
+	 */
+	_enableDatepicker: function( target ) {
+		var nodeName, inline,
+			$target = $( target ),
+			inst = $.data( target, "datepicker" );
+
+		if ( !$target.hasClass( this.markerClassName ) ) {
+			return;
+		}
+
+		nodeName = target.nodeName.toLowerCase();
+		if ( nodeName === "input" ) {
+			target.disabled = false;
+			inst.trigger.filter( "button" ).
+				each( function() { this.disabled = false; } ).end().
+				filter( "img" ).css( { opacity: "1.0", cursor: "" } );
+		} else if ( nodeName === "div" || nodeName === "span" ) {
+			inline = $target.children( "." + this._inlineClass );
+			inline.children().removeClass( "ui-state-disabled" );
+			inline.find( "select.ui-datepicker-month, select.ui-datepicker-year" ).
+				prop( "disabled", false );
+		}
+		this._disabledInputs = $.map( this._disabledInputs,
+			function( value ) { return ( value === target ? null : value ); } ); // delete entry
+	},
+
+	/* Disable the date picker to a jQuery selection.
+	 * @param  target	element - the target input field or division or span
+	 */
+	_disableDatepicker: function( target ) {
+		var nodeName, inline,
+			$target = $( target ),
+			inst = $.data( target, "datepicker" );
+
+		if ( !$target.hasClass( this.markerClassName ) ) {
+			return;
+		}
+
+		nodeName = target.nodeName.toLowerCase();
+		if ( nodeName === "input" ) {
+			target.disabled = true;
+			inst.trigger.filter( "button" ).
+				each( function() { this.disabled = true; } ).end().
+				filter( "img" ).css( { opacity: "0.5", cursor: "default" } );
+		} else if ( nodeName === "div" || nodeName === "span" ) {
+			inline = $target.children( "." + this._inlineClass );
+			inline.children().addClass( "ui-state-disabled" );
+			inline.find( "select.ui-datepicker-month, select.ui-datepicker-year" ).
+				prop( "disabled", true );
+		}
+		this._disabledInputs = $.map( this._disabledInputs,
+			function( value ) { return ( value === target ? null : value ); } ); // delete entry
+		this._disabledInputs[ this._disabledInputs.length ] = target;
+	},
+
+	/* Is the first field in a jQuery collection disabled as a datepicker?
+	 * @param  target	element - the target input field or division or span
+	 * @return boolean - true if disabled, false if enabled
+	 */
+	_isDisabledDatepicker: function( target ) {
+		if ( !target ) {
+			return false;
+		}
+		for ( var i = 0; i < this._disabledInputs.length; i++ ) {
+			if ( this._disabledInputs[ i ] === target ) {
+				return true;
+			}
+		}
+		return false;
+	},
+
+	/* Retrieve the instance data for the target control.
+	 * @param  target  element - the target input field or division or span
+	 * @return  object - the associated instance data
+	 * @throws  error if a jQuery problem getting data
+	 */
+	_getInst: function( target ) {
+		try {
+			return $.data( target, "datepicker" );
+		}
+		catch ( err ) {
+			throw "Missing instance data for this datepicker";
+		}
+	},
+
+	/* Update or retrieve the settings for a date picker attached to an input field or division.
+	 * @param  target  element - the target input field or division or span
+	 * @param  name	object - the new settings to update or
+	 *				string - the name of the setting to change or retrieve,
+	 *				when retrieving also "all" for all instance settings or
+	 *				"defaults" for all global defaults
+	 * @param  value   any - the new value for the setting
+	 *				(omit if above is an object or to retrieve a value)
+	 */
+	_optionDatepicker: function( target, name, value ) {
+		var settings, date, minDate, maxDate,
+			inst = this._getInst( target );
+
+		if ( arguments.length === 2 && typeof name === "string" ) {
+			return ( name === "defaults" ? $.extend( {}, $.datepicker._defaults ) :
+				( inst ? ( name === "all" ? $.extend( {}, inst.settings ) :
+				this._get( inst, name ) ) : null ) );
+		}
+
+		settings = name || {};
+		if ( typeof name === "string" ) {
+			settings = {};
+			settings[ name ] = value;
+		}
+
+		if ( inst ) {
+			if ( this._curInst === inst ) {
+				this._hideDatepicker();
+			}
+
+			date = this._getDateDatepicker( target, true );
+			minDate = this._getMinMaxDate( inst, "min" );
+			maxDate = this._getMinMaxDate( inst, "max" );
+			datepicker_extendRemove( inst.settings, settings );
+
+			// reformat the old minDate/maxDate values if dateFormat changes and a new minDate/maxDate isn't provided
+			if ( minDate !== null && settings.dateFormat !== undefined && settings.minDate === undefined ) {
+				inst.settings.minDate = this._formatDate( inst, minDate );
+			}
+			if ( maxDate !== null && settings.dateFormat !== undefined && settings.maxDate === undefined ) {
+				inst.settings.maxDate = this._formatDate( inst, maxDate );
+			}
+			if ( "disabled" in settings ) {
+				if ( settings.disabled ) {
+					this._disableDatepicker( target );
+				} else {
+					this._enableDatepicker( target );
+				}
+			}
+			this._attachments( $( target ), inst );
+			this._autoSize( inst );
+			this._setDate( inst, date );
+			this._updateAlternate( inst );
+			this._updateDatepicker( inst );
+		}
+	},
+
+	// Change method deprecated
+	_changeDatepicker: function( target, name, value ) {
+		this._optionDatepicker( target, name, value );
+	},
+
+	/* Redraw the date picker attached to an input field or division.
+	 * @param  target  element - the target input field or division or span
+	 */
+	_refreshDatepicker: function( target ) {
+		var inst = this._getInst( target );
+		if ( inst ) {
+			this._updateDatepicker( inst );
+		}
+	},
+
+	/* Set the dates for a jQuery selection.
+	 * @param  target element - the target input field or division or span
+	 * @param  date	Date - the new date
+	 */
+	_setDateDatepicker: function( target, date ) {
+		var inst = this._getInst( target );
+		if ( inst ) {
+			this._setDate( inst, date );
+			this._updateDatepicker( inst );
+			this._updateAlternate( inst );
+		}
+	},
+
+	/* Get the date(s) for the first entry in a jQuery selection.
+	 * @param  target element - the target input field or division or span
+	 * @param  noDefault boolean - true if no default date is to be used
+	 * @return Date - the current date
+	 */
+	_getDateDatepicker: function( target, noDefault ) {
+		var inst = this._getInst( target );
+		if ( inst && !inst.inline ) {
+			this._setDateFromField( inst, noDefault );
+		}
+		return ( inst ? this._getDate( inst ) : null );
+	},
+
+	/* Handle keystrokes. */
+	_doKeyDown: function( event ) {
+		var onSelect, dateStr, sel,
+			inst = $.datepicker._getInst( event.target ),
+			handled = true,
+			isRTL = inst.dpDiv.is( ".ui-datepicker-rtl" );
+
+		inst._keyEvent = true;
+		if ( $.datepicker._datepickerShowing ) {
+			switch ( event.keyCode ) {
+				case 9: $.datepicker._hideDatepicker();
+						handled = false;
+						break; // hide on tab out
+				case 13: sel = $( "td." + $.datepicker._dayOverClass + ":not(." +
+									$.datepicker._currentClass + ")", inst.dpDiv );
+						if ( sel[ 0 ] ) {
+							$.datepicker._selectDay( event.target, inst.selectedMonth, inst.selectedYear, sel[ 0 ] );
+						}
+
+						onSelect = $.datepicker._get( inst, "onSelect" );
+						if ( onSelect ) {
+							dateStr = $.datepicker._formatDate( inst );
+
+							// Trigger custom callback
+							onSelect.apply( ( inst.input ? inst.input[ 0 ] : null ), [ dateStr, inst ] );
+						} else {
+							$.datepicker._hideDatepicker();
+						}
+
+						return false; // don't submit the form
+				case 27: $.datepicker._hideDatepicker();
+						break; // hide on escape
+				case 33: $.datepicker._adjustDate( event.target, ( event.ctrlKey ?
+							-$.datepicker._get( inst, "stepBigMonths" ) :
+							-$.datepicker._get( inst, "stepMonths" ) ), "M" );
+						break; // previous month/year on page up/+ ctrl
+				case 34: $.datepicker._adjustDate( event.target, ( event.ctrlKey ?
+							+$.datepicker._get( inst, "stepBigMonths" ) :
+							+$.datepicker._get( inst, "stepMonths" ) ), "M" );
+						break; // next month/year on page down/+ ctrl
+				case 35: if ( event.ctrlKey || event.metaKey ) {
+							$.datepicker._clearDate( event.target );
+						}
+						handled = event.ctrlKey || event.metaKey;
+						break; // clear on ctrl or command +end
+				case 36: if ( event.ctrlKey || event.metaKey ) {
+							$.datepicker._gotoToday( event.target );
+						}
+						handled = event.ctrlKey || event.metaKey;
+						break; // current on ctrl or command +home
+				case 37: if ( event.ctrlKey || event.metaKey ) {
+							$.datepicker._adjustDate( event.target, ( isRTL ? +1 : -1 ), "D" );
+						}
+						handled = event.ctrlKey || event.metaKey;
+
+						// -1 day on ctrl or command +left
+						if ( event.originalEvent.altKey ) {
+							$.datepicker._adjustDate( event.target, ( event.ctrlKey ?
+								-$.datepicker._get( inst, "stepBigMonths" ) :
+								-$.datepicker._get( inst, "stepMonths" ) ), "M" );
+						}
+
+						// next month/year on alt +left on Mac
+						break;
+				case 38: if ( event.ctrlKey || event.metaKey ) {
+							$.datepicker._adjustDate( event.target, -7, "D" );
+						}
+						handled = event.ctrlKey || event.metaKey;
+						break; // -1 week on ctrl or command +up
+				case 39: if ( event.ctrlKey || event.metaKey ) {
+							$.datepicker._adjustDate( event.target, ( isRTL ? -1 : +1 ), "D" );
+						}
+						handled = event.ctrlKey || event.metaKey;
+
+						// +1 day on ctrl or command +right
+						if ( event.originalEvent.altKey ) {
+							$.datepicker._adjustDate( event.target, ( event.ctrlKey ?
+								+$.datepicker._get( inst, "stepBigMonths" ) :
+								+$.datepicker._get( inst, "stepMonths" ) ), "M" );
+						}
+
+						// next month/year on alt +right
+						break;
+				case 40: if ( event.ctrlKey || event.metaKey ) {
+							$.datepicker._adjustDate( event.target, +7, "D" );
+						}
+						handled = event.ctrlKey || event.metaKey;
+						break; // +1 week on ctrl or command +down
+				default: handled = false;
+			}
+		} else if ( event.keyCode === 36 && event.ctrlKey ) { // display the date picker on ctrl+home
+			$.datepicker._showDatepicker( this );
+		} else {
+			handled = false;
+		}
+
+		if ( handled ) {
+			event.preventDefault();
+			event.stopPropagation();
+		}
+	},
+
+	/* Filter entered characters - based on date format. */
+	_doKeyPress: function( event ) {
+		var chars, chr,
+			inst = $.datepicker._getInst( event.target );
+
+		if ( $.datepicker._get( inst, "constrainInput" ) ) {
+			chars = $.datepicker._possibleChars( $.datepicker._get( inst, "dateFormat" ) );
+			chr = String.fromCharCode( event.charCode == null ? event.keyCode : event.charCode );
+			return event.ctrlKey || event.metaKey || ( chr < " " || !chars || chars.indexOf( chr ) > -1 );
+		}
+	},
+
+	/* Synchronise manual entry and field/alternate field. */
+	_doKeyUp: function( event ) {
+		var date,
+			inst = $.datepicker._getInst( event.target );
+
+		if ( inst.input.val() !== inst.lastVal ) {
+			try {
+				date = $.datepicker.parseDate( $.datepicker._get( inst, "dateFormat" ),
+					( inst.input ? inst.input.val() : null ),
+					$.datepicker._getFormatConfig( inst ) );
+
+				if ( date ) { // only if valid
+					$.datepicker._setDateFromField( inst );
+					$.datepicker._updateAlternate( inst );
+					$.datepicker._updateDatepicker( inst );
+				}
+			}
+			catch ( err ) {
+			}
+		}
+		return true;
+	},
+
+	/* Pop-up the date picker for a given input field.
+	 * If false returned from beforeShow event handler do not show.
+	 * @param  input  element - the input field attached to the date picker or
+	 *					event - if triggered by focus
+	 */
+	_showDatepicker: function( input ) {
+		input = input.target || input;
+		if ( input.nodeName.toLowerCase() !== "input" ) { // find from button/image trigger
+			input = $( "input", input.parentNode )[ 0 ];
+		}
+
+		if ( $.datepicker._isDisabledDatepicker( input ) || $.datepicker._lastInput === input ) { // already here
+			return;
+		}
+
+		var inst, beforeShow, beforeShowSettings, isFixed,
+			offset, showAnim, duration;
+
+		inst = $.datepicker._getInst( input );
+		if ( $.datepicker._curInst && $.datepicker._curInst !== inst ) {
+			$.datepicker._curInst.dpDiv.stop( true, true );
+			if ( inst && $.datepicker._datepickerShowing ) {
+				$.datepicker._hideDatepicker( $.datepicker._curInst.input[ 0 ] );
+			}
+		}
+
+		beforeShow = $.datepicker._get( inst, "beforeShow" );
+		beforeShowSettings = beforeShow ? beforeShow.apply( input, [ input, inst ] ) : {};
+		if ( beforeShowSettings === false ) {
+			return;
+		}
+		datepicker_extendRemove( inst.settings, beforeShowSettings );
+
+		inst.lastVal = null;
+		$.datepicker._lastInput = input;
+		$.datepicker._setDateFromField( inst );
+
+		if ( $.datepicker._inDialog ) { // hide cursor
+			input.value = "";
+		}
+		if ( !$.datepicker._pos ) { // position below input
+			$.datepicker._pos = $.datepicker._findPos( input );
+			$.datepicker._pos[ 1 ] += input.offsetHeight; // add the height
+		}
+
+		isFixed = false;
+		$( input ).parents().each( function() {
+			isFixed |= $( this ).css( "position" ) === "fixed";
+			return !isFixed;
+		} );
+
+		offset = { left: $.datepicker._pos[ 0 ], top: $.datepicker._pos[ 1 ] };
+		$.datepicker._pos = null;
+
+		//to avoid flashes on Firefox
+		inst.dpDiv.empty();
+
+		// determine sizing offscreen
+		inst.dpDiv.css( { position: "absolute", display: "block", top: "-1000px" } );
+		$.datepicker._updateDatepicker( inst );
+
+		// fix width for dynamic number of date pickers
+		// and adjust position before showing
+		offset = $.datepicker._checkOffset( inst, offset, isFixed );
+		inst.dpDiv.css( { position: ( $.datepicker._inDialog && $.blockUI ?
+			"static" : ( isFixed ? "fixed" : "absolute" ) ), display: "none",
+			left: offset.left + "px", top: offset.top + "px" } );
+
+		if ( !inst.inline ) {
+			showAnim = $.datepicker._get( inst, "showAnim" );
+			duration = $.datepicker._get( inst, "duration" );
+			inst.dpDiv.css( "z-index", datepicker_getZindex( $( input ) ) + 1 );
+			$.datepicker._datepickerShowing = true;
+
+			if ( $.effects && $.effects.effect[ showAnim ] ) {
+				inst.dpDiv.show( showAnim, $.datepicker._get( inst, "showOptions" ), duration );
+			} else {
+				inst.dpDiv[ showAnim || "show" ]( showAnim ? duration : null );
+			}
+
+			if ( $.datepicker._shouldFocusInput( inst ) ) {
+				inst.input.trigger( "focus" );
+			}
+
+			$.datepicker._curInst = inst;
+		}
+	},
+
+	/* Generate the date picker content. */
+	_updateDatepicker: function( inst ) {
+		this.maxRows = 4; //Reset the max number of rows being displayed (see #7043)
+		datepicker_instActive = inst; // for delegate hover events
+		inst.dpDiv.empty().append( this._generateHTML( inst ) );
+		this._attachHandlers( inst );
+
+		var origyearshtml,
+			numMonths = this._getNumberOfMonths( inst ),
+			cols = numMonths[ 1 ],
+			width = 17,
+			activeCell = inst.dpDiv.find( "." + this._dayOverClass + " a" );
+
+		if ( activeCell.length > 0 ) {
+			datepicker_handleMouseover.apply( activeCell.get( 0 ) );
+		}
+
+		inst.dpDiv.removeClass( "ui-datepicker-multi-2 ui-datepicker-multi-3 ui-datepicker-multi-4" ).width( "" );
+		if ( cols > 1 ) {
+			inst.dpDiv.addClass( "ui-datepicker-multi-" + cols ).css( "width", ( width * cols ) + "em" );
+		}
+		inst.dpDiv[ ( numMonths[ 0 ] !== 1 || numMonths[ 1 ] !== 1 ? "add" : "remove" ) +
+			"Class" ]( "ui-datepicker-multi" );
+		inst.dpDiv[ ( this._get( inst, "isRTL" ) ? "add" : "remove" ) +
+			"Class" ]( "ui-datepicker-rtl" );
+
+		if ( inst === $.datepicker._curInst && $.datepicker._datepickerShowing && $.datepicker._shouldFocusInput( inst ) ) {
+			inst.input.trigger( "focus" );
+		}
+
+		// Deffered render of the years select (to avoid flashes on Firefox)
+		if ( inst.yearshtml ) {
+			origyearshtml = inst.yearshtml;
+			setTimeout( function() {
+
+				//assure that inst.yearshtml didn't change.
+				if ( origyearshtml === inst.yearshtml && inst.yearshtml ) {
+					inst.dpDiv.find( "select.ui-datepicker-year:first" ).replaceWith( inst.yearshtml );
+				}
+				origyearshtml = inst.yearshtml = null;
+			}, 0 );
+		}
+	},
+
+	// #6694 - don't focus the input if it's already focused
+	// this breaks the change event in IE
+	// Support: IE and jQuery <1.9
+	_shouldFocusInput: function( inst ) {
+		return inst.input && inst.input.is( ":visible" ) && !inst.input.is( ":disabled" ) && !inst.input.is( ":focus" );
+	},
+
+	/* Check positioning to remain on screen. */
+	_checkOffset: function( inst, offset, isFixed ) {
+		var dpWidth = inst.dpDiv.outerWidth(),
+			dpHeight = inst.dpDiv.outerHeight(),
+			inputWidth = inst.input ? inst.input.outerWidth() : 0,
+			inputHeight = inst.input ? inst.input.outerHeight() : 0,
+			viewWidth = document.documentElement.clientWidth + ( isFixed ? 0 : $( document ).scrollLeft() ),
+			viewHeight = document.documentElement.clientHeight + ( isFixed ? 0 : $( document ).scrollTop() );
+
+		offset.left -= ( this._get( inst, "isRTL" ) ? ( dpWidth - inputWidth ) : 0 );
+		offset.left -= ( isFixed && offset.left === inst.input.offset().left ) ? $( document ).scrollLeft() : 0;
+		offset.top -= ( isFixed && offset.top === ( inst.input.offset().top + inputHeight ) ) ? $( document ).scrollTop() : 0;
+
+		// Now check if datepicker is showing outside window viewport - move to a better place if so.
+		offset.left -= Math.min( offset.left, ( offset.left + dpWidth > viewWidth && viewWidth > dpWidth ) ?
+			Math.abs( offset.left + dpWidth - viewWidth ) : 0 );
+		offset.top -= Math.min( offset.top, ( offset.top + dpHeight > viewHeight && viewHeight > dpHeight ) ?
+			Math.abs( dpHeight + inputHeight ) : 0 );
+
+		return offset;
+	},
+
+	/* Find an object's position on the screen. */
+	_findPos: function( obj ) {
+		var position,
+			inst = this._getInst( obj ),
+			isRTL = this._get( inst, "isRTL" );
+
+		while ( obj && ( obj.type === "hidden" || obj.nodeType !== 1 || $.expr.filters.hidden( obj ) ) ) {
+			obj = obj[ isRTL ? "previousSibling" : "nextSibling" ];
+		}
+
+		position = $( obj ).offset();
+		return [ position.left, position.top ];
+	},
+
+	/* Hide the date picker from view.
+	 * @param  input  element - the input field attached to the date picker
+	 */
+	_hideDatepicker: function( input ) {
+		var showAnim, duration, postProcess, onClose,
+			inst = this._curInst;
+
+		if ( !inst || ( input && inst !== $.data( input, "datepicker" ) ) ) {
+			return;
+		}
+
+		if ( this._datepickerShowing ) {
+			showAnim = this._get( inst, "showAnim" );
+			duration = this._get( inst, "duration" );
+			postProcess = function() {
+				$.datepicker._tidyDialog( inst );
+			};
+
+			// DEPRECATED: after BC for 1.8.x $.effects[ showAnim ] is not needed
+			if ( $.effects && ( $.effects.effect[ showAnim ] || $.effects[ showAnim ] ) ) {
+				inst.dpDiv.hide( showAnim, $.datepicker._get( inst, "showOptions" ), duration, postProcess );
+			} else {
+				inst.dpDiv[ ( showAnim === "slideDown" ? "slideUp" :
+					( showAnim === "fadeIn" ? "fadeOut" : "hide" ) ) ]( ( showAnim ? duration : null ), postProcess );
+			}
+
+			if ( !showAnim ) {
+				postProcess();
+			}
+			this._datepickerShowing = false;
+
+			onClose = this._get( inst, "onClose" );
+			if ( onClose ) {
+				onClose.apply( ( inst.input ? inst.input[ 0 ] : null ), [ ( inst.input ? inst.input.val() : "" ), inst ] );
+			}
+
+			this._lastInput = null;
+			if ( this._inDialog ) {
+				this._dialogInput.css( { position: "absolute", left: "0", top: "-100px" } );
+				if ( $.blockUI ) {
+					$.unblockUI();
+					$( "body" ).append( this.dpDiv );
+				}
+			}
+			this._inDialog = false;
+		}
+	},
+
+	/* Tidy up after a dialog display. */
+	_tidyDialog: function( inst ) {
+		inst.dpDiv.removeClass( this._dialogClass ).off( ".ui-datepicker-calendar" );
+	},
+
+	/* Close date picker if clicked elsewhere. */
+	_checkExternalClick: function( event ) {
+		if ( !$.datepicker._curInst ) {
+			return;
+		}
+
+		var $target = $( event.target ),
+			inst = $.datepicker._getInst( $target[ 0 ] );
+
+		if ( ( ( $target[ 0 ].id !== $.datepicker._mainDivId &&
+				$target.parents( "#" + $.datepicker._mainDivId ).length === 0 &&
+				!$target.hasClass( $.datepicker.markerClassName ) &&
+				!$target.closest( "." + $.datepicker._triggerClass ).length &&
+				$.datepicker._datepickerShowing && !( $.datepicker._inDialog && $.blockUI ) ) ) ||
+			( $target.hasClass( $.datepicker.markerClassName ) && $.datepicker._curInst !== inst ) ) {
+				$.datepicker._hideDatepicker();
+		}
+	},
+
+	/* Adjust one of the date sub-fields. */
+	_adjustDate: function( id, offset, period ) {
+		var target = $( id ),
+			inst = this._getInst( target[ 0 ] );
+
+		if ( this._isDisabledDatepicker( target[ 0 ] ) ) {
+			return;
+		}
+		this._adjustInstDate( inst, offset +
+			( period === "M" ? this._get( inst, "showCurrentAtPos" ) : 0 ), // undo positioning
+			period );
+		this._updateDatepicker( inst );
+	},
+
+	/* Action for current link. */
+	_gotoToday: function( id ) {
+		var date,
+			target = $( id ),
+			inst = this._getInst( target[ 0 ] );
+
+		if ( this._get( inst, "gotoCurrent" ) && inst.currentDay ) {
+			inst.selectedDay = inst.currentDay;
+			inst.drawMonth = inst.selectedMonth = inst.currentMonth;
+			inst.drawYear = inst.selectedYear = inst.currentYear;
+		} else {
+			date = new Date();
+			inst.selectedDay = date.getDate();
+			inst.drawMonth = inst.selectedMonth = date.getMonth();
+			inst.drawYear = inst.selectedYear = date.getFullYear();
+		}
+		this._notifyChange( inst );
+		this._adjustDate( target );
+	},
+
+	/* Action for selecting a new month/year. */
+	_selectMonthYear: function( id, select, period ) {
+		var target = $( id ),
+			inst = this._getInst( target[ 0 ] );
+
+		inst[ "selected" + ( period === "M" ? "Month" : "Year" ) ] =
+		inst[ "draw" + ( period === "M" ? "Month" : "Year" ) ] =
+			parseInt( select.options[ select.selectedIndex ].value, 10 );
+
+		this._notifyChange( inst );
+		this._adjustDate( target );
+	},
+
+	/* Action for selecting a day. */
+	_selectDay: function( id, month, year, td ) {
+		var inst,
+			target = $( id );
+
+		if ( $( td ).hasClass( this._unselectableClass ) || this._isDisabledDatepicker( target[ 0 ] ) ) {
+			return;
+		}
+
+		inst = this._getInst( target[ 0 ] );
+		inst.selectedDay = inst.currentDay = $( "a", td ).html();
+		inst.selectedMonth = inst.currentMonth = month;
+		inst.selectedYear = inst.currentYear = year;
+		this._selectDate( id, this._formatDate( inst,
+			inst.currentDay, inst.currentMonth, inst.currentYear ) );
+	},
+
+	/* Erase the input field and hide the date picker. */
+	_clearDate: function( id ) {
+		var target = $( id );
+		this._selectDate( target, "" );
+	},
+
+	/* Update the input field with the selected date. */
+	_selectDate: function( id, dateStr ) {
+		var onSelect,
+			target = $( id ),
+			inst = this._getInst( target[ 0 ] );
+
+		dateStr = ( dateStr != null ? dateStr : this._formatDate( inst ) );
+		if ( inst.input ) {
+			inst.input.val( dateStr );
+		}
+		this._updateAlternate( inst );
+
+		onSelect = this._get( inst, "onSelect" );
+		if ( onSelect ) {
+			onSelect.apply( ( inst.input ? inst.input[ 0 ] : null ), [ dateStr, inst ] );  // trigger custom callback
+		} else if ( inst.input ) {
+			inst.input.trigger( "change" ); // fire the change event
+		}
+
+		if ( inst.inline ) {
+			this._updateDatepicker( inst );
+		} else {
+			this._hideDatepicker();
+			this._lastInput = inst.input[ 0 ];
+			if ( typeof( inst.input[ 0 ] ) !== "object" ) {
+				inst.input.trigger( "focus" ); // restore focus
+			}
+			this._lastInput = null;
+		}
+	},
+
+	/* Update any alternate field to synchronise with the main field. */
+	_updateAlternate: function( inst ) {
+		var altFormat, date, dateStr,
+			altField = this._get( inst, "altField" );
+
+		if ( altField ) { // update alternate field too
+			altFormat = this._get( inst, "altFormat" ) || this._get( inst, "dateFormat" );
+			date = this._getDate( inst );
+			dateStr = this.formatDate( altFormat, date, this._getFormatConfig( inst ) );
+			$( altField ).val( dateStr );
+		}
+	},
+
+	/* Set as beforeShowDay function to prevent selection of weekends.
+	 * @param  date  Date - the date to customise
+	 * @return [boolean, string] - is this date selectable?, what is its CSS class?
+	 */
+	noWeekends: function( date ) {
+		var day = date.getDay();
+		return [ ( day > 0 && day < 6 ), "" ];
+	},
+
+	/* Set as calculateWeek to determine the week of the year based on the ISO 8601 definition.
+	 * @param  date  Date - the date to get the week for
+	 * @return  number - the number of the week within the year that contains this date
+	 */
+	iso8601Week: function( date ) {
+		var time,
+			checkDate = new Date( date.getTime() );
+
+		// Find Thursday of this week starting on Monday
+		checkDate.setDate( checkDate.getDate() + 4 - ( checkDate.getDay() || 7 ) );
+
+		time = checkDate.getTime();
+		checkDate.setMonth( 0 ); // Compare with Jan 1
+		checkDate.setDate( 1 );
+		return Math.floor( Math.round( ( time - checkDate ) / 86400000 ) / 7 ) + 1;
+	},
+
+	/* Parse a string value into a date object.
+	 * See formatDate below for the possible formats.
+	 *
+	 * @param  format string - the expected format of the date
+	 * @param  value string - the date in the above format
+	 * @param  settings Object - attributes include:
+	 *					shortYearCutoff  number - the cutoff year for determining the century (optional)
+	 *					dayNamesShort	string[7] - abbreviated names of the days from Sunday (optional)
+	 *					dayNames		string[7] - names of the days from Sunday (optional)
+	 *					monthNamesShort string[12] - abbreviated names of the months (optional)
+	 *					monthNames		string[12] - names of the months (optional)
+	 * @return  Date - the extracted date value or null if value is blank
+	 */
+	parseDate: function( format, value, settings ) {
+		if ( format == null || value == null ) {
+			throw "Invalid arguments";
+		}
+
+		value = ( typeof value === "object" ? value.toString() : value + "" );
+		if ( value === "" ) {
+			return null;
+		}
+
+		var iFormat, dim, extra,
+			iValue = 0,
+			shortYearCutoffTemp = ( settings ? settings.shortYearCutoff : null ) || this._defaults.shortYearCutoff,
+			shortYearCutoff = ( typeof shortYearCutoffTemp !== "string" ? shortYearCutoffTemp :
+				new Date().getFullYear() % 100 + parseInt( shortYearCutoffTemp, 10 ) ),
+			dayNamesShort = ( settings ? settings.dayNamesShort : null ) || this._defaults.dayNamesShort,
+			dayNames = ( settings ? settings.dayNames : null ) || this._defaults.dayNames,
+			monthNamesShort = ( settings ? settings.monthNamesShort : null ) || this._defaults.monthNamesShort,
+			monthNames = ( settings ? settings.monthNames : null ) || this._defaults.monthNames,
+			year = -1,
+			month = -1,
+			day = -1,
+			doy = -1,
+			literal = false,
+			date,
+
+			// Check whether a format character is doubled
+			lookAhead = function( match ) {
+				var matches = ( iFormat + 1 < format.length && format.charAt( iFormat + 1 ) === match );
+				if ( matches ) {
+					iFormat++;
+				}
+				return matches;
+			},
+
+			// Extract a number from the string value
+			getNumber = function( match ) {
+				var isDoubled = lookAhead( match ),
+					size = ( match === "@" ? 14 : ( match === "!" ? 20 :
+					( match === "y" && isDoubled ? 4 : ( match === "o" ? 3 : 2 ) ) ) ),
+					minSize = ( match === "y" ? size : 1 ),
+					digits = new RegExp( "^\\d{" + minSize + "," + size + "}" ),
+					num = value.substring( iValue ).match( digits );
+				if ( !num ) {
+					throw "Missing number at position " + iValue;
+				}
+				iValue += num[ 0 ].length;
+				return parseInt( num[ 0 ], 10 );
+			},
+
+			// Extract a name from the string value and convert to an index
+			getName = function( match, shortNames, longNames ) {
+				var index = -1,
+					names = $.map( lookAhead( match ) ? longNames : shortNames, function( v, k ) {
+						return [ [ k, v ] ];
+					} ).sort( function( a, b ) {
+						return -( a[ 1 ].length - b[ 1 ].length );
+					} );
+
+				$.each( names, function( i, pair ) {
+					var name = pair[ 1 ];
+					if ( value.substr( iValue, name.length ).toLowerCase() === name.toLowerCase() ) {
+						index = pair[ 0 ];
+						iValue += name.length;
+						return false;
+					}
+				} );
+				if ( index !== -1 ) {
+					return index + 1;
+				} else {
+					throw "Unknown name at position " + iValue;
+				}
+			},
+
+			// Confirm that a literal character matches the string value
+			checkLiteral = function() {
+				if ( value.charAt( iValue ) !== format.charAt( iFormat ) ) {
+					throw "Unexpected literal at position " + iValue;
+				}
+				iValue++;
+			};
+
+		for ( iFormat = 0; iFormat < format.length; iFormat++ ) {
+			if ( literal ) {
+				if ( format.charAt( iFormat ) === "'" && !lookAhead( "'" ) ) {
+					literal = false;
+				} else {
+					checkLiteral();
+				}
+			} else {
+				switch ( format.charAt( iFormat ) ) {
+					case "d":
+						day = getNumber( "d" );
+						break;
+					case "D":
+						getName( "D", dayNamesShort, dayNames );
+						break;
+					case "o":
+						doy = getNumber( "o" );
+						break;
+					case "m":
+						month = getNumber( "m" );
+						break;
+					case "M":
+						month = getName( "M", monthNamesShort, monthNames );
+						break;
+					case "y":
+						year = getNumber( "y" );
+						break;
+					case "@":
+						date = new Date( getNumber( "@" ) );
+						year = date.getFullYear();
+						month = date.getMonth() + 1;
+						day = date.getDate();
+						break;
+					case "!":
+						date = new Date( ( getNumber( "!" ) - this._ticksTo1970 ) / 10000 );
+						year = date.getFullYear();
+						month = date.getMonth() + 1;
+						day = date.getDate();
+						break;
+					case "'":
+						if ( lookAhead( "'" ) ) {
+							checkLiteral();
+						} else {
+							literal = true;
+						}
+						break;
+					default:
+						checkLiteral();
+				}
+			}
+		}
+
+		if ( iValue < value.length ) {
+			extra = value.substr( iValue );
+			if ( !/^\s+/.test( extra ) ) {
+				throw "Extra/unparsed characters found in date: " + extra;
+			}
+		}
+
+		if ( year === -1 ) {
+			year = new Date().getFullYear();
+		} else if ( year < 100 ) {
+			year += new Date().getFullYear() - new Date().getFullYear() % 100 +
+				( year <= shortYearCutoff ? 0 : -100 );
+		}
+
+		if ( doy > -1 ) {
+			month = 1;
+			day = doy;
+			do {
+				dim = this._getDaysInMonth( year, month - 1 );
+				if ( day <= dim ) {
+					break;
+				}
+				month++;
+				day -= dim;
+			} while ( true );
+		}
+
+		date = this._daylightSavingAdjust( new Date( year, month - 1, day ) );
+		if ( date.getFullYear() !== year || date.getMonth() + 1 !== month || date.getDate() !== day ) {
+			throw "Invalid date"; // E.g. 31/02/00
+		}
+		return date;
+	},
+
+	/* Standard date formats. */
+	ATOM: "yy-mm-dd", // RFC 3339 (ISO 8601)
+	COOKIE: "D, dd M yy",
+	ISO_8601: "yy-mm-dd",
+	RFC_822: "D, d M y",
+	RFC_850: "DD, dd-M-y",
+	RFC_1036: "D, d M y",
+	RFC_1123: "D, d M yy",
+	RFC_2822: "D, d M yy",
+	RSS: "D, d M y", // RFC 822
+	TICKS: "!",
+	TIMESTAMP: "@",
+	W3C: "yy-mm-dd", // ISO 8601
+
+	_ticksTo1970: ( ( ( 1970 - 1 ) * 365 + Math.floor( 1970 / 4 ) - Math.floor( 1970 / 100 ) +
+		Math.floor( 1970 / 400 ) ) * 24 * 60 * 60 * 10000000 ),
+
+	/* Format a date object into a string value.
+	 * The format can be combinations of the following:
+	 * d  - day of month (no leading zero)
+	 * dd - day of month (two digit)
+	 * o  - day of year (no leading zeros)
+	 * oo - day of year (three digit)
+	 * D  - day name short
+	 * DD - day name long
+	 * m  - month of year (no leading zero)
+	 * mm - month of year (two digit)
+	 * M  - month name short
+	 * MM - month name long
+	 * y  - year (two digit)
+	 * yy - year (four digit)
+	 * @ - Unix timestamp (ms since 01/01/1970)
+	 * ! - Windows ticks (100ns since 01/01/0001)
+	 * "..." - literal text
+	 * '' - single quote
+	 *
+	 * @param  format string - the desired format of the date
+	 * @param  date Date - the date value to format
+	 * @param  settings Object - attributes include:
+	 *					dayNamesShort	string[7] - abbreviated names of the days from Sunday (optional)
+	 *					dayNames		string[7] - names of the days from Sunday (optional)
+	 *					monthNamesShort string[12] - abbreviated names of the months (optional)
+	 *					monthNames		string[12] - names of the months (optional)
+	 * @return  string - the date in the above format
+	 */
+	formatDate: function( format, date, settings ) {
+		if ( !date ) {
+			return "";
+		}
+
+		var iFormat,
+			dayNamesShort = ( settings ? settings.dayNamesShort : null ) || this._defaults.dayNamesShort,
+			dayNames = ( settings ? settings.dayNames : null ) || this._defaults.dayNames,
+			monthNamesShort = ( settings ? settings.monthNamesShort : null ) || this._defaults.monthNamesShort,
+			monthNames = ( settings ? settings.monthNames : null ) || this._defaults.monthNames,
+
+			// Check whether a format character is doubled
+			lookAhead = function( match ) {
+				var matches = ( iFormat + 1 < format.length && format.charAt( iFormat + 1 ) === match );
+				if ( matches ) {
+					iFormat++;
+				}
+				return matches;
+			},
+
+			// Format a number, with leading zero if necessary
+			formatNumber = function( match, value, len ) {
+				var num = "" + value;
+				if ( lookAhead( match ) ) {
+					while ( num.length < len ) {
+						num = "0" + num;
+					}
+				}
+				return num;
+			},
+
+			// Format a name, short or long as requested
+			formatName = function( match, value, shortNames, longNames ) {
+				return ( lookAhead( match ) ? longNames[ value ] : shortNames[ value ] );
+			},
+			output = "",
+			literal = false;
+
+		if ( date ) {
+			for ( iFormat = 0; iFormat < format.length; iFormat++ ) {
+				if ( literal ) {
+					if ( format.charAt( iFormat ) === "'" && !lookAhead( "'" ) ) {
+						literal = false;
+					} else {
+						output += format.charAt( iFormat );
+					}
+				} else {
+					switch ( format.charAt( iFormat ) ) {
+						case "d":
+							output += formatNumber( "d", date.getDate(), 2 );
+							break;
+						case "D":
+							output += formatName( "D", date.getDay(), dayNamesShort, dayNames );
+							break;
+						case "o":
+							output += formatNumber( "o",
+								Math.round( ( new Date( date.getFullYear(), date.getMonth(), date.getDate() ).getTime() - new Date( date.getFullYear(), 0, 0 ).getTime() ) / 86400000 ), 3 );
+							break;
+						case "m":
+							output += formatNumber( "m", date.getMonth() + 1, 2 );
+							break;
+						case "M":
+							output += formatName( "M", date.getMonth(), monthNamesShort, monthNames );
+							break;
+						case "y":
+							output += ( lookAhead( "y" ) ? date.getFullYear() :
+								( date.getFullYear() % 100 < 10 ? "0" : "" ) + date.getFullYear() % 100 );
+							break;
+						case "@":
+							output += date.getTime();
+							break;
+						case "!":
+							output += date.getTime() * 10000 + this._ticksTo1970;
+							break;
+						case "'":
+							if ( lookAhead( "'" ) ) {
+								output += "'";
+							} else {
+								literal = true;
+							}
+							break;
+						default:
+							output += format.charAt( iFormat );
+					}
+				}
+			}
+		}
+		return output;
+	},
+
+	/* Extract all possible characters from the date format. */
+	_possibleChars: function( format ) {
+		var iFormat,
+			chars = "",
+			literal = false,
+
+			// Check whether a format character is doubled
+			lookAhead = function( match ) {
+				var matches = ( iFormat + 1 < format.length && format.charAt( iFormat + 1 ) === match );
+				if ( matches ) {
+					iFormat++;
+				}
+				return matches;
+			};
+
+		for ( iFormat = 0; iFormat < format.length; iFormat++ ) {
+			if ( literal ) {
+				if ( format.charAt( iFormat ) === "'" && !lookAhead( "'" ) ) {
+					literal = false;
+				} else {
+					chars += format.charAt( iFormat );
+				}
+			} else {
+				switch ( format.charAt( iFormat ) ) {
+					case "d": case "m": case "y": case "@":
+						chars += "0123456789";
+						break;
+					case "D": case "M":
+						return null; // Accept anything
+					case "'":
+						if ( lookAhead( "'" ) ) {
+							chars += "'";
+						} else {
+							literal = true;
+						}
+						break;
+					default:
+						chars += format.charAt( iFormat );
+				}
+			}
+		}
+		return chars;
+	},
+
+	/* Get a setting value, defaulting if necessary. */
+	_get: function( inst, name ) {
+		return inst.settings[ name ] !== undefined ?
+			inst.settings[ name ] : this._defaults[ name ];
+	},
+
+	/* Parse existing date and initialise date picker. */
+	_setDateFromField: function( inst, noDefault ) {
+		if ( inst.input.val() === inst.lastVal ) {
+			return;
+		}
+
+		var dateFormat = this._get( inst, "dateFormat" ),
+			dates = inst.lastVal = inst.input ? inst.input.val() : null,
+			defaultDate = this._getDefaultDate( inst ),
+			date = defaultDate,
+			settings = this._getFormatConfig( inst );
+
+		try {
+			date = this.parseDate( dateFormat, dates, settings ) || defaultDate;
+		} catch ( event ) {
+			dates = ( noDefault ? "" : dates );
+		}
+		inst.selectedDay = date.getDate();
+		inst.drawMonth = inst.selectedMonth = date.getMonth();
+		inst.drawYear = inst.selectedYear = date.getFullYear();
+		inst.currentDay = ( dates ? date.getDate() : 0 );
+		inst.currentMonth = ( dates ? date.getMonth() : 0 );
+		inst.currentYear = ( dates ? date.getFullYear() : 0 );
+		this._adjustInstDate( inst );
+	},
+
+	/* Retrieve the default date shown on opening. */
+	_getDefaultDate: function( inst ) {
+		return this._restrictMinMax( inst,
+			this._determineDate( inst, this._get( inst, "defaultDate" ), new Date() ) );
+	},
+
+	/* A date may be specified as an exact value or a relative one. */
+	_determineDate: function( inst, date, defaultDate ) {
+		var offsetNumeric = function( offset ) {
+				var date = new Date();
+				date.setDate( date.getDate() + offset );
+				return date;
+			},
+			offsetString = function( offset ) {
+				try {
+					return $.datepicker.parseDate( $.datepicker._get( inst, "dateFormat" ),
+						offset, $.datepicker._getFormatConfig( inst ) );
+				}
+				catch ( e ) {
+
+					// Ignore
+				}
+
+				var date = ( offset.toLowerCase().match( /^c/ ) ?
+					$.datepicker._getDate( inst ) : null ) || new Date(),
+					year = date.getFullYear(),
+					month = date.getMonth(),
+					day = date.getDate(),
+					pattern = /([+\-]?[0-9]+)\s*(d|D|w|W|m|M|y|Y)?/g,
+					matches = pattern.exec( offset );
+
+				while ( matches ) {
+					switch ( matches[ 2 ] || "d" ) {
+						case "d" : case "D" :
+							day += parseInt( matches[ 1 ], 10 ); break;
+						case "w" : case "W" :
+							day += parseInt( matches[ 1 ], 10 ) * 7; break;
+						case "m" : case "M" :
+							month += parseInt( matches[ 1 ], 10 );
+							day = Math.min( day, $.datepicker._getDaysInMonth( year, month ) );
+							break;
+						case "y": case "Y" :
+							year += parseInt( matches[ 1 ], 10 );
+							day = Math.min( day, $.datepicker._getDaysInMonth( year, month ) );
+							break;
+					}
+					matches = pattern.exec( offset );
+				}
+				return new Date( year, month, day );
+			},
+			newDate = ( date == null || date === "" ? defaultDate : ( typeof date === "string" ? offsetString( date ) :
+				( typeof date === "number" ? ( isNaN( date ) ? defaultDate : offsetNumeric( date ) ) : new Date( date.getTime() ) ) ) );
+
+		newDate = ( newDate && newDate.toString() === "Invalid Date" ? defaultDate : newDate );
+		if ( newDate ) {
+			newDate.setHours( 0 );
+			newDate.setMinutes( 0 );
+			newDate.setSeconds( 0 );
+			newDate.setMilliseconds( 0 );
+		}
+		return this._daylightSavingAdjust( newDate );
+	},
+
+	/* Handle switch to/from daylight saving.
+	 * Hours may be non-zero on daylight saving cut-over:
+	 * > 12 when midnight changeover, but then cannot generate
+	 * midnight datetime, so jump to 1AM, otherwise reset.
+	 * @param  date  (Date) the date to check
+	 * @return  (Date) the corrected date
+	 */
+	_daylightSavingAdjust: function( date ) {
+		if ( !date ) {
+			return null;
+		}
+		date.setHours( date.getHours() > 12 ? date.getHours() + 2 : 0 );
+		return date;
+	},
+
+	/* Set the date(s) directly. */
+	_setDate: function( inst, date, noChange ) {
+		var clear = !date,
+			origMonth = inst.selectedMonth,
+			origYear = inst.selectedYear,
+			newDate = this._restrictMinMax( inst, this._determineDate( inst, date, new Date() ) );
+
+		inst.selectedDay = inst.currentDay = newDate.getDate();
+		inst.drawMonth = inst.selectedMonth = inst.currentMonth = newDate.getMonth();
+		inst.drawYear = inst.selectedYear = inst.currentYear = newDate.getFullYear();
+		if ( ( origMonth !== inst.selectedMonth || origYear !== inst.selectedYear ) && !noChange ) {
+			this._notifyChange( inst );
+		}
+		this._adjustInstDate( inst );
+		if ( inst.input ) {
+			inst.input.val( clear ? "" : this._formatDate( inst ) );
+		}
+	},
+
+	/* Retrieve the date(s) directly. */
+	_getDate: function( inst ) {
+		var startDate = ( !inst.currentYear || ( inst.input && inst.input.val() === "" ) ? null :
+			this._daylightSavingAdjust( new Date(
+			inst.currentYear, inst.currentMonth, inst.currentDay ) ) );
+			return startDate;
+	},
+
+	/* Attach the onxxx handlers.  These are declared statically so
+	 * they work with static code transformers like Caja.
+	 */
+	_attachHandlers: function( inst ) {
+		var stepMonths = this._get( inst, "stepMonths" ),
+			id = "#" + inst.id.replace( /\\\\/g, "\\" );
+		inst.dpDiv.find( "[data-handler]" ).map( function() {
+			var handler = {
+				prev: function() {
+					$.datepicker._adjustDate( id, -stepMonths, "M" );
+				},
+				next: function() {
+					$.datepicker._adjustDate( id, +stepMonths, "M" );
+				},
+				hide: function() {
+					$.datepicker._hideDatepicker();
+				},
+				today: function() {
+					$.datepicker._gotoToday( id );
+				},
+				selectDay: function() {
+					$.datepicker._selectDay( id, +this.getAttribute( "data-month" ), +this.getAttribute( "data-year" ), this );
+					return false;
+				},
+				selectMonth: function() {
+					$.datepicker._selectMonthYear( id, this, "M" );
+					return false;
+				},
+				selectYear: function() {
+					$.datepicker._selectMonthYear( id, this, "Y" );
+					return false;
+				}
+			};
+			$( this ).on( this.getAttribute( "data-event" ), handler[ this.getAttribute( "data-handler" ) ] );
+		} );
+	},
+
+	/* Generate the HTML for the current state of the date picker. */
+	_generateHTML: function( inst ) {
+		var maxDraw, prevText, prev, nextText, next, currentText, gotoDate,
+			controls, buttonPanel, firstDay, showWeek, dayNames, dayNamesMin,
+			monthNames, monthNamesShort, beforeShowDay, showOtherMonths,
+			selectOtherMonths, defaultDate, html, dow, row, group, col, selectedDate,
+			cornerClass, calender, thead, day, daysInMonth, leadDays, curRows, numRows,
+			printDate, dRow, tbody, daySettings, otherMonth, unselectable,
+			tempDate = new Date(),
+			today = this._daylightSavingAdjust(
+				new Date( tempDate.getFullYear(), tempDate.getMonth(), tempDate.getDate() ) ), // clear time
+			isRTL = this._get( inst, "isRTL" ),
+			showButtonPanel = this._get( inst, "showButtonPanel" ),
+			hideIfNoPrevNext = this._get( inst, "hideIfNoPrevNext" ),
+			navigationAsDateFormat = this._get( inst, "navigationAsDateFormat" ),
+			numMonths = this._getNumberOfMonths( inst ),
+			showCurrentAtPos = this._get( inst, "showCurrentAtPos" ),
+			stepMonths = this._get( inst, "stepMonths" ),
+			isMultiMonth = ( numMonths[ 0 ] !== 1 || numMonths[ 1 ] !== 1 ),
+			currentDate = this._daylightSavingAdjust( ( !inst.currentDay ? new Date( 9999, 9, 9 ) :
+				new Date( inst.currentYear, inst.currentMonth, inst.currentDay ) ) ),
+			minDate = this._getMinMaxDate( inst, "min" ),
+			maxDate = this._getMinMaxDate( inst, "max" ),
+			drawMonth = inst.drawMonth - showCurrentAtPos,
+			drawYear = inst.drawYear;
+
+		if ( drawMonth < 0 ) {
+			drawMonth += 12;
+			drawYear--;
+		}
+		if ( maxDate ) {
+			maxDraw = this._daylightSavingAdjust( new Date( maxDate.getFullYear(),
+				maxDate.getMonth() - ( numMonths[ 0 ] * numMonths[ 1 ] ) + 1, maxDate.getDate() ) );
+			maxDraw = ( minDate && maxDraw < minDate ? minDate : maxDraw );
+			while ( this._daylightSavingAdjust( new Date( drawYear, drawMonth, 1 ) ) > maxDraw ) {
+				drawMonth--;
+				if ( drawMonth < 0 ) {
+					drawMonth = 11;
+					drawYear--;
+				}
+			}
+		}
+		inst.drawMonth = drawMonth;
+		inst.drawYear = drawYear;
+
+		prevText = this._get( inst, "prevText" );
+		prevText = ( !navigationAsDateFormat ? prevText : this.formatDate( prevText,
+			this._daylightSavingAdjust( new Date( drawYear, drawMonth - stepMonths, 1 ) ),
+			this._getFormatConfig( inst ) ) );
+
+		prev = ( this._canAdjustMonth( inst, -1, drawYear, drawMonth ) ?
+			"<a class='ui-datepicker-prev ui-corner-all' data-handler='prev' data-event='click'" +
+			" title='" + prevText + "'><span class='ui-icon ui-icon-circle-triangle-" + ( isRTL ? "e" : "w" ) + "'>" + prevText + "</span></a>" :
+			( hideIfNoPrevNext ? "" : "<a class='ui-datepicker-prev ui-corner-all ui-state-disabled' title='" + prevText + "'><span class='ui-icon ui-icon-circle-triangle-" + ( isRTL ? "e" : "w" ) + "'>" + prevText + "</span></a>" ) );
+
+		nextText = this._get( inst, "nextText" );
+		nextText = ( !navigationAsDateFormat ? nextText : this.formatDate( nextText,
+			this._daylightSavingAdjust( new Date( drawYear, drawMonth + stepMonths, 1 ) ),
+			this._getFormatConfig( inst ) ) );
+
+		next = ( this._canAdjustMonth( inst, +1, drawYear, drawMonth ) ?
+			"<a class='ui-datepicker-next ui-corner-all' data-handler='next' data-event='click'" +
+			" title='" + nextText + "'><span class='ui-icon ui-icon-circle-triangle-" + ( isRTL ? "w" : "e" ) + "'>" + nextText + "</span></a>" :
+			( hideIfNoPrevNext ? "" : "<a class='ui-datepicker-next ui-corner-all ui-state-disabled' title='" + nextText + "'><span class='ui-icon ui-icon-circle-triangle-" + ( isRTL ? "w" : "e" ) + "'>" + nextText + "</span></a>" ) );
+
+		currentText = this._get( inst, "currentText" );
+		gotoDate = ( this._get( inst, "gotoCurrent" ) && inst.currentDay ? currentDate : today );
+		currentText = ( !navigationAsDateFormat ? currentText :
+			this.formatDate( currentText, gotoDate, this._getFormatConfig( inst ) ) );
+
+		controls = ( !inst.inline ? "<button type='button' class='ui-datepicker-close ui-state-default ui-priority-primary ui-corner-all' data-handler='hide' data-event='click'>" +
+			this._get( inst, "closeText" ) + "</button>" : "" );
+
+		buttonPanel = ( showButtonPanel ) ? "<div class='ui-datepicker-buttonpane ui-widget-content'>" + ( isRTL ? controls : "" ) +
+			( this._isInRange( inst, gotoDate ) ? "<button type='button' class='ui-datepicker-current ui-state-default ui-priority-secondary ui-corner-all' data-handler='today' data-event='click'" +
+			">" + currentText + "</button>" : "" ) + ( isRTL ? "" : controls ) + "</div>" : "";
+
+		firstDay = parseInt( this._get( inst, "firstDay" ), 10 );
+		firstDay = ( isNaN( firstDay ) ? 0 : firstDay );
+
+		showWeek = this._get( inst, "showWeek" );
+		dayNames = this._get( inst, "dayNames" );
+		dayNamesMin = this._get( inst, "dayNamesMin" );
+		monthNames = this._get( inst, "monthNames" );
+		monthNamesShort = this._get( inst, "monthNamesShort" );
+		beforeShowDay = this._get( inst, "beforeShowDay" );
+		showOtherMonths = this._get( inst, "showOtherMonths" );
+		selectOtherMonths = this._get( inst, "selectOtherMonths" );
+		defaultDate = this._getDefaultDate( inst );
+		html = "";
+
+		for ( row = 0; row < numMonths[ 0 ]; row++ ) {
+			group = "";
+			this.maxRows = 4;
+			for ( col = 0; col < numMonths[ 1 ]; col++ ) {
+				selectedDate = this._daylightSavingAdjust( new Date( drawYear, drawMonth, inst.selectedDay ) );
+				cornerClass = " ui-corner-all";
+				calender = "";
+				if ( isMultiMonth ) {
+					calender += "<div class='ui-datepicker-group";
+					if ( numMonths[ 1 ] > 1 ) {
+						switch ( col ) {
+							case 0: calender += " ui-datepicker-group-first";
+								cornerClass = " ui-corner-" + ( isRTL ? "right" : "left" ); break;
+							case numMonths[ 1 ] - 1: calender += " ui-datepicker-group-last";
+								cornerClass = " ui-corner-" + ( isRTL ? "left" : "right" ); break;
+							default: calender += " ui-datepicker-group-middle"; cornerClass = ""; break;
+						}
+					}
+					calender += "'>";
+				}
+				calender += "<div class='ui-datepicker-header ui-widget-header ui-helper-clearfix" + cornerClass + "'>" +
+					( /all|left/.test( cornerClass ) && row === 0 ? ( isRTL ? next : prev ) : "" ) +
+					( /all|right/.test( cornerClass ) && row === 0 ? ( isRTL ? prev : next ) : "" ) +
+					this._generateMonthYearHeader( inst, drawMonth, drawYear, minDate, maxDate,
+					row > 0 || col > 0, monthNames, monthNamesShort ) + // draw month headers
+					"</div><table class='ui-datepicker-calendar'><thead>" +
+					"<tr>";
+				thead = ( showWeek ? "<th class='ui-datepicker-week-col'>" + this._get( inst, "weekHeader" ) + "</th>" : "" );
+				for ( dow = 0; dow < 7; dow++ ) { // days of the week
+					day = ( dow + firstDay ) % 7;
+					thead += "<th scope='col'" + ( ( dow + firstDay + 6 ) % 7 >= 5 ? " class='ui-datepicker-week-end'" : "" ) + ">" +
+						"<span title='" + dayNames[ day ] + "'>" + dayNamesMin[ day ] + "</span></th>";
+				}
+				calender += thead + "</tr></thead><tbody>";
+				daysInMonth = this._getDaysInMonth( drawYear, drawMonth );
+				if ( drawYear === inst.selectedYear && drawMonth === inst.selectedMonth ) {
+					inst.selectedDay = Math.min( inst.selectedDay, daysInMonth );
+				}
+				leadDays = ( this._getFirstDayOfMonth( drawYear, drawMonth ) - firstDay + 7 ) % 7;
+				curRows = Math.ceil( ( leadDays + daysInMonth ) / 7 ); // calculate the number of rows to generate
+				numRows = ( isMultiMonth ? this.maxRows > curRows ? this.maxRows : curRows : curRows ); //If multiple months, use the higher number of rows (see #7043)
+				this.maxRows = numRows;
+				printDate = this._daylightSavingAdjust( new Date( drawYear, drawMonth, 1 - leadDays ) );
+				for ( dRow = 0; dRow < numRows; dRow++ ) { // create date picker rows
+					calender += "<tr>";
+					tbody = ( !showWeek ? "" : "<td class='ui-datepicker-week-col'>" +
+						this._get( inst, "calculateWeek" )( printDate ) + "</td>" );
+					for ( dow = 0; dow < 7; dow++ ) { // create date picker days
+						daySettings = ( beforeShowDay ?
+							beforeShowDay.apply( ( inst.input ? inst.input[ 0 ] : null ), [ printDate ] ) : [ true, "" ] );
+						otherMonth = ( printDate.getMonth() !== drawMonth );
+						unselectable = ( otherMonth && !selectOtherMonths ) || !daySettings[ 0 ] ||
+							( minDate && printDate < minDate ) || ( maxDate && printDate > maxDate );
+						tbody += "<td class='" +
+							( ( dow + firstDay + 6 ) % 7 >= 5 ? " ui-datepicker-week-end" : "" ) + // highlight weekends
+							( otherMonth ? " ui-datepicker-other-month" : "" ) + // highlight days from other months
+							( ( printDate.getTime() === selectedDate.getTime() && drawMonth === inst.selectedMonth && inst._keyEvent ) || // user pressed key
+							( defaultDate.getTime() === printDate.getTime() && defaultDate.getTime() === selectedDate.getTime() ) ?
+
+							// or defaultDate is current printedDate and defaultDate is selectedDate
+							" " + this._dayOverClass : "" ) + // highlight selected day
+							( unselectable ? " " + this._unselectableClass + " ui-state-disabled" : "" ) +  // highlight unselectable days
+							( otherMonth && !showOtherMonths ? "" : " " + daySettings[ 1 ] + // highlight custom dates
+							( printDate.getTime() === currentDate.getTime() ? " " + this._currentClass : "" ) + // highlight selected day
+							( printDate.getTime() === today.getTime() ? " ui-datepicker-today" : "" ) ) + "'" + // highlight today (if different)
+							( ( !otherMonth || showOtherMonths ) && daySettings[ 2 ] ? " title='" + daySettings[ 2 ].replace( /'/g, "&#39;" ) + "'" : "" ) + // cell title
+							( unselectable ? "" : " data-handler='selectDay' data-event='click' data-month='" + printDate.getMonth() + "' data-year='" + printDate.getFullYear() + "'" ) + ">" + // actions
+							( otherMonth && !showOtherMonths ? "&#xa0;" : // display for other months
+							( unselectable ? "<span class='ui-state-default'>" + printDate.getDate() + "</span>" : "<a class='ui-state-default" +
+							( printDate.getTime() === today.getTime() ? " ui-state-highlight" : "" ) +
+							( printDate.getTime() === currentDate.getTime() ? " ui-state-active" : "" ) + // highlight selected day
+							( otherMonth ? " ui-priority-secondary" : "" ) + // distinguish dates from other months
+							"' href='#'>" + printDate.getDate() + "</a>" ) ) + "</td>"; // display selectable date
+						printDate.setDate( printDate.getDate() + 1 );
+						printDate = this._daylightSavingAdjust( printDate );
+					}
+					calender += tbody + "</tr>";
+				}
+				drawMonth++;
+				if ( drawMonth > 11 ) {
+					drawMonth = 0;
+					drawYear++;
+				}
+				calender += "</tbody></table>" + ( isMultiMonth ? "</div>" +
+							( ( numMonths[ 0 ] > 0 && col === numMonths[ 1 ] - 1 ) ? "<div class='ui-datepicker-row-break'></div>" : "" ) : "" );
+				group += calender;
+			}
+			html += group;
+		}
+		html += buttonPanel;
+		inst._keyEvent = false;
+		return html;
+	},
+
+	/* Generate the month and year header. */
+	_generateMonthYearHeader: function( inst, drawMonth, drawYear, minDate, maxDate,
+			secondary, monthNames, monthNamesShort ) {
+
+		var inMinYear, inMaxYear, month, years, thisYear, determineYear, year, endYear,
+			changeMonth = this._get( inst, "changeMonth" ),
+			changeYear = this._get( inst, "changeYear" ),
+			showMonthAfterYear = this._get( inst, "showMonthAfterYear" ),
+			html = "<div class='ui-datepicker-title'>",
+			monthHtml = "";
+
+		// Month selection
+		if ( secondary || !changeMonth ) {
+			monthHtml += "<span class='ui-datepicker-month'>" + monthNames[ drawMonth ] + "</span>";
+		} else {
+			inMinYear = ( minDate && minDate.getFullYear() === drawYear );
+			inMaxYear = ( maxDate && maxDate.getFullYear() === drawYear );
+			monthHtml += "<select class='ui-datepicker-month' data-handler='selectMonth' data-event='change'>";
+			for ( month = 0; month < 12; month++ ) {
+				if ( ( !inMinYear || month >= minDate.getMonth() ) && ( !inMaxYear || month <= maxDate.getMonth() ) ) {
+					monthHtml += "<option value='" + month + "'" +
+						( month === drawMonth ? " selected='selected'" : "" ) +
+						">" + monthNamesShort[ month ] + "</option>";
+				}
+			}
+			monthHtml += "</select>";
+		}
+
+		if ( !showMonthAfterYear ) {
+			html += monthHtml + ( secondary || !( changeMonth && changeYear ) ? "&#xa0;" : "" );
+		}
+
+		// Year selection
+		if ( !inst.yearshtml ) {
+			inst.yearshtml = "";
+			if ( secondary || !changeYear ) {
+				html += "<span class='ui-datepicker-year'>" + drawYear + "</span>";
+			} else {
+
+				// determine range of years to display
+				years = this._get( inst, "yearRange" ).split( ":" );
+				thisYear = new Date().getFullYear();
+				determineYear = function( value ) {
+					var year = ( value.match( /c[+\-].*/ ) ? drawYear + parseInt( value.substring( 1 ), 10 ) :
+						( value.match( /[+\-].*/ ) ? thisYear + parseInt( value, 10 ) :
+						parseInt( value, 10 ) ) );
+					return ( isNaN( year ) ? thisYear : year );
+				};
+				year = determineYear( years[ 0 ] );
+				endYear = Math.max( year, determineYear( years[ 1 ] || "" ) );
+				year = ( minDate ? Math.max( year, minDate.getFullYear() ) : year );
+				endYear = ( maxDate ? Math.min( endYear, maxDate.getFullYear() ) : endYear );
+				inst.yearshtml += "<select class='ui-datepicker-year' data-handler='selectYear' data-event='change'>";
+				for ( ; year <= endYear; year++ ) {
+					inst.yearshtml += "<option value='" + year + "'" +
+						( year === drawYear ? " selected='selected'" : "" ) +
+						">" + year + "</option>";
+				}
+				inst.yearshtml += "</select>";
+
+				html += inst.yearshtml;
+				inst.yearshtml = null;
+			}
+		}
+
+		html += this._get( inst, "yearSuffix" );
+		if ( showMonthAfterYear ) {
+			html += ( secondary || !( changeMonth && changeYear ) ? "&#xa0;" : "" ) + monthHtml;
+		}
+		html += "</div>"; // Close datepicker_header
+		return html;
+	},
+
+	/* Adjust one of the date sub-fields. */
+	_adjustInstDate: function( inst, offset, period ) {
+		var year = inst.selectedYear + ( period === "Y" ? offset : 0 ),
+			month = inst.selectedMonth + ( period === "M" ? offset : 0 ),
+			day = Math.min( inst.selectedDay, this._getDaysInMonth( year, month ) ) + ( period === "D" ? offset : 0 ),
+			date = this._restrictMinMax( inst, this._daylightSavingAdjust( new Date( year, month, day ) ) );
+
+		inst.selectedDay = date.getDate();
+		inst.drawMonth = inst.selectedMonth = date.getMonth();
+		inst.drawYear = inst.selectedYear = date.getFullYear();
+		if ( period === "M" || period === "Y" ) {
+			this._notifyChange( inst );
+		}
+	},
+
+	/* Ensure a date is within any min/max bounds. */
+	_restrictMinMax: function( inst, date ) {
+		var minDate = this._getMinMaxDate( inst, "min" ),
+			maxDate = this._getMinMaxDate( inst, "max" ),
+			newDate = ( minDate && date < minDate ? minDate : date );
+		return ( maxDate && newDate > maxDate ? maxDate : newDate );
+	},
+
+	/* Notify change of month/year. */
+	_notifyChange: function( inst ) {
+		var onChange = this._get( inst, "onChangeMonthYear" );
+		if ( onChange ) {
+			onChange.apply( ( inst.input ? inst.input[ 0 ] : null ),
+				[ inst.selectedYear, inst.selectedMonth + 1, inst ] );
+		}
+	},
+
+	/* Determine the number of months to show. */
+	_getNumberOfMonths: function( inst ) {
+		var numMonths = this._get( inst, "numberOfMonths" );
+		return ( numMonths == null ? [ 1, 1 ] : ( typeof numMonths === "number" ? [ 1, numMonths ] : numMonths ) );
+	},
+
+	/* Determine the current maximum date - ensure no time components are set. */
+	_getMinMaxDate: function( inst, minMax ) {
+		return this._determineDate( inst, this._get( inst, minMax + "Date" ), null );
+	},
+
+	/* Find the number of days in a given month. */
+	_getDaysInMonth: function( year, month ) {
+		return 32 - this._daylightSavingAdjust( new Date( year, month, 32 ) ).getDate();
+	},
+
+	/* Find the day of the week of the first of a month. */
+	_getFirstDayOfMonth: function( year, month ) {
+		return new Date( year, month, 1 ).getDay();
+	},
+
+	/* Determines if we should allow a "next/prev" month display change. */
+	_canAdjustMonth: function( inst, offset, curYear, curMonth ) {
+		var numMonths = this._getNumberOfMonths( inst ),
+			date = this._daylightSavingAdjust( new Date( curYear,
+			curMonth + ( offset < 0 ? offset : numMonths[ 0 ] * numMonths[ 1 ] ), 1 ) );
+
+		if ( offset < 0 ) {
+			date.setDate( this._getDaysInMonth( date.getFullYear(), date.getMonth() ) );
+		}
+		return this._isInRange( inst, date );
+	},
+
+	/* Is the given date in the accepted range? */
+	_isInRange: function( inst, date ) {
+		var yearSplit, currentYear,
+			minDate = this._getMinMaxDate( inst, "min" ),
+			maxDate = this._getMinMaxDate( inst, "max" ),
+			minYear = null,
+			maxYear = null,
+			years = this._get( inst, "yearRange" );
+			if ( years ) {
+				yearSplit = years.split( ":" );
+				currentYear = new Date().getFullYear();
+				minYear = parseInt( yearSplit[ 0 ], 10 );
+				maxYear = parseInt( yearSplit[ 1 ], 10 );
+				if ( yearSplit[ 0 ].match( /[+\-].*/ ) ) {
+					minYear += currentYear;
+				}
+				if ( yearSplit[ 1 ].match( /[+\-].*/ ) ) {
+					maxYear += currentYear;
+				}
+			}
+
+		return ( ( !minDate || date.getTime() >= minDate.getTime() ) &&
+			( !maxDate || date.getTime() <= maxDate.getTime() ) &&
+			( !minYear || date.getFullYear() >= minYear ) &&
+			( !maxYear || date.getFullYear() <= maxYear ) );
+	},
+
+	/* Provide the configuration settings for formatting/parsing. */
+	_getFormatConfig: function( inst ) {
+		var shortYearCutoff = this._get( inst, "shortYearCutoff" );
+		shortYearCutoff = ( typeof shortYearCutoff !== "string" ? shortYearCutoff :
+			new Date().getFullYear() % 100 + parseInt( shortYearCutoff, 10 ) );
+		return { shortYearCutoff: shortYearCutoff,
+			dayNamesShort: this._get( inst, "dayNamesShort" ), dayNames: this._get( inst, "dayNames" ),
+			monthNamesShort: this._get( inst, "monthNamesShort" ), monthNames: this._get( inst, "monthNames" ) };
+	},
+
+	/* Format the given date for display. */
+	_formatDate: function( inst, day, month, year ) {
+		if ( !day ) {
+			inst.currentDay = inst.selectedDay;
+			inst.currentMonth = inst.selectedMonth;
+			inst.currentYear = inst.selectedYear;
+		}
+		var date = ( day ? ( typeof day === "object" ? day :
+			this._daylightSavingAdjust( new Date( year, month, day ) ) ) :
+			this._daylightSavingAdjust( new Date( inst.currentYear, inst.currentMonth, inst.currentDay ) ) );
+		return this.formatDate( this._get( inst, "dateFormat" ), date, this._getFormatConfig( inst ) );
+	}
+} );
+
+/*
+ * Bind hover events for datepicker elements.
+ * Done via delegate so the binding only occurs once in the lifetime of the parent div.
+ * Global datepicker_instActive, set by _updateDatepicker allows the handlers to find their way back to the active picker.
+ */
+function datepicker_bindHover( dpDiv ) {
+	var selector = "button, .ui-datepicker-prev, .ui-datepicker-next, .ui-datepicker-calendar td a";
+	return dpDiv.on( "mouseout", selector, function() {
+			$( this ).removeClass( "ui-state-hover" );
+			if ( this.className.indexOf( "ui-datepicker-prev" ) !== -1 ) {
+				$( this ).removeClass( "ui-datepicker-prev-hover" );
+			}
+			if ( this.className.indexOf( "ui-datepicker-next" ) !== -1 ) {
+				$( this ).removeClass( "ui-datepicker-next-hover" );
+			}
+		} )
+		.on( "mouseover", selector, datepicker_handleMouseover );
+}
+
+function datepicker_handleMouseover() {
+	if ( !$.datepicker._isDisabledDatepicker( datepicker_instActive.inline ? datepicker_instActive.dpDiv.parent()[ 0 ] : datepicker_instActive.input[ 0 ] ) ) {
+		$( this ).parents( ".ui-datepicker-calendar" ).find( "a" ).removeClass( "ui-state-hover" );
+		$( this ).addClass( "ui-state-hover" );
+		if ( this.className.indexOf( "ui-datepicker-prev" ) !== -1 ) {
+			$( this ).addClass( "ui-datepicker-prev-hover" );
+		}
+		if ( this.className.indexOf( "ui-datepicker-next" ) !== -1 ) {
+			$( this ).addClass( "ui-datepicker-next-hover" );
+		}
+	}
+}
+
+/* jQuery extend now ignores nulls! */
+function datepicker_extendRemove( target, props ) {
+	$.extend( target, props );
+	for ( var name in props ) {
+		if ( props[ name ] == null ) {
+			target[ name ] = props[ name ];
+		}
+	}
+	return target;
+}
+
+/* Invoke the datepicker functionality.
+   @param  options  string - a command, optionally followed by additional parameters or
+					Object - settings for attaching new datepicker functionality
+   @return  jQuery object */
+$.fn.datepicker = function( options ) {
+
+	/* Verify an empty collection wasn't passed - Fixes #6976 */
+	if ( !this.length ) {
+		return this;
+	}
+
+	/* Initialise the date picker. */
+	if ( !$.datepicker.initialized ) {
+		$( document ).on( "mousedown", $.datepicker._checkExternalClick );
+		$.datepicker.initialized = true;
+	}
+
+	/* Append datepicker main container to body if not exist. */
+	if ( $( "#" + $.datepicker._mainDivId ).length === 0 ) {
+		$( "body" ).append( $.datepicker.dpDiv );
+	}
+
+	var otherArgs = Array.prototype.slice.call( arguments, 1 );
+	if ( typeof options === "string" && ( options === "isDisabled" || options === "getDate" || options === "widget" ) ) {
+		return $.datepicker[ "_" + options + "Datepicker" ].
+			apply( $.datepicker, [ this[ 0 ] ].concat( otherArgs ) );
+	}
+	if ( options === "option" && arguments.length === 2 && typeof arguments[ 1 ] === "string" ) {
+		return $.datepicker[ "_" + options + "Datepicker" ].
+			apply( $.datepicker, [ this[ 0 ] ].concat( otherArgs ) );
+	}
+	return this.each( function() {
+		typeof options === "string" ?
+			$.datepicker[ "_" + options + "Datepicker" ].
+				apply( $.datepicker, [ this ].concat( otherArgs ) ) :
+			$.datepicker._attachDatepicker( this, options );
+	} );
+};
+
+$.datepicker = new Datepicker(); // singleton instance
+$.datepicker.initialized = false;
+$.datepicker.uuid = new Date().getTime();
+$.datepicker.version = "1.12.0";
+
+var widgetsDatepicker = $.datepicker;
+
+
+
+
+// This file is deprecated
+var ie = $.ui.ie = !!/msie [\w.]+/.exec( navigator.userAgent.toLowerCase() );
+
+/*!
+ * jQuery UI Mouse 1.12.0
+ * http://jqueryui.com
+ *
+ * Copyright jQuery Foundation and other contributors
+ * Released under the MIT license.
+ * http://jquery.org/license
+ */
+
+//>>label: Mouse
+//>>group: Widgets
+//>>description: Abstracts mouse-based interactions to assist in creating certain widgets.
+//>>docs: http://api.jqueryui.com/mouse/
+
+
+
+var mouseHandled = false;
+$( document ).on( "mouseup", function() {
+	mouseHandled = false;
+} );
+
+var widgetsMouse = $.widget( "ui.mouse", {
+	version: "1.12.0",
+	options: {
+		cancel: "input, textarea, button, select, option",
+		distance: 1,
+		delay: 0
+	},
+	_mouseInit: function() {
+		var that = this;
+
+		this.element
+			.on( "mousedown." + this.widgetName, function( event ) {
+				return that._mouseDown( event );
+			} )
+			.on( "click." + this.widgetName, function( event ) {
+				if ( true === $.data( event.target, that.widgetName + ".preventClickEvent" ) ) {
+					$.removeData( event.target, that.widgetName + ".preventClickEvent" );
+					event.stopImmediatePropagation();
+					return false;
+				}
+			} );
+
+		this.started = false;
+	},
+
+	// TODO: make sure destroying one instance of mouse doesn't mess with
+	// other instances of mouse
+	_mouseDestroy: function() {
+		this.element.off( "." + this.widgetName );
+		if ( this._mouseMoveDelegate ) {
+			this.document
+				.off( "mousemove." + this.widgetName, this._mouseMoveDelegate )
+				.off( "mouseup." + this.widgetName, this._mouseUpDelegate );
+		}
+	},
+
+	_mouseDown: function( event ) {
+
+		// don't let more than one widget handle mouseStart
+		if ( mouseHandled ) {
+			return;
+		}
+
+		this._mouseMoved = false;
+
+		// We may have missed mouseup (out of window)
+		( this._mouseStarted && this._mouseUp( event ) );
+
+		this._mouseDownEvent = event;
+
+		var that = this,
+			btnIsLeft = ( event.which === 1 ),
+
+			// event.target.nodeName works around a bug in IE 8 with
+			// disabled inputs (#7620)
+			elIsCancel = ( typeof this.options.cancel === "string" && event.target.nodeName ?
+				$( event.target ).closest( this.options.cancel ).length : false );
+		if ( !btnIsLeft || elIsCancel || !this._mouseCapture( event ) ) {
+			return true;
+		}
+
+		this.mouseDelayMet = !this.options.delay;
+		if ( !this.mouseDelayMet ) {
+			this._mouseDelayTimer = setTimeout( function() {
+				that.mouseDelayMet = true;
+			}, this.options.delay );
+		}
+
+		if ( this._mouseDistanceMet( event ) && this._mouseDelayMet( event ) ) {
+			this._mouseStarted = ( this._mouseStart( event ) !== false );
+			if ( !this._mouseStarted ) {
+				event.preventDefault();
+				return true;
+			}
+		}
+
+		// Click event may never have fired (Gecko & Opera)
+		if ( true === $.data( event.target, this.widgetName + ".preventClickEvent" ) ) {
+			$.removeData( event.target, this.widgetName + ".preventClickEvent" );
+		}
+
+		// These delegates are required to keep context
+		this._mouseMoveDelegate = function( event ) {
+			return that._mouseMove( event );
+		};
+		this._mouseUpDelegate = function( event ) {
+			return that._mouseUp( event );
+		};
+
+		this.document
+			.on( "mousemove." + this.widgetName, this._mouseMoveDelegate )
+			.on( "mouseup." + this.widgetName, this._mouseUpDelegate );
+
+		event.preventDefault();
+
+		mouseHandled = true;
+		return true;
+	},
+
+	_mouseMove: function( event ) {
+
+		// Only check for mouseups outside the document if you've moved inside the document
+		// at least once. This prevents the firing of mouseup in the case of IE<9, which will
+		// fire a mousemove event if content is placed under the cursor. See #7778
+		// Support: IE <9
+		if ( this._mouseMoved ) {
+
+			// IE mouseup check - mouseup happened when mouse was out of window
+			if ( $.ui.ie && ( !document.documentMode || document.documentMode < 9 ) &&
+					!event.button ) {
+				return this._mouseUp( event );
+
+			// Iframe mouseup check - mouseup occurred in another document
+			} else if ( !event.which ) {
+
+				// Support: Safari <=8 - 9
+				// Safari sets which to 0 if you press any of the following keys
+				// during a drag (#14461)
+				if ( event.originalEvent.altKey || event.originalEvent.ctrlKey ||
+						event.originalEvent.metaKey || event.originalEvent.shiftKey ) {
+					this.ignoreMissingWhich = true;
+				} else if ( !this.ignoreMissingWhich ) {
+					return this._mouseUp( event );
+				}
+			}
+		}
+
+		if ( event.which || event.button ) {
+			this._mouseMoved = true;
+		}
+
+		if ( this._mouseStarted ) {
+			this._mouseDrag( event );
+			return event.preventDefault();
+		}
+
+		if ( this._mouseDistanceMet( event ) && this._mouseDelayMet( event ) ) {
+			this._mouseStarted =
+				( this._mouseStart( this._mouseDownEvent, event ) !== false );
+			( this._mouseStarted ? this._mouseDrag( event ) : this._mouseUp( event ) );
+		}
+
+		return !this._mouseStarted;
+	},
+
+	_mouseUp: function( event ) {
+		this.document
+			.off( "mousemove." + this.widgetName, this._mouseMoveDelegate )
+			.off( "mouseup." + this.widgetName, this._mouseUpDelegate );
+
+		if ( this._mouseStarted ) {
+			this._mouseStarted = false;
+
+			if ( event.target === this._mouseDownEvent.target ) {
+				$.data( event.target, this.widgetName + ".preventClickEvent", true );
+			}
+
+			this._mouseStop( event );
+		}
+
+		if ( this._mouseDelayTimer ) {
+			clearTimeout( this._mouseDelayTimer );
+			delete this._mouseDelayTimer;
+		}
+
+		this.ignoreMissingWhich = false;
+		mouseHandled = false;
+		event.preventDefault();
+	},
+
+	_mouseDistanceMet: function( event ) {
+		return ( Math.max(
+				Math.abs( this._mouseDownEvent.pageX - event.pageX ),
+				Math.abs( this._mouseDownEvent.pageY - event.pageY )
+			) >= this.options.distance
+		);
+	},
+
+	_mouseDelayMet: function( /* event */ ) {
+		return this.mouseDelayMet;
+	},
+
+	// These are placeholder methods, to be overriden by extending plugin
+	_mouseStart: function( /* event */ ) {},
+	_mouseDrag: function( /* event */ ) {},
+	_mouseStop: function( /* event */ ) {},
+	_mouseCapture: function( /* event */ ) { return true; }
+} );
+
+
+
+
+// $.ui.plugin is deprecated. Use $.widget() extensions instead.
+var plugin = $.ui.plugin = {
+	add: function( module, option, set ) {
+		var i,
+			proto = $.ui[ module ].prototype;
+		for ( i in set ) {
+			proto.plugins[ i ] = proto.plugins[ i ] || [];
+			proto.plugins[ i ].push( [ option, set[ i ] ] );
+		}
+	},
+	call: function( instance, name, args, allowDisconnected ) {
+		var i,
+			set = instance.plugins[ name ];
+
+		if ( !set ) {
+			return;
+		}
+
+		if ( !allowDisconnected && ( !instance.element[ 0 ].parentNode ||
+				instance.element[ 0 ].parentNode.nodeType === 11 ) ) {
+			return;
+		}
+
+		for ( i = 0; i < set.length; i++ ) {
+			if ( instance.options[ set[ i ][ 0 ] ] ) {
+				set[ i ][ 1 ].apply( instance.element, args );
+			}
+		}
+	}
+};
+
+
+
+var safeBlur = $.ui.safeBlur = function( element ) {
+
+	// Support: IE9 - 10 only
+	// If the <body> is blurred, IE will switch windows, see #9420
+	if ( element && element.nodeName.toLowerCase() !== "body" ) {
+		$( element ).trigger( "blur" );
+	}
+};
+
+
+/*!
+ * jQuery UI Draggable 1.12.0
+ * http://jqueryui.com
+ *
+ * Copyright jQuery Foundation and other contributors
+ * Released under the MIT license.
+ * http://jquery.org/license
+ */
+
+//>>label: Draggable
+//>>group: Interactions
+//>>description: Enables dragging functionality for any element.
+//>>docs: http://api.jqueryui.com/draggable/
+//>>demos: http://jqueryui.com/draggable/
+//>>css.structure: ../../themes/base/draggable.css
+
+
+
+$.widget( "ui.draggable", $.ui.mouse, {
+	version: "1.12.0",
+	widgetEventPrefix: "drag",
+	options: {
+		addClasses: true,
+		appendTo: "parent",
+		axis: false,
+		connectToSortable: false,
+		containment: false,
+		cursor: "auto",
+		cursorAt: false,
+		grid: false,
+		handle: false,
+		helper: "original",
+		iframeFix: false,
+		opacity: false,
+		refreshPositions: false,
+		revert: false,
+		revertDuration: 500,
+		scope: "default",
+		scroll: true,
+		scrollSensitivity: 20,
+		scrollSpeed: 20,
+		snap: false,
+		snapMode: "both",
+		snapTolerance: 20,
+		stack: false,
+		zIndex: false,
+
+		// Callbacks
+		drag: null,
+		start: null,
+		stop: null
+	},
+	_create: function() {
+
+		if ( this.options.helper === "original" ) {
+			this._setPositionRelative();
+		}
+		if ( this.options.addClasses ) {
+			this._addClass( "ui-draggable" );
+		}
+		this._setHandleClassName();
+
+		this._mouseInit();
+	},
+
+	_setOption: function( key, value ) {
+		this._super( key, value );
+		if ( key === "handle" ) {
+			this._removeHandleClassName();
+			this._setHandleClassName();
+		}
+	},
+
+	_destroy: function() {
+		if ( ( this.helper || this.element ).is( ".ui-draggable-dragging" ) ) {
+			this.destroyOnClear = true;
+			return;
+		}
+		this._removeHandleClassName();
+		this._mouseDestroy();
+	},
+
+	_mouseCapture: function( event ) {
+		var o = this.options;
+
+		this._blurActiveElement( event );
+
+		// Among others, prevent a drag on a resizable-handle
+		if ( this.helper || o.disabled ||
+				$( event.target ).closest( ".ui-resizable-handle" ).length > 0 ) {
+			return false;
+		}
+
+		//Quit if we're not on a valid handle
+		this.handle = this._getHandle( event );
+		if ( !this.handle ) {
+			return false;
+		}
+
+		this._blockFrames( o.iframeFix === true ? "iframe" : o.iframeFix );
+
+		return true;
+
+	},
+
+	_blockFrames: function( selector ) {
+		this.iframeBlocks = this.document.find( selector ).map( function() {
+			var iframe = $( this );
+
+			return $( "<div>" )
+				.css( "position", "absolute" )
+				.appendTo( iframe.parent() )
+				.outerWidth( iframe.outerWidth() )
+				.outerHeight( iframe.outerHeight() )
+				.offset( iframe.offset() )[ 0 ];
+		} );
+	},
+
+	_unblockFrames: function() {
+		if ( this.iframeBlocks ) {
+			this.iframeBlocks.remove();
+			delete this.iframeBlocks;
+		}
+	},
+
+	_blurActiveElement: function( event ) {
+		var activeElement = $.ui.safeActiveElement( this.document[ 0 ] ),
+			target = $( event.target );
+
+		// Only blur if the event occurred on an element that is:
+		// 1) within the draggable handle
+		// 2) but not within the currently focused element
+		// See #10527, #12472
+		if ( this._getHandle( event ) && target.closest( activeElement ).length ) {
+			return;
+		}
+
+		// Blur any element that currently has focus, see #4261
+		$.ui.safeBlur( activeElement );
+	},
+
+	_mouseStart: function( event ) {
+
+		var o = this.options;
+
+		//Create and append the visible helper
+		this.helper = this._createHelper( event );
+
+		this._addClass( this.helper, "ui-draggable-dragging" );
+
+		//Cache the helper size
+		this._cacheHelperProportions();
+
+		//If ddmanager is used for droppables, set the global draggable
+		if ( $.ui.ddmanager ) {
+			$.ui.ddmanager.current = this;
+		}
+
+		/*
+		 * - Position generation -
+		 * This block generates everything position related - it's the core of draggables.
+		 */
+
+		//Cache the margins of the original element
+		this._cacheMargins();
+
+		//Store the helper's css position
+		this.cssPosition = this.helper.css( "position" );
+		this.scrollParent = this.helper.scrollParent( true );
+		this.offsetParent = this.helper.offsetParent();
+		this.hasFixedAncestor = this.helper.parents().filter( function() {
+				return $( this ).css( "position" ) === "fixed";
+			} ).length > 0;
+
+		//The element's absolute position on the page minus margins
+		this.positionAbs = this.element.offset();
+		this._refreshOffsets( event );
+
+		//Generate the original position
+		this.originalPosition = this.position = this._generatePosition( event, false );
+		this.originalPageX = event.pageX;
+		this.originalPageY = event.pageY;
+
+		//Adjust the mouse offset relative to the helper if "cursorAt" is supplied
+		( o.cursorAt && this._adjustOffsetFromHelper( o.cursorAt ) );
+
+		//Set a containment if given in the options
+		this._setContainment();
+
+		//Trigger event + callbacks
+		if ( this._trigger( "start", event ) === false ) {
+			this._clear();
+			return false;
+		}
+
+		//Recache the helper size
+		this._cacheHelperProportions();
+
+		//Prepare the droppable offsets
+		if ( $.ui.ddmanager && !o.dropBehaviour ) {
+			$.ui.ddmanager.prepareOffsets( this, event );
+		}
+
+		// Execute the drag once - this causes the helper not to be visible before getting its
+		// correct position
+		this._mouseDrag( event, true );
+
+		// If the ddmanager is used for droppables, inform the manager that dragging has started
+		// (see #5003)
+		if ( $.ui.ddmanager ) {
+			$.ui.ddmanager.dragStart( this, event );
+		}
+
+		return true;
+	},
+
+	_refreshOffsets: function( event ) {
+		this.offset = {
+			top: this.positionAbs.top - this.margins.top,
+			left: this.positionAbs.left - this.margins.left,
+			scroll: false,
+			parent: this._getParentOffset(),
+			relative: this._getRelativeOffset()
+		};
+
+		this.offset.click = {
+			left: event.pageX - this.offset.left,
+			top: event.pageY - this.offset.top
+		};
+	},
+
+	_mouseDrag: function( event, noPropagation ) {
+
+		// reset any necessary cached properties (see #5009)
+		if ( this.hasFixedAncestor ) {
+			this.offset.parent = this._getParentOffset();
+		}
+
+		//Compute the helpers position
+		this.position = this._generatePosition( event, true );
+		this.positionAbs = this._convertPositionTo( "absolute" );
+
+		//Call plugins and callbacks and use the resulting position if something is returned
+		if ( !noPropagation ) {
+			var ui = this._uiHash();
+			if ( this._trigger( "drag", event, ui ) === false ) {
+				this._mouseUp( new $.Event( "mouseup", event ) );
+				return false;
+			}
+			this.position = ui.position;
+		}
+
+		this.helper[ 0 ].style.left = this.position.left + "px";
+		this.helper[ 0 ].style.top = this.position.top + "px";
+
+		if ( $.ui.ddmanager ) {
+			$.ui.ddmanager.drag( this, event );
+		}
+
+		return false;
+	},
+
+	_mouseStop: function( event ) {
+
+		//If we are using droppables, inform the manager about the drop
+		var that = this,
+			dropped = false;
+		if ( $.ui.ddmanager && !this.options.dropBehaviour ) {
+			dropped = $.ui.ddmanager.drop( this, event );
+		}
+
+		//if a drop comes from outside (a sortable)
+		if ( this.dropped ) {
+			dropped = this.dropped;
+			this.dropped = false;
+		}
+
+		if ( ( this.options.revert === "invalid" && !dropped ) ||
+				( this.options.revert === "valid" && dropped ) ||
+				this.options.revert === true || ( $.isFunction( this.options.revert ) &&
+				this.options.revert.call( this.element, dropped ) )
+		) {
+			$( this.helper ).animate(
+				this.originalPosition,
+				parseInt( this.options.revertDuration, 10 ),
+				function() {
+					if ( that._trigger( "stop", event ) !== false ) {
+						that._clear();
+					}
+				}
+			);
+		} else {
+			if ( this._trigger( "stop", event ) !== false ) {
+				this._clear();
+			}
+		}
+
+		return false;
+	},
+
+	_mouseUp: function( event ) {
+		this._unblockFrames();
+
+		// If the ddmanager is used for droppables, inform the manager that dragging has stopped
+		// (see #5003)
+		if ( $.ui.ddmanager ) {
+			$.ui.ddmanager.dragStop( this, event );
+		}
+
+		// Only need to focus if the event occurred on the draggable itself, see #10527
+		if ( this.handleElement.is( event.target ) ) {
+
+			// The interaction is over; whether or not the click resulted in a drag,
+			// focus the element
+			this.element.trigger( "focus" );
+		}
+
+		return $.ui.mouse.prototype._mouseUp.call( this, event );
+	},
+
+	cancel: function() {
+
+		if ( this.helper.is( ".ui-draggable-dragging" ) ) {
+			this._mouseUp( new $.Event( "mouseup", { target: this.element[ 0 ] } ) );
+		} else {
+			this._clear();
+		}
+
+		return this;
+
+	},
+
+	_getHandle: function( event ) {
+		return this.options.handle ?
+			!!$( event.target ).closest( this.element.find( this.options.handle ) ).length :
+			true;
+	},
+
+	_setHandleClassName: function() {
+		this.handleElement = this.options.handle ?
+			this.element.find( this.options.handle ) : this.element;
+		this._addClass( this.handleElement, "ui-draggable-handle" );
+	},
+
+	_removeHandleClassName: function() {
+		this._removeClass( this.handleElement, "ui-draggable-handle" );
+	},
+
+	_createHelper: function( event ) {
+
+		var o = this.options,
+			helperIsFunction = $.isFunction( o.helper ),
+			helper = helperIsFunction ?
+				$( o.helper.apply( this.element[ 0 ], [ event ] ) ) :
+				( o.helper === "clone" ?
+					this.element.clone().removeAttr( "id" ) :
+					this.element );
+
+		if ( !helper.parents( "body" ).length ) {
+			helper.appendTo( ( o.appendTo === "parent" ?
+				this.element[ 0 ].parentNode :
+				o.appendTo ) );
+		}
+
+		// Http://bugs.jqueryui.com/ticket/9446
+		// a helper function can return the original element
+		// which wouldn't have been set to relative in _create
+		if ( helperIsFunction && helper[ 0 ] === this.element[ 0 ] ) {
+			this._setPositionRelative();
+		}
+
+		if ( helper[ 0 ] !== this.element[ 0 ] &&
+				!( /(fixed|absolute)/ ).test( helper.css( "position" ) ) ) {
+			helper.css( "position", "absolute" );
+		}
+
+		return helper;
+
+	},
+
+	_setPositionRelative: function() {
+		if ( !( /^(?:r|a|f)/ ).test( this.element.css( "position" ) ) ) {
+			this.element[ 0 ].style.position = "relative";
+		}
+	},
+
+	_adjustOffsetFromHelper: function( obj ) {
+		if ( typeof obj === "string" ) {
+			obj = obj.split( " " );
+		}
+		if ( $.isArray( obj ) ) {
+			obj = { left: +obj[ 0 ], top: +obj[ 1 ] || 0 };
+		}
+		if ( "left" in obj ) {
+			this.offset.click.left = obj.left + this.margins.left;
+		}
+		if ( "right" in obj ) {
+			this.offset.click.left = this.helperProportions.width - obj.right + this.margins.left;
+		}
+		if ( "top" in obj ) {
+			this.offset.click.top = obj.top + this.margins.top;
+		}
+		if ( "bottom" in obj ) {
+			this.offset.click.top = this.helperProportions.height - obj.bottom + this.margins.top;
+		}
+	},
+
+	_isRootNode: function( element ) {
+		return ( /(html|body)/i ).test( element.tagName ) || element === this.document[ 0 ];
+	},
+
+	_getParentOffset: function() {
+
+		//Get the offsetParent and cache its position
+		var po = this.offsetParent.offset(),
+			document = this.document[ 0 ];
+
+		// This is a special case where we need to modify a offset calculated on start, since the
+		// following happened:
+		// 1. The position of the helper is absolute, so it's position is calculated based on the
+		// next positioned parent
+		// 2. The actual offset parent is a child of the scroll parent, and the scroll parent isn't
+		// the document, which means that the scroll is included in the initial calculation of the
+		// offset of the parent, and never recalculated upon drag
+		if ( this.cssPosition === "absolute" && this.scrollParent[ 0 ] !== document &&
+				$.contains( this.scrollParent[ 0 ], this.offsetParent[ 0 ] ) ) {
+			po.left += this.scrollParent.scrollLeft();
+			po.top += this.scrollParent.scrollTop();
+		}
+
+		if ( this._isRootNode( this.offsetParent[ 0 ] ) ) {
+			po = { top: 0, left: 0 };
+		}
+
+		return {
+			top: po.top + ( parseInt( this.offsetParent.css( "borderTopWidth" ), 10 ) || 0 ),
+			left: po.left + ( parseInt( this.offsetParent.css( "borderLeftWidth" ), 10 ) || 0 )
+		};
+
+	},
+
+	_getRelativeOffset: function() {
+		if ( this.cssPosition !== "relative" ) {
+			return { top: 0, left: 0 };
+		}
+
+		var p = this.element.position(),
+			scrollIsRootNode = this._isRootNode( this.scrollParent[ 0 ] );
+
+		return {
+			top: p.top - ( parseInt( this.helper.css( "top" ), 10 ) || 0 ) +
+				( !scrollIsRootNode ? this.scrollParent.scrollTop() : 0 ),
+			left: p.left - ( parseInt( this.helper.css( "left" ), 10 ) || 0 ) +
+				( !scrollIsRootNode ? this.scrollParent.scrollLeft() : 0 )
+		};
+
+	},
+
+	_cacheMargins: function() {
+		this.margins = {
+			left: ( parseInt( this.element.css( "marginLeft" ), 10 ) || 0 ),
+			top: ( parseInt( this.element.css( "marginTop" ), 10 ) || 0 ),
+			right: ( parseInt( this.element.css( "marginRight" ), 10 ) || 0 ),
+			bottom: ( parseInt( this.element.css( "marginBottom" ), 10 ) || 0 )
+		};
+	},
+
+	_cacheHelperProportions: function() {
+		this.helperProportions = {
+			width: this.helper.outerWidth(),
+			height: this.helper.outerHeight()
+		};
+	},
+
+	_setContainment: function() {
+
+		var isUserScrollable, c, ce,
+			o = this.options,
+			document = this.document[ 0 ];
+
+		this.relativeContainer = null;
+
+		if ( !o.containment ) {
+			this.containment = null;
+			return;
+		}
+
+		if ( o.containment === "window" ) {
+			this.containment = [
+				$( window ).scrollLeft() - this.offset.relative.left - this.offset.parent.left,
+				$( window ).scrollTop() - this.offset.relative.top - this.offset.parent.top,
+				$( window ).scrollLeft() + $( window ).width() -
+					this.helperProportions.width - this.margins.left,
+				$( window ).scrollTop() +
+					( $( window ).height() || document.body.parentNode.scrollHeight ) -
+					this.helperProportions.height - this.margins.top
+			];
+			return;
+		}
+
+		if ( o.containment === "document" ) {
+			this.containment = [
+				0,
+				0,
+				$( document ).width() - this.helperProportions.width - this.margins.left,
+				( $( document ).height() || document.body.parentNode.scrollHeight ) -
+					this.helperProportions.height - this.margins.top
+			];
+			return;
+		}
+
+		if ( o.containment.constructor === Array ) {
+			this.containment = o.containment;
+			return;
+		}
+
+		if ( o.containment === "parent" ) {
+			o.containment = this.helper[ 0 ].parentNode;
+		}
+
+		c = $( o.containment );
+		ce = c[ 0 ];
+
+		if ( !ce ) {
+			return;
+		}
+
+		isUserScrollable = /(scroll|auto)/.test( c.css( "overflow" ) );
+
+		this.containment = [
+			( parseInt( c.css( "borderLeftWidth" ), 10 ) || 0 ) +
+				( parseInt( c.css( "paddingLeft" ), 10 ) || 0 ),
+			( parseInt( c.css( "borderTopWidth" ), 10 ) || 0 ) +
+				( parseInt( c.css( "paddingTop" ), 10 ) || 0 ),
+			( isUserScrollable ? Math.max( ce.scrollWidth, ce.offsetWidth ) : ce.offsetWidth ) -
+				( parseInt( c.css( "borderRightWidth" ), 10 ) || 0 ) -
+				( parseInt( c.css( "paddingRight" ), 10 ) || 0 ) -
+				this.helperProportions.width -
+				this.margins.left -
+				this.margins.right,
+			( isUserScrollable ? Math.max( ce.scrollHeight, ce.offsetHeight ) : ce.offsetHeight ) -
+				( parseInt( c.css( "borderBottomWidth" ), 10 ) || 0 ) -
+				( parseInt( c.css( "paddingBottom" ), 10 ) || 0 ) -
+				this.helperProportions.height -
+				this.margins.top -
+				this.margins.bottom
+		];
+		this.relativeContainer = c;
+	},
+
+	_convertPositionTo: function( d, pos ) {
+
+		if ( !pos ) {
+			pos = this.position;
+		}
+
+		var mod = d === "absolute" ? 1 : -1,
+			scrollIsRootNode = this._isRootNode( this.scrollParent[ 0 ] );
+
+		return {
+			top: (
+
+				// The absolute mouse position
+				pos.top	+
+
+				// Only for relative positioned nodes: Relative offset from element to offset parent
+				this.offset.relative.top * mod +
+
+				// The offsetParent's offset without borders (offset + border)
+				this.offset.parent.top * mod -
+				( ( this.cssPosition === "fixed" ?
+					-this.offset.scroll.top :
+					( scrollIsRootNode ? 0 : this.offset.scroll.top ) ) * mod )
+			),
+			left: (
+
+				// The absolute mouse position
+				pos.left +
+
+				// Only for relative positioned nodes: Relative offset from element to offset parent
+				this.offset.relative.left * mod +
+
+				// The offsetParent's offset without borders (offset + border)
+				this.offset.parent.left * mod	-
+				( ( this.cssPosition === "fixed" ?
+					-this.offset.scroll.left :
+					( scrollIsRootNode ? 0 : this.offset.scroll.left ) ) * mod )
+			)
+		};
+
+	},
+
+	_generatePosition: function( event, constrainPosition ) {
+
+		var containment, co, top, left,
+			o = this.options,
+			scrollIsRootNode = this._isRootNode( this.scrollParent[ 0 ] ),
+			pageX = event.pageX,
+			pageY = event.pageY;
+
+		// Cache the scroll
+		if ( !scrollIsRootNode || !this.offset.scroll ) {
+			this.offset.scroll = {
+				top: this.scrollParent.scrollTop(),
+				left: this.scrollParent.scrollLeft()
+			};
+		}
+
+		/*
+		 * - Position constraining -
+		 * Constrain the position to a mix of grid, containment.
+		 */
+
+		// If we are not dragging yet, we won't check for options
+		if ( constrainPosition ) {
+			if ( this.containment ) {
+				if ( this.relativeContainer ) {
+					co = this.relativeContainer.offset();
+					containment = [
+						this.containment[ 0 ] + co.left,
+						this.containment[ 1 ] + co.top,
+						this.containment[ 2 ] + co.left,
+						this.containment[ 3 ] + co.top
+					];
+				} else {
+					containment = this.containment;
+				}
+
+				if ( event.pageX - this.offset.click.left < containment[ 0 ] ) {
+					pageX = containment[ 0 ] + this.offset.click.left;
+				}
+				if ( event.pageY - this.offset.click.top < containment[ 1 ] ) {
+					pageY = containment[ 1 ] + this.offset.click.top;
+				}
+				if ( event.pageX - this.offset.click.left > containment[ 2 ] ) {
+					pageX = containment[ 2 ] + this.offset.click.left;
+				}
+				if ( event.pageY - this.offset.click.top > containment[ 3 ] ) {
+					pageY = containment[ 3 ] + this.offset.click.top;
+				}
+			}
+
+			if ( o.grid ) {
+
+				//Check for grid elements set to 0 to prevent divide by 0 error causing invalid
+				// argument errors in IE (see ticket #6950)
+				top = o.grid[ 1 ] ? this.originalPageY + Math.round( ( pageY -
+					this.originalPageY ) / o.grid[ 1 ] ) * o.grid[ 1 ] : this.originalPageY;
+				pageY = containment ? ( ( top - this.offset.click.top >= containment[ 1 ] ||
+					top - this.offset.click.top > containment[ 3 ] ) ?
+						top :
+						( ( top - this.offset.click.top >= containment[ 1 ] ) ?
+							top - o.grid[ 1 ] : top + o.grid[ 1 ] ) ) : top;
+
+				left = o.grid[ 0 ] ? this.originalPageX +
+					Math.round( ( pageX - this.originalPageX ) / o.grid[ 0 ] ) * o.grid[ 0 ] :
+					this.originalPageX;
+				pageX = containment ? ( ( left - this.offset.click.left >= containment[ 0 ] ||
+					left - this.offset.click.left > containment[ 2 ] ) ?
+						left :
+						( ( left - this.offset.click.left >= containment[ 0 ] ) ?
+							left - o.grid[ 0 ] : left + o.grid[ 0 ] ) ) : left;
+			}
+
+			if ( o.axis === "y" ) {
+				pageX = this.originalPageX;
+			}
+
+			if ( o.axis === "x" ) {
+				pageY = this.originalPageY;
+			}
+		}
+
+		return {
+			top: (
+
+				// The absolute mouse position
+				pageY -
+
+				// Click offset (relative to the element)
+				this.offset.click.top -
+
+				// Only for relative positioned nodes: Relative offset from element to offset parent
+				this.offset.relative.top -
+
+				// The offsetParent's offset without borders (offset + border)
+				this.offset.parent.top +
+				( this.cssPosition === "fixed" ?
+					-this.offset.scroll.top :
+					( scrollIsRootNode ? 0 : this.offset.scroll.top ) )
+			),
+			left: (
+
+				// The absolute mouse position
+				pageX -
+
+				// Click offset (relative to the element)
+				this.offset.click.left -
+
+				// Only for relative positioned nodes: Relative offset from element to offset parent
+				this.offset.relative.left -
+
+				// The offsetParent's offset without borders (offset + border)
+				this.offset.parent.left +
+				( this.cssPosition === "fixed" ?
+					-this.offset.scroll.left :
+					( scrollIsRootNode ? 0 : this.offset.scroll.left ) )
+			)
+		};
+
+	},
+
+	_clear: function() {
+		this._removeClass( this.helper, "ui-draggable-dragging" );
+		if ( this.helper[ 0 ] !== this.element[ 0 ] && !this.cancelHelperRemoval ) {
+			this.helper.remove();
+		}
+		this.helper = null;
+		this.cancelHelperRemoval = false;
+		if ( this.destroyOnClear ) {
+			this.destroy();
+		}
+	},
+
+	// From now on bulk stuff - mainly helpers
+
+	_trigger: function( type, event, ui ) {
+		ui = ui || this._uiHash();
+		$.ui.plugin.call( this, type, [ event, ui, this ], true );
+
+		// Absolute position and offset (see #6884 ) have to be recalculated after plugins
+		if ( /^(drag|start|stop)/.test( type ) ) {
+			this.positionAbs = this._convertPositionTo( "absolute" );
+			ui.offset = this.positionAbs;
+		}
+		return $.Widget.prototype._trigger.call( this, type, event, ui );
+	},
+
+	plugins: {},
+
+	_uiHash: function() {
+		return {
+			helper: this.helper,
+			position: this.position,
+			originalPosition: this.originalPosition,
+			offset: this.positionAbs
+		};
+	}
+
+} );
+
+$.ui.plugin.add( "draggable", "connectToSortable", {
+	start: function( event, ui, draggable ) {
+		var uiSortable = $.extend( {}, ui, {
+			item: draggable.element
+		} );
+
+		draggable.sortables = [];
+		$( draggable.options.connectToSortable ).each( function() {
+			var sortable = $( this ).sortable( "instance" );
+
+			if ( sortable && !sortable.options.disabled ) {
+				draggable.sortables.push( sortable );
+
+				// RefreshPositions is called at drag start to refresh the containerCache
+				// which is used in drag. This ensures it's initialized and synchronized
+				// with any changes that might have happened on the page since initialization.
+				sortable.refreshPositions();
+				sortable._trigger( "activate", event, uiSortable );
+			}
+		} );
+	},
+	stop: function( event, ui, draggable ) {
+		var uiSortable = $.extend( {}, ui, {
+			item: draggable.element
+		} );
+
+		draggable.cancelHelperRemoval = false;
+
+		$.each( draggable.sortables, function() {
+			var sortable = this;
+
+			if ( sortable.isOver ) {
+				sortable.isOver = 0;
+
+				// Allow this sortable to handle removing the helper
+				draggable.cancelHelperRemoval = true;
+				sortable.cancelHelperRemoval = false;
+
+				// Use _storedCSS To restore properties in the sortable,
+				// as this also handles revert (#9675) since the draggable
+				// may have modified them in unexpected ways (#8809)
+				sortable._storedCSS = {
+					position: sortable.placeholder.css( "position" ),
+					top: sortable.placeholder.css( "top" ),
+					left: sortable.placeholder.css( "left" )
+				};
+
+				sortable._mouseStop( event );
+
+				// Once drag has ended, the sortable should return to using
+				// its original helper, not the shared helper from draggable
+				sortable.options.helper = sortable.options._helper;
+			} else {
+
+				// Prevent this Sortable from removing the helper.
+				// However, don't set the draggable to remove the helper
+				// either as another connected Sortable may yet handle the removal.
+				sortable.cancelHelperRemoval = true;
+
+				sortable._trigger( "deactivate", event, uiSortable );
+			}
+		} );
+	},
+	drag: function( event, ui, draggable ) {
+		$.each( draggable.sortables, function() {
+			var innermostIntersecting = false,
+				sortable = this;
+
+			// Copy over variables that sortable's _intersectsWith uses
+			sortable.positionAbs = draggable.positionAbs;
+			sortable.helperProportions = draggable.helperProportions;
+			sortable.offset.click = draggable.offset.click;
+
+			if ( sortable._intersectsWith( sortable.containerCache ) ) {
+				innermostIntersecting = true;
+
+				$.each( draggable.sortables, function() {
+
+					// Copy over variables that sortable's _intersectsWith uses
+					this.positionAbs = draggable.positionAbs;
+					this.helperProportions = draggable.helperProportions;
+					this.offset.click = draggable.offset.click;
+
+					if ( this !== sortable &&
+							this._intersectsWith( this.containerCache ) &&
+							$.contains( sortable.element[ 0 ], this.element[ 0 ] ) ) {
+						innermostIntersecting = false;
+					}
+
+					return innermostIntersecting;
+				} );
+			}
+
+			if ( innermostIntersecting ) {
+
+				// If it intersects, we use a little isOver variable and set it once,
+				// so that the move-in stuff gets fired only once.
+				if ( !sortable.isOver ) {
+					sortable.isOver = 1;
+
+					// Store draggable's parent in case we need to reappend to it later.
+					draggable._parent = ui.helper.parent();
+
+					sortable.currentItem = ui.helper
+						.appendTo( sortable.element )
+						.data( "ui-sortable-item", true );
+
+					// Store helper option to later restore it
+					sortable.options._helper = sortable.options.helper;
+
+					sortable.options.helper = function() {
+						return ui.helper[ 0 ];
+					};
+
+					// Fire the start events of the sortable with our passed browser event,
+					// and our own helper (so it doesn't create a new one)
+					event.target = sortable.currentItem[ 0 ];
+					sortable._mouseCapture( event, true );
+					sortable._mouseStart( event, true, true );
+
+					// Because the browser event is way off the new appended portlet,
+					// modify necessary variables to reflect the changes
+					sortable.offset.click.top = draggable.offset.click.top;
+					sortable.offset.click.left = draggable.offset.click.left;
+					sortable.offset.parent.left -= draggable.offset.parent.left -
+						sortable.offset.parent.left;
+					sortable.offset.parent.top -= draggable.offset.parent.top -
+						sortable.offset.parent.top;
+
+					draggable._trigger( "toSortable", event );
+
+					// Inform draggable that the helper is in a valid drop zone,
+					// used solely in the revert option to handle "valid/invalid".
+					draggable.dropped = sortable.element;
+
+					// Need to refreshPositions of all sortables in the case that
+					// adding to one sortable changes the location of the other sortables (#9675)
+					$.each( draggable.sortables, function() {
+						this.refreshPositions();
+					} );
+
+					// Hack so receive/update callbacks work (mostly)
+					draggable.currentItem = draggable.element;
+					sortable.fromOutside = draggable;
+				}
+
+				if ( sortable.currentItem ) {
+					sortable._mouseDrag( event );
+
+					// Copy the sortable's position because the draggable's can potentially reflect
+					// a relative position, while sortable is always absolute, which the dragged
+					// element has now become. (#8809)
+					ui.position = sortable.position;
+				}
+			} else {
+
+				// If it doesn't intersect with the sortable, and it intersected before,
+				// we fake the drag stop of the sortable, but make sure it doesn't remove
+				// the helper by using cancelHelperRemoval.
+				if ( sortable.isOver ) {
+
+					sortable.isOver = 0;
+					sortable.cancelHelperRemoval = true;
+
+					// Calling sortable's mouseStop would trigger a revert,
+					// so revert must be temporarily false until after mouseStop is called.
+					sortable.options._revert = sortable.options.revert;
+					sortable.options.revert = false;
+
+					sortable._trigger( "out", event, sortable._uiHash( sortable ) );
+					sortable._mouseStop( event, true );
+
+					// Restore sortable behaviors that were modfied
+					// when the draggable entered the sortable area (#9481)
+					sortable.options.revert = sortable.options._revert;
+					sortable.options.helper = sortable.options._helper;
+
+					if ( sortable.placeholder ) {
+						sortable.placeholder.remove();
+					}
+
+					// Restore and recalculate the draggable's offset considering the sortable
+					// may have modified them in unexpected ways. (#8809, #10669)
+					ui.helper.appendTo( draggable._parent );
+					draggable._refreshOffsets( event );
+					ui.position = draggable._generatePosition( event, true );
+
+					draggable._trigger( "fromSortable", event );
+
+					// Inform draggable that the helper is no longer in a valid drop zone
+					draggable.dropped = false;
+
+					// Need to refreshPositions of all sortables just in case removing
+					// from one sortable changes the location of other sortables (#9675)
+					$.each( draggable.sortables, function() {
+						this.refreshPositions();
+					} );
+				}
+			}
+		} );
+	}
+} );
+
+$.ui.plugin.add( "draggable", "cursor", {
+	start: function( event, ui, instance ) {
+		var t = $( "body" ),
+			o = instance.options;
+
+		if ( t.css( "cursor" ) ) {
+			o._cursor = t.css( "cursor" );
+		}
+		t.css( "cursor", o.cursor );
+	},
+	stop: function( event, ui, instance ) {
+		var o = instance.options;
+		if ( o._cursor ) {
+			$( "body" ).css( "cursor", o._cursor );
+		}
+	}
+} );
+
+$.ui.plugin.add( "draggable", "opacity", {
+	start: function( event, ui, instance ) {
+		var t = $( ui.helper ),
+			o = instance.options;
+		if ( t.css( "opacity" ) ) {
+			o._opacity = t.css( "opacity" );
+		}
+		t.css( "opacity", o.opacity );
+	},
+	stop: function( event, ui, instance ) {
+		var o = instance.options;
+		if ( o._opacity ) {
+			$( ui.helper ).css( "opacity", o._opacity );
+		}
+	}
+} );
+
+$.ui.plugin.add( "draggable", "scroll", {
+	start: function( event, ui, i ) {
+		if ( !i.scrollParentNotHidden ) {
+			i.scrollParentNotHidden = i.helper.scrollParent( false );
+		}
+
+		if ( i.scrollParentNotHidden[ 0 ] !== i.document[ 0 ] &&
+				i.scrollParentNotHidden[ 0 ].tagName !== "HTML" ) {
+			i.overflowOffset = i.scrollParentNotHidden.offset();
+		}
+	},
+	drag: function( event, ui, i  ) {
+
+		var o = i.options,
+			scrolled = false,
+			scrollParent = i.scrollParentNotHidden[ 0 ],
+			document = i.document[ 0 ];
+
+		if ( scrollParent !== document && scrollParent.tagName !== "HTML" ) {
+			if ( !o.axis || o.axis !== "x" ) {
+				if ( ( i.overflowOffset.top + scrollParent.offsetHeight ) - event.pageY <
+						o.scrollSensitivity ) {
+					scrollParent.scrollTop = scrolled = scrollParent.scrollTop + o.scrollSpeed;
+				} else if ( event.pageY - i.overflowOffset.top < o.scrollSensitivity ) {
+					scrollParent.scrollTop = scrolled = scrollParent.scrollTop - o.scrollSpeed;
+				}
+			}
+
+			if ( !o.axis || o.axis !== "y" ) {
+				if ( ( i.overflowOffset.left + scrollParent.offsetWidth ) - event.pageX <
+						o.scrollSensitivity ) {
+					scrollParent.scrollLeft = scrolled = scrollParent.scrollLeft + o.scrollSpeed;
+				} else if ( event.pageX - i.overflowOffset.left < o.scrollSensitivity ) {
+					scrollParent.scrollLeft = scrolled = scrollParent.scrollLeft - o.scrollSpeed;
+				}
+			}
+
+		} else {
+
+			if ( !o.axis || o.axis !== "x" ) {
+				if ( event.pageY - $( document ).scrollTop() < o.scrollSensitivity ) {
+					scrolled = $( document ).scrollTop( $( document ).scrollTop() - o.scrollSpeed );
+				} else if ( $( window ).height() - ( event.pageY - $( document ).scrollTop() ) <
+						o.scrollSensitivity ) {
+					scrolled = $( document ).scrollTop( $( document ).scrollTop() + o.scrollSpeed );
+				}
+			}
+
+			if ( !o.axis || o.axis !== "y" ) {
+				if ( event.pageX - $( document ).scrollLeft() < o.scrollSensitivity ) {
+					scrolled = $( document ).scrollLeft(
+						$( document ).scrollLeft() - o.scrollSpeed
+					);
+				} else if ( $( window ).width() - ( event.pageX - $( document ).scrollLeft() ) <
+						o.scrollSensitivity ) {
+					scrolled = $( document ).scrollLeft(
+						$( document ).scrollLeft() + o.scrollSpeed
+					);
+				}
+			}
+
+		}
+
+		if ( scrolled !== false && $.ui.ddmanager && !o.dropBehaviour ) {
+			$.ui.ddmanager.prepareOffsets( i, event );
+		}
+
+	}
+} );
+
+$.ui.plugin.add( "draggable", "snap", {
+	start: function( event, ui, i ) {
+
+		var o = i.options;
+
+		i.snapElements = [];
+
+		$( o.snap.constructor !== String ? ( o.snap.items || ":data(ui-draggable)" ) : o.snap )
+			.each( function() {
+				var $t = $( this ),
+					$o = $t.offset();
+				if ( this !== i.element[ 0 ] ) {
+					i.snapElements.push( {
+						item: this,
+						width: $t.outerWidth(), height: $t.outerHeight(),
+						top: $o.top, left: $o.left
+					} );
+				}
+			} );
+
+	},
+	drag: function( event, ui, inst ) {
+
+		var ts, bs, ls, rs, l, r, t, b, i, first,
+			o = inst.options,
+			d = o.snapTolerance,
+			x1 = ui.offset.left, x2 = x1 + inst.helperProportions.width,
+			y1 = ui.offset.top, y2 = y1 + inst.helperProportions.height;
+
+		for ( i = inst.snapElements.length - 1; i >= 0; i-- ) {
+
+			l = inst.snapElements[ i ].left - inst.margins.left;
+			r = l + inst.snapElements[ i ].width;
+			t = inst.snapElements[ i ].top - inst.margins.top;
+			b = t + inst.snapElements[ i ].height;
+
+			if ( x2 < l - d || x1 > r + d || y2 < t - d || y1 > b + d ||
+					!$.contains( inst.snapElements[ i ].item.ownerDocument,
+					inst.snapElements[ i ].item ) ) {
+				if ( inst.snapElements[ i ].snapping ) {
+					( inst.options.snap.release &&
+						inst.options.snap.release.call(
+							inst.element,
+							event,
+							$.extend( inst._uiHash(), { snapItem: inst.snapElements[ i ].item } )
+						) );
+				}
+				inst.snapElements[ i ].snapping = false;
+				continue;
+			}
+
+			if ( o.snapMode !== "inner" ) {
+				ts = Math.abs( t - y2 ) <= d;
+				bs = Math.abs( b - y1 ) <= d;
+				ls = Math.abs( l - x2 ) <= d;
+				rs = Math.abs( r - x1 ) <= d;
+				if ( ts ) {
+					ui.position.top = inst._convertPositionTo( "relative", {
+						top: t - inst.helperProportions.height,
+						left: 0
+					} ).top;
+				}
+				if ( bs ) {
+					ui.position.top = inst._convertPositionTo( "relative", {
+						top: b,
+						left: 0
+					} ).top;
+				}
+				if ( ls ) {
+					ui.position.left = inst._convertPositionTo( "relative", {
+						top: 0,
+						left: l - inst.helperProportions.width
+					} ).left;
+				}
+				if ( rs ) {
+					ui.position.left = inst._convertPositionTo( "relative", {
+						top: 0,
+						left: r
+					} ).left;
+				}
+			}
+
+			first = ( ts || bs || ls || rs );
+
+			if ( o.snapMode !== "outer" ) {
+				ts = Math.abs( t - y1 ) <= d;
+				bs = Math.abs( b - y2 ) <= d;
+				ls = Math.abs( l - x1 ) <= d;
+				rs = Math.abs( r - x2 ) <= d;
+				if ( ts ) {
+					ui.position.top = inst._convertPositionTo( "relative", {
+						top: t,
+						left: 0
+					} ).top;
+				}
+				if ( bs ) {
+					ui.position.top = inst._convertPositionTo( "relative", {
+						top: b - inst.helperProportions.height,
+						left: 0
+					} ).top;
+				}
+				if ( ls ) {
+					ui.position.left = inst._convertPositionTo( "relative", {
+						top: 0,
+						left: l
+					} ).left;
+				}
+				if ( rs ) {
+					ui.position.left = inst._convertPositionTo( "relative", {
+						top: 0,
+						left: r - inst.helperProportions.width
+					} ).left;
+				}
+			}
+
+			if ( !inst.snapElements[ i ].snapping && ( ts || bs || ls || rs || first ) ) {
+				( inst.options.snap.snap &&
+					inst.options.snap.snap.call(
+						inst.element,
+						event,
+						$.extend( inst._uiHash(), {
+							snapItem: inst.snapElements[ i ].item
+						} ) ) );
+			}
+			inst.snapElements[ i ].snapping = ( ts || bs || ls || rs || first );
+
+		}
+
+	}
+} );
+
+$.ui.plugin.add( "draggable", "stack", {
+	start: function( event, ui, instance ) {
+		var min,
+			o = instance.options,
+			group = $.makeArray( $( o.stack ) ).sort( function( a, b ) {
+				return ( parseInt( $( a ).css( "zIndex" ), 10 ) || 0 ) -
+					( parseInt( $( b ).css( "zIndex" ), 10 ) || 0 );
+			} );
+
+		if ( !group.length ) { return; }
+
+		min = parseInt( $( group[ 0 ] ).css( "zIndex" ), 10 ) || 0;
+		$( group ).each( function( i ) {
+			$( this ).css( "zIndex", min + i );
+		} );
+		this.css( "zIndex", ( min + group.length ) );
+	}
+} );
+
+$.ui.plugin.add( "draggable", "zIndex", {
+	start: function( event, ui, instance ) {
+		var t = $( ui.helper ),
+			o = instance.options;
+
+		if ( t.css( "zIndex" ) ) {
+			o._zIndex = t.css( "zIndex" );
+		}
+		t.css( "zIndex", o.zIndex );
+	},
+	stop: function( event, ui, instance ) {
+		var o = instance.options;
+
+		if ( o._zIndex ) {
+			$( ui.helper ).css( "zIndex", o._zIndex );
+		}
+	}
+} );
+
+var widgetsDraggable = $.ui.draggable;
+
+
+/*!
+ * jQuery UI Resizable 1.12.0
+ * http://jqueryui.com
+ *
+ * Copyright jQuery Foundation and other contributors
+ * Released under the MIT license.
+ * http://jquery.org/license
+ */
+
+//>>label: Resizable
+//>>group: Interactions
+//>>description: Enables resize functionality for any element.
+//>>docs: http://api.jqueryui.com/resizable/
+//>>demos: http://jqueryui.com/resizable/
+//>>css.structure: ../../themes/base/core.css
+//>>css.structure: ../../themes/base/resizable.css
+//>>css.theme: ../../themes/base/theme.css
+
+
+
+$.widget( "ui.resizable", $.ui.mouse, {
+	version: "1.12.0",
+	widgetEventPrefix: "resize",
+	options: {
+		alsoResize: false,
+		animate: false,
+		animateDuration: "slow",
+		animateEasing: "swing",
+		aspectRatio: false,
+		autoHide: false,
+		classes: {
+			"ui-resizable-se": "ui-icon ui-icon-gripsmall-diagonal-se"
+		},
+		containment: false,
+		ghost: false,
+		grid: false,
+		handles: "e,s,se",
+		helper: false,
+		maxHeight: null,
+		maxWidth: null,
+		minHeight: 10,
+		minWidth: 10,
+
+		// See #7960
+		zIndex: 90,
+
+		// Callbacks
+		resize: null,
+		start: null,
+		stop: null
+	},
+
+	_num: function( value ) {
+		return parseFloat( value ) || 0;
+	},
+
+	_isNumber: function( value ) {
+		return !isNaN( parseFloat( value ) );
+	},
+
+	_hasScroll: function( el, a ) {
+
+		if ( $( el ).css( "overflow" ) === "hidden" ) {
+			return false;
+		}
+
+		var scroll = ( a && a === "left" ) ? "scrollLeft" : "scrollTop",
+			has = false;
+
+		if ( el[ scroll ] > 0 ) {
+			return true;
+		}
+
+		// TODO: determine which cases actually cause this to happen
+		// if the element doesn't have the scroll set, see if it's possible to
+		// set the scroll
+		el[ scroll ] = 1;
+		has = ( el[ scroll ] > 0 );
+		el[ scroll ] = 0;
+		return has;
+	},
+
+	_create: function() {
+
+		var margins,
+			o = this.options,
+			that = this;
+		this._addClass( "ui-resizable" );
+
+		$.extend( this, {
+			_aspectRatio: !!( o.aspectRatio ),
+			aspectRatio: o.aspectRatio,
+			originalElement: this.element,
+			_proportionallyResizeElements: [],
+			_helper: o.helper || o.ghost || o.animate ? o.helper || "ui-resizable-helper" : null
+		} );
+
+		// Wrap the element if it cannot hold child nodes
+		if ( this.element[ 0 ].nodeName.match( /^(canvas|textarea|input|select|button|img)$/i ) ) {
+
+			this.element.wrap(
+				$( "<div class='ui-wrapper' style='overflow: hidden;'></div>" ).css( {
+					position: this.element.css( "position" ),
+					width: this.element.outerWidth(),
+					height: this.element.outerHeight(),
+					top: this.element.css( "top" ),
+					left: this.element.css( "left" )
+				} )
+			);
+
+			this.element = this.element.parent().data(
+				"ui-resizable", this.element.resizable( "instance" )
+			);
+
+			this.elementIsWrapper = true;
+
+			margins = {
+				marginTop: this.originalElement.css( "marginTop" ),
+				marginRight: this.originalElement.css( "marginRight" ),
+				marginBottom: this.originalElement.css( "marginBottom" ),
+				marginLeft: this.originalElement.css( "marginLeft" )
+			};
+
+			this.element.css( margins );
+			this.originalElement.css( "margin", 0 );
+
+			// support: Safari
+			// Prevent Safari textarea resize
+			this.originalResizeStyle = this.originalElement.css( "resize" );
+			this.originalElement.css( "resize", "none" );
+
+			this._proportionallyResizeElements.push( this.originalElement.css( {
+				position: "static",
+				zoom: 1,
+				display: "block"
+			} ) );
+
+			// Support: IE9
+			// avoid IE jump (hard set the margin)
+			this.originalElement.css( margins );
+
+			this._proportionallyResize();
+		}
+
+		this._setupHandles();
+
+		if ( o.autoHide ) {
+			$( this.element )
+				.on( "mouseenter", function() {
+					if ( o.disabled ) {
+						return;
+					}
+					that._removeClass( "ui-resizable-autohide" );
+					that._handles.show();
+				} )
+				.on( "mouseleave", function() {
+					if ( o.disabled ) {
+						return;
+					}
+					if ( !that.resizing ) {
+						that._addClass( "ui-resizable-autohide" );
+						that._handles.hide();
+					}
+				} );
+		}
+
+		this._mouseInit();
+	},
+
+	_destroy: function() {
+
+		this._mouseDestroy();
+
+		var wrapper,
+			_destroy = function( exp ) {
+				$( exp )
+					.removeData( "resizable" )
+					.removeData( "ui-resizable" )
+					.off( ".resizable" )
+					.find( ".ui-resizable-handle" )
+						.remove();
+			};
+
+		// TODO: Unwrap at same DOM position
+		if ( this.elementIsWrapper ) {
+			_destroy( this.element );
+			wrapper = this.element;
+			this.originalElement.css( {
+				position: wrapper.css( "position" ),
+				width: wrapper.outerWidth(),
+				height: wrapper.outerHeight(),
+				top: wrapper.css( "top" ),
+				left: wrapper.css( "left" )
+			} ).insertAfter( wrapper );
+			wrapper.remove();
+		}
+
+		this.originalElement.css( "resize", this.originalResizeStyle );
+		_destroy( this.originalElement );
+
+		return this;
+	},
+
+	_setOption: function( key, value ) {
+		this._super( key, value );
+
+		switch ( key ) {
+		case "handles":
+			this._removeHandles();
+			this._setupHandles();
+			break;
+		default:
+			break;
+		}
+	},
+
+	_setupHandles: function() {
+		var o = this.options, handle, i, n, hname, axis, that = this;
+		this.handles = o.handles ||
+			( !$( ".ui-resizable-handle", this.element ).length ?
+				"e,s,se" : {
+					n: ".ui-resizable-n",
+					e: ".ui-resizable-e",
+					s: ".ui-resizable-s",
+					w: ".ui-resizable-w",
+					se: ".ui-resizable-se",
+					sw: ".ui-resizable-sw",
+					ne: ".ui-resizable-ne",
+					nw: ".ui-resizable-nw"
+				} );
+
+		this._handles = $();
+		if ( this.handles.constructor === String ) {
+
+			if ( this.handles === "all" ) {
+				this.handles = "n,e,s,w,se,sw,ne,nw";
+			}
+
+			n = this.handles.split( "," );
+			this.handles = {};
+
+			for ( i = 0; i < n.length; i++ ) {
+
+				handle = $.trim( n[ i ] );
+				hname = "ui-resizable-" + handle;
+				axis = $( "<div>" );
+				this._addClass( axis, "ui-resizable-handle " + hname );
+
+				axis.css( { zIndex: o.zIndex } );
+
+				this.handles[ handle ] = ".ui-resizable-" + handle;
+				this.element.append( axis );
+			}
+
+		}
+
+		this._renderAxis = function( target ) {
+
+			var i, axis, padPos, padWrapper;
+
+			target = target || this.element;
+
+			for ( i in this.handles ) {
+
+				if ( this.handles[ i ].constructor === String ) {
+					this.handles[ i ] = this.element.children( this.handles[ i ] ).first().show();
+				} else if ( this.handles[ i ].jquery || this.handles[ i ].nodeType ) {
+					this.handles[ i ] = $( this.handles[ i ] );
+					this._on( this.handles[ i ], { "mousedown": that._mouseDown } );
+				}
+
+				if ( this.elementIsWrapper &&
+						this.originalElement[ 0 ]
+							.nodeName
+							.match( /^(textarea|input|select|button)$/i ) ) {
+					axis = $( this.handles[ i ], this.element );
+
+					padWrapper = /sw|ne|nw|se|n|s/.test( i ) ?
+						axis.outerHeight() :
+						axis.outerWidth();
+
+					padPos = [ "padding",
+						/ne|nw|n/.test( i ) ? "Top" :
+						/se|sw|s/.test( i ) ? "Bottom" :
+						/^e$/.test( i ) ? "Right" : "Left" ].join( "" );
+
+					target.css( padPos, padWrapper );
+
+					this._proportionallyResize();
+				}
+
+				this._handles = this._handles.add( this.handles[ i ] );
+			}
+		};
+
+		// TODO: make renderAxis a prototype function
+		this._renderAxis( this.element );
+
+		this._handles = this._handles.add( this.element.find( ".ui-resizable-handle" ) );
+		this._handles.disableSelection();
+
+		this._handles.on( "mouseover", function() {
+			if ( !that.resizing ) {
+				if ( this.className ) {
+					axis = this.className.match( /ui-resizable-(se|sw|ne|nw|n|e|s|w)/i );
+				}
+				that.axis = axis && axis[ 1 ] ? axis[ 1 ] : "se";
+			}
+		} );
+
+		if ( o.autoHide ) {
+			this._handles.hide();
+			this._addClass( "ui-resizable-autohide" );
+		}
+	},
+
+	_removeHandles: function() {
+		this._handles.remove();
+	},
+
+	_mouseCapture: function( event ) {
+		var i, handle,
+			capture = false;
+
+		for ( i in this.handles ) {
+			handle = $( this.handles[ i ] )[ 0 ];
+			if ( handle === event.target || $.contains( handle, event.target ) ) {
+				capture = true;
+			}
+		}
+
+		return !this.options.disabled && capture;
+	},
+
+	_mouseStart: function( event ) {
+
+		var curleft, curtop, cursor,
+			o = this.options,
+			el = this.element;
+
+		this.resizing = true;
+
+		this._renderProxy();
+
+		curleft = this._num( this.helper.css( "left" ) );
+		curtop = this._num( this.helper.css( "top" ) );
+
+		if ( o.containment ) {
+			curleft += $( o.containment ).scrollLeft() || 0;
+			curtop += $( o.containment ).scrollTop() || 0;
+		}
+
+		this.offset = this.helper.offset();
+		this.position = { left: curleft, top: curtop };
+
+		this.size = this._helper ? {
+				width: this.helper.width(),
+				height: this.helper.height()
+			} : {
+				width: el.width(),
+				height: el.height()
+			};
+
+		this.originalSize = this._helper ? {
+				width: el.outerWidth(),
+				height: el.outerHeight()
+			} : {
+				width: el.width(),
+				height: el.height()
+			};
+
+		this.sizeDiff = {
+			width: el.outerWidth() - el.width(),
+			height: el.outerHeight() - el.height()
+		};
+
+		this.originalPosition = { left: curleft, top: curtop };
+		this.originalMousePosition = { left: event.pageX, top: event.pageY };
+
+		this.aspectRatio = ( typeof o.aspectRatio === "number" ) ?
+			o.aspectRatio :
+			( ( this.originalSize.width / this.originalSize.height ) || 1 );
+
+		cursor = $( ".ui-resizable-" + this.axis ).css( "cursor" );
+		$( "body" ).css( "cursor", cursor === "auto" ? this.axis + "-resize" : cursor );
+
+		this._addClass( "ui-resizable-resizing" );
+		this._propagate( "start", event );
+		return true;
+	},
+
+	_mouseDrag: function( event ) {
+
+		var data, props,
+			smp = this.originalMousePosition,
+			a = this.axis,
+			dx = ( event.pageX - smp.left ) || 0,
+			dy = ( event.pageY - smp.top ) || 0,
+			trigger = this._change[ a ];
+
+		this._updatePrevProperties();
+
+		if ( !trigger ) {
+			return false;
+		}
+
+		data = trigger.apply( this, [ event, dx, dy ] );
+
+		this._updateVirtualBoundaries( event.shiftKey );
+		if ( this._aspectRatio || event.shiftKey ) {
+			data = this._updateRatio( data, event );
+		}
+
+		data = this._respectSize( data, event );
+
+		this._updateCache( data );
+
+		this._propagate( "resize", event );
+
+		props = this._applyChanges();
+
+		if ( !this._helper && this._proportionallyResizeElements.length ) {
+			this._proportionallyResize();
+		}
+
+		if ( !$.isEmptyObject( props ) ) {
+			this._updatePrevProperties();
+			this._trigger( "resize", event, this.ui() );
+			this._applyChanges();
+		}
+
+		return false;
+	},
+
+	_mouseStop: function( event ) {
+
+		this.resizing = false;
+		var pr, ista, soffseth, soffsetw, s, left, top,
+			o = this.options, that = this;
+
+		if ( this._helper ) {
+
+			pr = this._proportionallyResizeElements;
+			ista = pr.length && ( /textarea/i ).test( pr[ 0 ].nodeName );
+			soffseth = ista && this._hasScroll( pr[ 0 ], "left" ) ? 0 : that.sizeDiff.height;
+			soffsetw = ista ? 0 : that.sizeDiff.width;
+
+			s = {
+				width: ( that.helper.width()  - soffsetw ),
+				height: ( that.helper.height() - soffseth )
+			};
+			left = ( parseFloat( that.element.css( "left" ) ) +
+				( that.position.left - that.originalPosition.left ) ) || null;
+			top = ( parseFloat( that.element.css( "top" ) ) +
+				( that.position.top - that.originalPosition.top ) ) || null;
+
+			if ( !o.animate ) {
+				this.element.css( $.extend( s, { top: top, left: left } ) );
+			}
+
+			that.helper.height( that.size.height );
+			that.helper.width( that.size.width );
+
+			if ( this._helper && !o.animate ) {
+				this._proportionallyResize();
+			}
+		}
+
+		$( "body" ).css( "cursor", "auto" );
+
+		this._removeClass( "ui-resizable-resizing" );
+
+		this._propagate( "stop", event );
+
+		if ( this._helper ) {
+			this.helper.remove();
+		}
+
+		return false;
+
+	},
+
+	_updatePrevProperties: function() {
+		this.prevPosition = {
+			top: this.position.top,
+			left: this.position.left
+		};
+		this.prevSize = {
+			width: this.size.width,
+			height: this.size.height
+		};
+	},
+
+	_applyChanges: function() {
+		var props = {};
+
+		if ( this.position.top !== this.prevPosition.top ) {
+			props.top = this.position.top + "px";
+		}
+		if ( this.position.left !== this.prevPosition.left ) {
+			props.left = this.position.left + "px";
+		}
+		if ( this.size.width !== this.prevSize.width ) {
+			props.width = this.size.width + "px";
+		}
+		if ( this.size.height !== this.prevSize.height ) {
+			props.height = this.size.height + "px";
+		}
+
+		this.helper.css( props );
+
+		return props;
+	},
+
+	_updateVirtualBoundaries: function( forceAspectRatio ) {
+		var pMinWidth, pMaxWidth, pMinHeight, pMaxHeight, b,
+			o = this.options;
+
+		b = {
+			minWidth: this._isNumber( o.minWidth ) ? o.minWidth : 0,
+			maxWidth: this._isNumber( o.maxWidth ) ? o.maxWidth : Infinity,
+			minHeight: this._isNumber( o.minHeight ) ? o.minHeight : 0,
+			maxHeight: this._isNumber( o.maxHeight ) ? o.maxHeight : Infinity
+		};
+
+		if ( this._aspectRatio || forceAspectRatio ) {
+			pMinWidth = b.minHeight * this.aspectRatio;
+			pMinHeight = b.minWidth / this.aspectRatio;
+			pMaxWidth = b.maxHeight * this.aspectRatio;
+			pMaxHeight = b.maxWidth / this.aspectRatio;
+
+			if ( pMinWidth > b.minWidth ) {
+				b.minWidth = pMinWidth;
+			}
+			if ( pMinHeight > b.minHeight ) {
+				b.minHeight = pMinHeight;
+			}
+			if ( pMaxWidth < b.maxWidth ) {
+				b.maxWidth = pMaxWidth;
+			}
+			if ( pMaxHeight < b.maxHeight ) {
+				b.maxHeight = pMaxHeight;
+			}
+		}
+		this._vBoundaries = b;
+	},
+
+	_updateCache: function( data ) {
+		this.offset = this.helper.offset();
+		if ( this._isNumber( data.left ) ) {
+			this.position.left = data.left;
+		}
+		if ( this._isNumber( data.top ) ) {
+			this.position.top = data.top;
+		}
+		if ( this._isNumber( data.height ) ) {
+			this.size.height = data.height;
+		}
+		if ( this._isNumber( data.width ) ) {
+			this.size.width = data.width;
+		}
+	},
+
+	_updateRatio: function( data ) {
+
+		var cpos = this.position,
+			csize = this.size,
+			a = this.axis;
+
+		if ( this._isNumber( data.height ) ) {
+			data.width = ( data.height * this.aspectRatio );
+		} else if ( this._isNumber( data.width ) ) {
+			data.height = ( data.width / this.aspectRatio );
+		}
+
+		if ( a === "sw" ) {
+			data.left = cpos.left + ( csize.width - data.width );
+			data.top = null;
+		}
+		if ( a === "nw" ) {
+			data.top = cpos.top + ( csize.height - data.height );
+			data.left = cpos.left + ( csize.width - data.width );
+		}
+
+		return data;
+	},
+
+	_respectSize: function( data ) {
+
+		var o = this._vBoundaries,
+			a = this.axis,
+			ismaxw = this._isNumber( data.width ) && o.maxWidth && ( o.maxWidth < data.width ),
+			ismaxh = this._isNumber( data.height ) && o.maxHeight && ( o.maxHeight < data.height ),
+			isminw = this._isNumber( data.width ) && o.minWidth && ( o.minWidth > data.width ),
+			isminh = this._isNumber( data.height ) && o.minHeight && ( o.minHeight > data.height ),
+			dw = this.originalPosition.left + this.originalSize.width,
+			dh = this.originalPosition.top + this.originalSize.height,
+			cw = /sw|nw|w/.test( a ), ch = /nw|ne|n/.test( a );
+		if ( isminw ) {
+			data.width = o.minWidth;
+		}
+		if ( isminh ) {
+			data.height = o.minHeight;
+		}
+		if ( ismaxw ) {
+			data.width = o.maxWidth;
+		}
+		if ( ismaxh ) {
+			data.height = o.maxHeight;
+		}
+
+		if ( isminw && cw ) {
+			data.left = dw - o.minWidth;
+		}
+		if ( ismaxw && cw ) {
+			data.left = dw - o.maxWidth;
+		}
+		if ( isminh && ch ) {
+			data.top = dh - o.minHeight;
+		}
+		if ( ismaxh && ch ) {
+			data.top = dh - o.maxHeight;
+		}
+
+		// Fixing jump error on top/left - bug #2330
+		if ( !data.width && !data.height && !data.left && data.top ) {
+			data.top = null;
+		} else if ( !data.width && !data.height && !data.top && data.left ) {
+			data.left = null;
+		}
+
+		return data;
+	},
+
+	_getPaddingPlusBorderDimensions: function( element ) {
+		var i = 0,
+			widths = [],
+			borders = [
+				element.css( "borderTopWidth" ),
+				element.css( "borderRightWidth" ),
+				element.css( "borderBottomWidth" ),
+				element.css( "borderLeftWidth" )
+			],
+			paddings = [
+				element.css( "paddingTop" ),
+				element.css( "paddingRight" ),
+				element.css( "paddingBottom" ),
+				element.css( "paddingLeft" )
+			];
+
+		for ( ; i < 4; i++ ) {
+			widths[ i ] = ( parseFloat( borders[ i ] ) || 0 );
+			widths[ i ] += ( parseFloat( paddings[ i ] ) || 0 );
+		}
+
+		return {
+			height: widths[ 0 ] + widths[ 2 ],
+			width: widths[ 1 ] + widths[ 3 ]
+		};
+	},
+
+	_proportionallyResize: function() {
+
+		if ( !this._proportionallyResizeElements.length ) {
+			return;
+		}
+
+		var prel,
+			i = 0,
+			element = this.helper || this.element;
+
+		for ( ; i < this._proportionallyResizeElements.length; i++ ) {
+
+			prel = this._proportionallyResizeElements[ i ];
+
+			// TODO: Seems like a bug to cache this.outerDimensions
+			// considering that we are in a loop.
+			if ( !this.outerDimensions ) {
+				this.outerDimensions = this._getPaddingPlusBorderDimensions( prel );
+			}
+
+			prel.css( {
+				height: ( element.height() - this.outerDimensions.height ) || 0,
+				width: ( element.width() - this.outerDimensions.width ) || 0
+			} );
+
+		}
+
+	},
+
+	_renderProxy: function() {
+
+		var el = this.element, o = this.options;
+		this.elementOffset = el.offset();
+
+		if ( this._helper ) {
+
+			this.helper = this.helper || $( "<div style='overflow:hidden;'></div>" );
+
+			this._addClass( this.helper, this._helper );
+			this.helper.css( {
+				width: this.element.outerWidth(),
+				height: this.element.outerHeight(),
+				position: "absolute",
+				left: this.elementOffset.left + "px",
+				top: this.elementOffset.top + "px",
+				zIndex: ++o.zIndex //TODO: Don't modify option
+			} );
+
+			this.helper
+				.appendTo( "body" )
+				.disableSelection();
+
+		} else {
+			this.helper = this.element;
+		}
+
+	},
+
+	_change: {
+		e: function( event, dx ) {
+			return { width: this.originalSize.width + dx };
+		},
+		w: function( event, dx ) {
+			var cs = this.originalSize, sp = this.originalPosition;
+			return { left: sp.left + dx, width: cs.width - dx };
+		},
+		n: function( event, dx, dy ) {
+			var cs = this.originalSize, sp = this.originalPosition;
+			return { top: sp.top + dy, height: cs.height - dy };
+		},
+		s: function( event, dx, dy ) {
+			return { height: this.originalSize.height + dy };
+		},
+		se: function( event, dx, dy ) {
+			return $.extend( this._change.s.apply( this, arguments ),
+				this._change.e.apply( this, [ event, dx, dy ] ) );
+		},
+		sw: function( event, dx, dy ) {
+			return $.extend( this._change.s.apply( this, arguments ),
+				this._change.w.apply( this, [ event, dx, dy ] ) );
+		},
+		ne: function( event, dx, dy ) {
+			return $.extend( this._change.n.apply( this, arguments ),
+				this._change.e.apply( this, [ event, dx, dy ] ) );
+		},
+		nw: function( event, dx, dy ) {
+			return $.extend( this._change.n.apply( this, arguments ),
+				this._change.w.apply( this, [ event, dx, dy ] ) );
+		}
+	},
+
+	_propagate: function( n, event ) {
+		$.ui.plugin.call( this, n, [ event, this.ui() ] );
+		( n !== "resize" && this._trigger( n, event, this.ui() ) );
+	},
+
+	plugins: {},
+
+	ui: function() {
+		return {
+			originalElement: this.originalElement,
+			element: this.element,
+			helper: this.helper,
+			position: this.position,
+			size: this.size,
+			originalSize: this.originalSize,
+			originalPosition: this.originalPosition
+		};
+	}
+
+} );
+
+/*
+ * Resizable Extensions
+ */
+
+$.ui.plugin.add( "resizable", "animate", {
+
+	stop: function( event ) {
+		var that = $( this ).resizable( "instance" ),
+			o = that.options,
+			pr = that._proportionallyResizeElements,
+			ista = pr.length && ( /textarea/i ).test( pr[ 0 ].nodeName ),
+			soffseth = ista && that._hasScroll( pr[ 0 ], "left" ) ? 0 : that.sizeDiff.height,
+			soffsetw = ista ? 0 : that.sizeDiff.width,
+			style = {
+				width: ( that.size.width - soffsetw ),
+				height: ( that.size.height - soffseth )
+			},
+			left = ( parseFloat( that.element.css( "left" ) ) +
+				( that.position.left - that.originalPosition.left ) ) || null,
+			top = ( parseFloat( that.element.css( "top" ) ) +
+				( that.position.top - that.originalPosition.top ) ) || null;
+
+		that.element.animate(
+			$.extend( style, top && left ? { top: top, left: left } : {} ), {
+				duration: o.animateDuration,
+				easing: o.animateEasing,
+				step: function() {
+
+					var data = {
+						width: parseFloat( that.element.css( "width" ) ),
+						height: parseFloat( that.element.css( "height" ) ),
+						top: parseFloat( that.element.css( "top" ) ),
+						left: parseFloat( that.element.css( "left" ) )
+					};
+
+					if ( pr && pr.length ) {
+						$( pr[ 0 ] ).css( { width: data.width, height: data.height } );
+					}
+
+					// Propagating resize, and updating values for each animation step
+					that._updateCache( data );
+					that._propagate( "resize", event );
+
+				}
+			}
+		);
+	}
+
+} );
+
+$.ui.plugin.add( "resizable", "containment", {
+
+	start: function() {
+		var element, p, co, ch, cw, width, height,
+			that = $( this ).resizable( "instance" ),
+			o = that.options,
+			el = that.element,
+			oc = o.containment,
+			ce = ( oc instanceof $ ) ?
+				oc.get( 0 ) :
+				( /parent/.test( oc ) ) ? el.parent().get( 0 ) : oc;
+
+		if ( !ce ) {
+			return;
+		}
+
+		that.containerElement = $( ce );
+
+		if ( /document/.test( oc ) || oc === document ) {
+			that.containerOffset = {
+				left: 0,
+				top: 0
+			};
+			that.containerPosition = {
+				left: 0,
+				top: 0
+			};
+
+			that.parentData = {
+				element: $( document ),
+				left: 0,
+				top: 0,
+				width: $( document ).width(),
+				height: $( document ).height() || document.body.parentNode.scrollHeight
+			};
+		} else {
+			element = $( ce );
+			p = [];
+			$( [ "Top", "Right", "Left", "Bottom" ] ).each( function( i, name ) {
+				p[ i ] = that._num( element.css( "padding" + name ) );
+			} );
+
+			that.containerOffset = element.offset();
+			that.containerPosition = element.position();
+			that.containerSize = {
+				height: ( element.innerHeight() - p[ 3 ] ),
+				width: ( element.innerWidth() - p[ 1 ] )
+			};
+
+			co = that.containerOffset;
+			ch = that.containerSize.height;
+			cw = that.containerSize.width;
+			width = ( that._hasScroll ( ce, "left" ) ? ce.scrollWidth : cw );
+			height = ( that._hasScroll ( ce ) ? ce.scrollHeight : ch ) ;
+
+			that.parentData = {
+				element: ce,
+				left: co.left,
+				top: co.top,
+				width: width,
+				height: height
+			};
+		}
+	},
+
+	resize: function( event ) {
+		var woset, hoset, isParent, isOffsetRelative,
+			that = $( this ).resizable( "instance" ),
+			o = that.options,
+			co = that.containerOffset,
+			cp = that.position,
+			pRatio = that._aspectRatio || event.shiftKey,
+			cop = {
+				top: 0,
+				left: 0
+			},
+			ce = that.containerElement,
+			continueResize = true;
+
+		if ( ce[ 0 ] !== document && ( /static/ ).test( ce.css( "position" ) ) ) {
+			cop = co;
+		}
+
+		if ( cp.left < ( that._helper ? co.left : 0 ) ) {
+			that.size.width = that.size.width +
+				( that._helper ?
+					( that.position.left - co.left ) :
+					( that.position.left - cop.left ) );
+
+			if ( pRatio ) {
+				that.size.height = that.size.width / that.aspectRatio;
+				continueResize = false;
+			}
+			that.position.left = o.helper ? co.left : 0;
+		}
+
+		if ( cp.top < ( that._helper ? co.top : 0 ) ) {
+			that.size.height = that.size.height +
+				( that._helper ?
+					( that.position.top - co.top ) :
+					that.position.top );
+
+			if ( pRatio ) {
+				that.size.width = that.size.height * that.aspectRatio;
+				continueResize = false;
+			}
+			that.position.top = that._helper ? co.top : 0;
+		}
+
+		isParent = that.containerElement.get( 0 ) === that.element.parent().get( 0 );
+		isOffsetRelative = /relative|absolute/.test( that.containerElement.css( "position" ) );
+
+		if ( isParent && isOffsetRelative ) {
+			that.offset.left = that.parentData.left + that.position.left;
+			that.offset.top = that.parentData.top + that.position.top;
+		} else {
+			that.offset.left = that.element.offset().left;
+			that.offset.top = that.element.offset().top;
+		}
+
+		woset = Math.abs( that.sizeDiff.width +
+			( that._helper ?
+				that.offset.left - cop.left :
+				( that.offset.left - co.left ) ) );
+
+		hoset = Math.abs( that.sizeDiff.height +
+			( that._helper ?
+				that.offset.top - cop.top :
+				( that.offset.top - co.top ) ) );
+
+		if ( woset + that.size.width >= that.parentData.width ) {
+			that.size.width = that.parentData.width - woset;
+			if ( pRatio ) {
+				that.size.height = that.size.width / that.aspectRatio;
+				continueResize = false;
+			}
+		}
+
+		if ( hoset + that.size.height >= that.parentData.height ) {
+			that.size.height = that.parentData.height - hoset;
+			if ( pRatio ) {
+				that.size.width = that.size.height * that.aspectRatio;
+				continueResize = false;
+			}
+		}
+
+		if ( !continueResize ) {
+			that.position.left = that.prevPosition.left;
+			that.position.top = that.prevPosition.top;
+			that.size.width = that.prevSize.width;
+			that.size.height = that.prevSize.height;
+		}
+	},
+
+	stop: function() {
+		var that = $( this ).resizable( "instance" ),
+			o = that.options,
+			co = that.containerOffset,
+			cop = that.containerPosition,
+			ce = that.containerElement,
+			helper = $( that.helper ),
+			ho = helper.offset(),
+			w = helper.outerWidth() - that.sizeDiff.width,
+			h = helper.outerHeight() - that.sizeDiff.height;
+
+		if ( that._helper && !o.animate && ( /relative/ ).test( ce.css( "position" ) ) ) {
+			$( this ).css( {
+				left: ho.left - cop.left - co.left,
+				width: w,
+				height: h
+			} );
+		}
+
+		if ( that._helper && !o.animate && ( /static/ ).test( ce.css( "position" ) ) ) {
+			$( this ).css( {
+				left: ho.left - cop.left - co.left,
+				width: w,
+				height: h
+			} );
+		}
+	}
+} );
+
+$.ui.plugin.add( "resizable", "alsoResize", {
+
+	start: function() {
+		var that = $( this ).resizable( "instance" ),
+			o = that.options;
+
+		$( o.alsoResize ).each( function() {
+			var el = $( this );
+			el.data( "ui-resizable-alsoresize", {
+				width: parseFloat( el.width() ), height: parseFloat( el.height() ),
+				left: parseFloat( el.css( "left" ) ), top: parseFloat( el.css( "top" ) )
+			} );
+		} );
+	},
+
+	resize: function( event, ui ) {
+		var that = $( this ).resizable( "instance" ),
+			o = that.options,
+			os = that.originalSize,
+			op = that.originalPosition,
+			delta = {
+				height: ( that.size.height - os.height ) || 0,
+				width: ( that.size.width - os.width ) || 0,
+				top: ( that.position.top - op.top ) || 0,
+				left: ( that.position.left - op.left ) || 0
+			};
+
+			$( o.alsoResize ).each( function() {
+				var el = $( this ), start = $( this ).data( "ui-resizable-alsoresize" ), style = {},
+					css = el.parents( ui.originalElement[ 0 ] ).length ?
+							[ "width", "height" ] :
+							[ "width", "height", "top", "left" ];
+
+				$.each( css, function( i, prop ) {
+					var sum = ( start[ prop ] || 0 ) + ( delta[ prop ] || 0 );
+					if ( sum && sum >= 0 ) {
+						style[ prop ] = sum || null;
+					}
+				} );
+
+				el.css( style );
+			} );
+	},
+
+	stop: function() {
+		$( this ).removeData( "ui-resizable-alsoresize" );
+	}
+} );
+
+$.ui.plugin.add( "resizable", "ghost", {
+
+	start: function() {
+
+		var that = $( this ).resizable( "instance" ), cs = that.size;
+
+		that.ghost = that.originalElement.clone();
+		that.ghost.css( {
+			opacity: 0.25,
+			display: "block",
+			position: "relative",
+			height: cs.height,
+			width: cs.width,
+			margin: 0,
+			left: 0,
+			top: 0
+		} );
+
+		that._addClass( that.ghost, "ui-resizable-ghost" );
+
+		// DEPRECATED
+		// TODO: remove after 1.12
+		if ( $.uiBackCompat !== false && typeof that.options.ghost === "string" ) {
+
+			// Ghost option
+			that.ghost.addClass( this.options.ghost );
+		}
+
+		that.ghost.appendTo( that.helper );
+
+	},
+
+	resize: function() {
+		var that = $( this ).resizable( "instance" );
+		if ( that.ghost ) {
+			that.ghost.css( {
+				position: "relative",
+				height: that.size.height,
+				width: that.size.width
+			} );
+		}
+	},
+
+	stop: function() {
+		var that = $( this ).resizable( "instance" );
+		if ( that.ghost && that.helper ) {
+			that.helper.get( 0 ).removeChild( that.ghost.get( 0 ) );
+		}
+	}
+
+} );
+
+$.ui.plugin.add( "resizable", "grid", {
+
+	resize: function() {
+		var outerDimensions,
+			that = $( this ).resizable( "instance" ),
+			o = that.options,
+			cs = that.size,
+			os = that.originalSize,
+			op = that.originalPosition,
+			a = that.axis,
+			grid = typeof o.grid === "number" ? [ o.grid, o.grid ] : o.grid,
+			gridX = ( grid[ 0 ] || 1 ),
+			gridY = ( grid[ 1 ] || 1 ),
+			ox = Math.round( ( cs.width - os.width ) / gridX ) * gridX,
+			oy = Math.round( ( cs.height - os.height ) / gridY ) * gridY,
+			newWidth = os.width + ox,
+			newHeight = os.height + oy,
+			isMaxWidth = o.maxWidth && ( o.maxWidth < newWidth ),
+			isMaxHeight = o.maxHeight && ( o.maxHeight < newHeight ),
+			isMinWidth = o.minWidth && ( o.minWidth > newWidth ),
+			isMinHeight = o.minHeight && ( o.minHeight > newHeight );
+
+		o.grid = grid;
+
+		if ( isMinWidth ) {
+			newWidth += gridX;
+		}
+		if ( isMinHeight ) {
+			newHeight += gridY;
+		}
+		if ( isMaxWidth ) {
+			newWidth -= gridX;
+		}
+		if ( isMaxHeight ) {
+			newHeight -= gridY;
+		}
+
+		if ( /^(se|s|e)$/.test( a ) ) {
+			that.size.width = newWidth;
+			that.size.height = newHeight;
+		} else if ( /^(ne)$/.test( a ) ) {
+			that.size.width = newWidth;
+			that.size.height = newHeight;
+			that.position.top = op.top - oy;
+		} else if ( /^(sw)$/.test( a ) ) {
+			that.size.width = newWidth;
+			that.size.height = newHeight;
+			that.position.left = op.left - ox;
+		} else {
+			if ( newHeight - gridY <= 0 || newWidth - gridX <= 0 ) {
+				outerDimensions = that._getPaddingPlusBorderDimensions( this );
+			}
+
+			if ( newHeight - gridY > 0 ) {
+				that.size.height = newHeight;
+				that.position.top = op.top - oy;
+			} else {
+				newHeight = gridY - outerDimensions.height;
+				that.size.height = newHeight;
+				that.position.top = op.top + os.height - newHeight;
+			}
+			if ( newWidth - gridX > 0 ) {
+				that.size.width = newWidth;
+				that.position.left = op.left - ox;
+			} else {
+				newWidth = gridX - outerDimensions.width;
+				that.size.width = newWidth;
+				that.position.left = op.left + os.width - newWidth;
+			}
+		}
+	}
+
+} );
+
+var widgetsResizable = $.ui.resizable;
+
+
+/*!
+ * jQuery UI Dialog 1.12.0
+ * http://jqueryui.com
+ *
+ * Copyright jQuery Foundation and other contributors
+ * Released under the MIT license.
+ * http://jquery.org/license
+ */
+
+//>>label: Dialog
+//>>group: Widgets
+//>>description: Displays customizable dialog windows.
+//>>docs: http://api.jqueryui.com/dialog/
+//>>demos: http://jqueryui.com/dialog/
+//>>css.structure: ../../themes/base/core.css
+//>>css.structure: ../../themes/base/dialog.css
+//>>css.theme: ../../themes/base/theme.css
+
+
+
+$.widget( "ui.dialog", {
+	version: "1.12.0",
+	options: {
+		appendTo: "body",
+		autoOpen: true,
+		buttons: [],
+		classes: {
+			"ui-dialog": "ui-corner-all",
+			"ui-dialog-titlebar": "ui-corner-all"
+		},
+		closeOnEscape: true,
+		closeText: "Close",
+		draggable: true,
+		hide: null,
+		height: "auto",
+		maxHeight: null,
+		maxWidth: null,
+		minHeight: 150,
+		minWidth: 150,
+		modal: false,
+		position: {
+			my: "center",
+			at: "center",
+			of: window,
+			collision: "fit",
+
+			// Ensure the titlebar is always visible
+			using: function( pos ) {
+				var topOffset = $( this ).css( pos ).offset().top;
+				if ( topOffset < 0 ) {
+					$( this ).css( "top", pos.top - topOffset );
+				}
+			}
+		},
+		resizable: true,
+		show: null,
+		title: null,
+		width: 300,
+
+		// Callbacks
+		beforeClose: null,
+		close: null,
+		drag: null,
+		dragStart: null,
+		dragStop: null,
+		focus: null,
+		open: null,
+		resize: null,
+		resizeStart: null,
+		resizeStop: null
+	},
+
+	sizeRelatedOptions: {
+		buttons: true,
+		height: true,
+		maxHeight: true,
+		maxWidth: true,
+		minHeight: true,
+		minWidth: true,
+		width: true
+	},
+
+	resizableRelatedOptions: {
+		maxHeight: true,
+		maxWidth: true,
+		minHeight: true,
+		minWidth: true
+	},
+
+	_create: function() {
+		this.originalCss = {
+			display: this.element[ 0 ].style.display,
+			width: this.element[ 0 ].style.width,
+			minHeight: this.element[ 0 ].style.minHeight,
+			maxHeight: this.element[ 0 ].style.maxHeight,
+			height: this.element[ 0 ].style.height
+		};
+		this.originalPosition = {
+			parent: this.element.parent(),
+			index: this.element.parent().children().index( this.element )
+		};
+		this.originalTitle = this.element.attr( "title" );
+		if ( this.options.title == null && this.originalTitle != null ) {
+			this.options.title = this.originalTitle;
+		}
+
+		// Dialogs can't be disabled
+		if ( this.options.disabled ) {
+			this.options.disabled = false;
+		}
+
+		this._createWrapper();
+
+		this.element
+			.show()
+			.removeAttr( "title" )
+			.appendTo( this.uiDialog );
+
+		this._addClass( "ui-dialog-content", "ui-widget-content" );
+
+		this._createTitlebar();
+		this._createButtonPane();
+
+		if ( this.options.draggable && $.fn.draggable ) {
+			this._makeDraggable();
+		}
+		if ( this.options.resizable && $.fn.resizable ) {
+			this._makeResizable();
+		}
+
+		this._isOpen = false;
+
+		this._trackFocus();
+	},
+
+	_init: function() {
+		if ( this.options.autoOpen ) {
+			this.open();
+		}
+	},
+
+	_appendTo: function() {
+		var element = this.options.appendTo;
+		if ( element && ( element.jquery || element.nodeType ) ) {
+			return $( element );
+		}
+		return this.document.find( element || "body" ).eq( 0 );
+	},
+
+	_destroy: function() {
+		var next,
+			originalPosition = this.originalPosition;
+
+		this._untrackInstance();
+		this._destroyOverlay();
+
+		this.element
+			.removeUniqueId()
+			.css( this.originalCss )
+
+			// Without detaching first, the following becomes really slow
+			.detach();
+
+		this.uiDialog.remove();
+
+		if ( this.originalTitle ) {
+			this.element.attr( "title", this.originalTitle );
+		}
+
+		next = originalPosition.parent.children().eq( originalPosition.index );
+
+		// Don't try to place the dialog next to itself (#8613)
+		if ( next.length && next[ 0 ] !== this.element[ 0 ] ) {
+			next.before( this.element );
+		} else {
+			originalPosition.parent.append( this.element );
+		}
+	},
+
+	widget: function() {
+		return this.uiDialog;
+	},
+
+	disable: $.noop,
+	enable: $.noop,
+
+	close: function( event ) {
+		var that = this;
+
+		if ( !this._isOpen || this._trigger( "beforeClose", event ) === false ) {
+			return;
+		}
+
+		this._isOpen = false;
+		this._focusedElement = null;
+		this._destroyOverlay();
+		this._untrackInstance();
+
+		if ( !this.opener.filter( ":focusable" ).trigger( "focus" ).length ) {
+
+			// Hiding a focused element doesn't trigger blur in WebKit
+			// so in case we have nothing to focus on, explicitly blur the active element
+			// https://bugs.webkit.org/show_bug.cgi?id=47182
+			$.ui.safeBlur( $.ui.safeActiveElement( this.document[ 0 ] ) );
+		}
+
+		this._hide( this.uiDialog, this.options.hide, function() {
+			that._trigger( "close", event );
+		} );
+	},
+
+	isOpen: function() {
+		return this._isOpen;
+	},
+
+	moveToTop: function() {
+		this._moveToTop();
+	},
+
+	_moveToTop: function( event, silent ) {
+		var moved = false,
+			zIndices = this.uiDialog.siblings( ".ui-front:visible" ).map( function() {
+				return +$( this ).css( "z-index" );
+			} ).get(),
+			zIndexMax = Math.max.apply( null, zIndices );
+
+		if ( zIndexMax >= +this.uiDialog.css( "z-index" ) ) {
+			this.uiDialog.css( "z-index", zIndexMax + 1 );
+			moved = true;
+		}
+
+		if ( moved && !silent ) {
+			this._trigger( "focus", event );
+		}
+		return moved;
+	},
+
+	open: function() {
+		var that = this;
+		if ( this._isOpen ) {
+			if ( this._moveToTop() ) {
+				this._focusTabbable();
+			}
+			return;
+		}
+
+		this._isOpen = true;
+		this.opener = $( $.ui.safeActiveElement( this.document[ 0 ] ) );
+
+		this._size();
+		this._position();
+		this._createOverlay();
+		this._moveToTop( null, true );
+
+		// Ensure the overlay is moved to the top with the dialog, but only when
+		// opening. The overlay shouldn't move after the dialog is open so that
+		// modeless dialogs opened after the modal dialog stack properly.
+		if ( this.overlay ) {
+			this.overlay.css( "z-index", this.uiDialog.css( "z-index" ) - 1 );
+		}
+
+		this._show( this.uiDialog, this.options.show, function() {
+			that._focusTabbable();
+			that._trigger( "focus" );
+		} );
+
+		// Track the dialog immediately upon openening in case a focus event
+		// somehow occurs outside of the dialog before an element inside the
+		// dialog is focused (#10152)
+		this._makeFocusTarget();
+
+		this._trigger( "open" );
+	},
+
+	_focusTabbable: function() {
+
+		// Set focus to the first match:
+		// 1. An element that was focused previously
+		// 2. First element inside the dialog matching [autofocus]
+		// 3. Tabbable element inside the content element
+		// 4. Tabbable element inside the buttonpane
+		// 5. The close button
+		// 6. The dialog itself
+		var hasFocus = this._focusedElement;
+		if ( !hasFocus ) {
+			hasFocus = this.element.find( "[autofocus]" );
+		}
+		if ( !hasFocus.length ) {
+			hasFocus = this.element.find( ":tabbable" );
+		}
+		if ( !hasFocus.length ) {
+			hasFocus = this.uiDialogButtonPane.find( ":tabbable" );
+		}
+		if ( !hasFocus.length ) {
+			hasFocus = this.uiDialogTitlebarClose.filter( ":tabbable" );
+		}
+		if ( !hasFocus.length ) {
+			hasFocus = this.uiDialog;
+		}
+		hasFocus.eq( 0 ).trigger( "focus" );
+	},
+
+	_keepFocus: function( event ) {
+		function checkFocus() {
+			var activeElement = $.ui.safeActiveElement( this.document[ 0 ] ),
+				isActive = this.uiDialog[ 0 ] === activeElement ||
+					$.contains( this.uiDialog[ 0 ], activeElement );
+			if ( !isActive ) {
+				this._focusTabbable();
+			}
+		}
+		event.preventDefault();
+		checkFocus.call( this );
+
+		// support: IE
+		// IE <= 8 doesn't prevent moving focus even with event.preventDefault()
+		// so we check again later
+		this._delay( checkFocus );
+	},
+
+	_createWrapper: function() {
+		this.uiDialog = $( "<div>" )
+			.hide()
+			.attr( {
+
+				// Setting tabIndex makes the div focusable
+				tabIndex: -1,
+				role: "dialog"
+			} )
+			.appendTo( this._appendTo() );
+
+		this._addClass( this.uiDialog, "ui-dialog", "ui-widget ui-widget-content ui-front" );
+		this._on( this.uiDialog, {
+			keydown: function( event ) {
+				if ( this.options.closeOnEscape && !event.isDefaultPrevented() && event.keyCode &&
+						event.keyCode === $.ui.keyCode.ESCAPE ) {
+					event.preventDefault();
+					this.close( event );
+					return;
+				}
+
+				// Prevent tabbing out of dialogs
+				if ( event.keyCode !== $.ui.keyCode.TAB || event.isDefaultPrevented() ) {
+					return;
+				}
+				var tabbables = this.uiDialog.find( ":tabbable" ),
+					first = tabbables.filter( ":first" ),
+					last = tabbables.filter( ":last" );
+
+				if ( ( event.target === last[ 0 ] || event.target === this.uiDialog[ 0 ] ) &&
+						!event.shiftKey ) {
+					this._delay( function() {
+						first.trigger( "focus" );
+					} );
+					event.preventDefault();
+				} else if ( ( event.target === first[ 0 ] ||
+						event.target === this.uiDialog[ 0 ] ) && event.shiftKey ) {
+					this._delay( function() {
+						last.trigger( "focus" );
+					} );
+					event.preventDefault();
+				}
+			},
+			mousedown: function( event ) {
+				if ( this._moveToTop( event ) ) {
+					this._focusTabbable();
+				}
+			}
+		} );
+
+		// We assume that any existing aria-describedby attribute means
+		// that the dialog content is marked up properly
+		// otherwise we brute force the content as the description
+		if ( !this.element.find( "[aria-describedby]" ).length ) {
+			this.uiDialog.attr( {
+				"aria-describedby": this.element.uniqueId().attr( "id" )
+			} );
+		}
+	},
+
+	_createTitlebar: function() {
+		var uiDialogTitle;
+
+		this.uiDialogTitlebar = $( "<div>" );
+		this._addClass( this.uiDialogTitlebar,
+			"ui-dialog-titlebar", "ui-widget-header ui-helper-clearfix" );
+		this._on( this.uiDialogTitlebar, {
+			mousedown: function( event ) {
+
+				// Don't prevent click on close button (#8838)
+				// Focusing a dialog that is partially scrolled out of view
+				// causes the browser to scroll it into view, preventing the click event
+				if ( !$( event.target ).closest( ".ui-dialog-titlebar-close" ) ) {
+
+					// Dialog isn't getting focus when dragging (#8063)
+					this.uiDialog.trigger( "focus" );
+				}
+			}
+		} );
+
+		// Support: IE
+		// Use type="button" to prevent enter keypresses in textboxes from closing the
+		// dialog in IE (#9312)
+		this.uiDialogTitlebarClose = $( "<button type='button'></button>" )
+			.button( {
+				label: $( "<a>" ).text( this.options.closeText ).html(),
+				icon: "ui-icon-closethick",
+				showLabel: false
+			} )
+			.appendTo( this.uiDialogTitlebar );
+
+		this._addClass( this.uiDialogTitlebarClose, "ui-dialog-titlebar-close" );
+		this._on( this.uiDialogTitlebarClose, {
+			click: function( event ) {
+				event.preventDefault();
+				this.close( event );
+			}
+		} );
+
+		uiDialogTitle = $( "<span>" ).uniqueId().prependTo( this.uiDialogTitlebar );
+		this._addClass( uiDialogTitle, "ui-dialog-title" );
+		this._title( uiDialogTitle );
+
+		this.uiDialogTitlebar.prependTo( this.uiDialog );
+
+		this.uiDialog.attr( {
+			"aria-labelledby": uiDialogTitle.attr( "id" )
+		} );
+	},
+
+	_title: function( title ) {
+		if ( this.options.title ) {
+			title.text( this.options.title );
+		} else {
+			title.html( "&#160;" );
+		}
+	},
+
+	_createButtonPane: function() {
+		this.uiDialogButtonPane = $( "<div>" );
+		this._addClass( this.uiDialogButtonPane, "ui-dialog-buttonpane",
+			"ui-widget-content ui-helper-clearfix" );
+
+		this.uiButtonSet = $( "<div>" )
+			.appendTo( this.uiDialogButtonPane );
+		this._addClass( this.uiButtonSet, "ui-dialog-buttonset" );
+
+		this._createButtons();
+	},
+
+	_createButtons: function() {
+		var that = this,
+			buttons = this.options.buttons;
+
+		// If we already have a button pane, remove it
+		this.uiDialogButtonPane.remove();
+		this.uiButtonSet.empty();
+
+		if ( $.isEmptyObject( buttons ) || ( $.isArray( buttons ) && !buttons.length ) ) {
+			this._removeClass( this.uiDialog, "ui-dialog-buttons" );
+			return;
+		}
+
+		$.each( buttons, function( name, props ) {
+			var click, buttonOptions;
+			props = $.isFunction( props ) ?
+				{ click: props, text: name } :
+				props;
+
+			// Default to a non-submitting button
+			props = $.extend( { type: "button" }, props );
+
+			// Change the context for the click callback to be the main element
+			click = props.click;
+			buttonOptions = {
+				icon: props.icon,
+				iconPosition: props.iconPosition,
+				showLabel: props.showLabel
+			};
+
+			delete props.click;
+			delete props.icon;
+			delete props.iconPosition;
+			delete props.showLabel;
+
+			$( "<button></button>", props )
+				.button( buttonOptions )
+				.appendTo( that.uiButtonSet )
+				.on( "click", function() {
+					click.apply( that.element[ 0 ], arguments );
+				} );
+		} );
+		this._addClass( this.uiDialog, "ui-dialog-buttons" );
+		this.uiDialogButtonPane.appendTo( this.uiDialog );
+	},
+
+	_makeDraggable: function() {
+		var that = this,
+			options = this.options;
+
+		function filteredUi( ui ) {
+			return {
+				position: ui.position,
+				offset: ui.offset
+			};
+		}
+
+		this.uiDialog.draggable( {
+			cancel: ".ui-dialog-content, .ui-dialog-titlebar-close",
+			handle: ".ui-dialog-titlebar",
+			containment: "document",
+			start: function( event, ui ) {
+				that._addClass( $( this ), "ui-dialog-dragging" );
+				that._blockFrames();
+				that._trigger( "dragStart", event, filteredUi( ui ) );
+			},
+			drag: function( event, ui ) {
+				that._trigger( "drag", event, filteredUi( ui ) );
+			},
+			stop: function( event, ui ) {
+				var left = ui.offset.left - that.document.scrollLeft(),
+					top = ui.offset.top - that.document.scrollTop();
+
+				options.position = {
+					my: "left top",
+					at: "left" + ( left >= 0 ? "+" : "" ) + left + " " +
+						"top" + ( top >= 0 ? "+" : "" ) + top,
+					of: that.window
+				};
+				that._removeClass( $( this ), "ui-dialog-dragging" );
+				that._unblockFrames();
+				that._trigger( "dragStop", event, filteredUi( ui ) );
+			}
+		} );
+	},
+
+	_makeResizable: function() {
+		var that = this,
+			options = this.options,
+			handles = options.resizable,
+
+			// .ui-resizable has position: relative defined in the stylesheet
+			// but dialogs have to use absolute or fixed positioning
+			position = this.uiDialog.css( "position" ),
+			resizeHandles = typeof handles === "string" ?
+				handles :
+				"n,e,s,w,se,sw,ne,nw";
+
+		function filteredUi( ui ) {
+			return {
+				originalPosition: ui.originalPosition,
+				originalSize: ui.originalSize,
+				position: ui.position,
+				size: ui.size
+			};
+		}
+
+		this.uiDialog.resizable( {
+			cancel: ".ui-dialog-content",
+			containment: "document",
+			alsoResize: this.element,
+			maxWidth: options.maxWidth,
+			maxHeight: options.maxHeight,
+			minWidth: options.minWidth,
+			minHeight: this._minHeight(),
+			handles: resizeHandles,
+			start: function( event, ui ) {
+				that._addClass( $( this ), "ui-dialog-resizing" );
+				that._blockFrames();
+				that._trigger( "resizeStart", event, filteredUi( ui ) );
+			},
+			resize: function( event, ui ) {
+				that._trigger( "resize", event, filteredUi( ui ) );
+			},
+			stop: function( event, ui ) {
+				var offset = that.uiDialog.offset(),
+					left = offset.left - that.document.scrollLeft(),
+					top = offset.top - that.document.scrollTop();
+
+				options.height = that.uiDialog.height();
+				options.width = that.uiDialog.width();
+				options.position = {
+					my: "left top",
+					at: "left" + ( left >= 0 ? "+" : "" ) + left + " " +
+						"top" + ( top >= 0 ? "+" : "" ) + top,
+					of: that.window
+				};
+				that._removeClass( $( this ), "ui-dialog-resizing" );
+				that._unblockFrames();
+				that._trigger( "resizeStop", event, filteredUi( ui ) );
+			}
+		} )
+			.css( "position", position );
+	},
+
+	_trackFocus: function() {
+		this._on( this.widget(), {
+			focusin: function( event ) {
+				this._makeFocusTarget();
+				this._focusedElement = $( event.target );
+			}
+		} );
+	},
+
+	_makeFocusTarget: function() {
+		this._untrackInstance();
+		this._trackingInstances().unshift( this );
+	},
+
+	_untrackInstance: function() {
+		var instances = this._trackingInstances(),
+			exists = $.inArray( this, instances );
+		if ( exists !== -1 ) {
+			instances.splice( exists, 1 );
+		}
+	},
+
+	_trackingInstances: function() {
+		var instances = this.document.data( "ui-dialog-instances" );
+		if ( !instances ) {
+			instances = [];
+			this.document.data( "ui-dialog-instances", instances );
+		}
+		return instances;
+	},
+
+	_minHeight: function() {
+		var options = this.options;
+
+		return options.height === "auto" ?
+			options.minHeight :
+			Math.min( options.minHeight, options.height );
+	},
+
+	_position: function() {
+
+		// Need to show the dialog to get the actual offset in the position plugin
+		var isVisible = this.uiDialog.is( ":visible" );
+		if ( !isVisible ) {
+			this.uiDialog.show();
+		}
+		this.uiDialog.position( this.options.position );
+		if ( !isVisible ) {
+			this.uiDialog.hide();
+		}
+	},
+
+	_setOptions: function( options ) {
+		var that = this,
+			resize = false,
+			resizableOptions = {};
+
+		$.each( options, function( key, value ) {
+			that._setOption( key, value );
+
+			if ( key in that.sizeRelatedOptions ) {
+				resize = true;
+			}
+			if ( key in that.resizableRelatedOptions ) {
+				resizableOptions[ key ] = value;
+			}
+		} );
+
+		if ( resize ) {
+			this._size();
+			this._position();
+		}
+		if ( this.uiDialog.is( ":data(ui-resizable)" ) ) {
+			this.uiDialog.resizable( "option", resizableOptions );
+		}
+	},
+
+	_setOption: function( key, value ) {
+		var isDraggable, isResizable,
+			uiDialog = this.uiDialog;
+
+		if ( key === "disabled" ) {
+			return;
+		}
+
+		this._super( key, value );
+
+		if ( key === "appendTo" ) {
+			this.uiDialog.appendTo( this._appendTo() );
+		}
+
+		if ( key === "buttons" ) {
+			this._createButtons();
+		}
+
+		if ( key === "closeText" ) {
+			this.uiDialogTitlebarClose.button( {
+
+				// Ensure that we always pass a string
+				label: $( "<a>" ).text( "" + this.options.closeText ).html()
+			} );
+		}
+
+		if ( key === "draggable" ) {
+			isDraggable = uiDialog.is( ":data(ui-draggable)" );
+			if ( isDraggable && !value ) {
+				uiDialog.draggable( "destroy" );
+			}
+
+			if ( !isDraggable && value ) {
+				this._makeDraggable();
+			}
+		}
+
+		if ( key === "position" ) {
+			this._position();
+		}
+
+		if ( key === "resizable" ) {
+
+			// currently resizable, becoming non-resizable
+			isResizable = uiDialog.is( ":data(ui-resizable)" );
+			if ( isResizable && !value ) {
+				uiDialog.resizable( "destroy" );
+			}
+
+			// Currently resizable, changing handles
+			if ( isResizable && typeof value === "string" ) {
+				uiDialog.resizable( "option", "handles", value );
+			}
+
+			// Currently non-resizable, becoming resizable
+			if ( !isResizable && value !== false ) {
+				this._makeResizable();
+			}
+		}
+
+		if ( key === "title" ) {
+			this._title( this.uiDialogTitlebar.find( ".ui-dialog-title" ) );
+		}
+	},
+
+	_size: function() {
+
+		// If the user has resized the dialog, the .ui-dialog and .ui-dialog-content
+		// divs will both have width and height set, so we need to reset them
+		var nonContentHeight, minContentHeight, maxContentHeight,
+			options = this.options;
+
+		// Reset content sizing
+		this.element.show().css( {
+			width: "auto",
+			minHeight: 0,
+			maxHeight: "none",
+			height: 0
+		} );
+
+		if ( options.minWidth > options.width ) {
+			options.width = options.minWidth;
+		}
+
+		// Reset wrapper sizing
+		// determine the height of all the non-content elements
+		nonContentHeight = this.uiDialog.css( {
+			height: "auto",
+			width: options.width
+		} )
+			.outerHeight();
+		minContentHeight = Math.max( 0, options.minHeight - nonContentHeight );
+		maxContentHeight = typeof options.maxHeight === "number" ?
+			Math.max( 0, options.maxHeight - nonContentHeight ) :
+			"none";
+
+		if ( options.height === "auto" ) {
+			this.element.css( {
+				minHeight: minContentHeight,
+				maxHeight: maxContentHeight,
+				height: "auto"
+			} );
+		} else {
+			this.element.height( Math.max( 0, options.height - nonContentHeight ) );
+		}
+
+		if ( this.uiDialog.is( ":data(ui-resizable)" ) ) {
+			this.uiDialog.resizable( "option", "minHeight", this._minHeight() );
+		}
+	},
+
+	_blockFrames: function() {
+		this.iframeBlocks = this.document.find( "iframe" ).map( function() {
+			var iframe = $( this );
+
+			return $( "<div>" )
+				.css( {
+					position: "absolute",
+					width: iframe.outerWidth(),
+					height: iframe.outerHeight()
+				} )
+				.appendTo( iframe.parent() )
+				.offset( iframe.offset() )[ 0 ];
+		} );
+	},
+
+	_unblockFrames: function() {
+		if ( this.iframeBlocks ) {
+			this.iframeBlocks.remove();
+			delete this.iframeBlocks;
+		}
+	},
+
+	_allowInteraction: function( event ) {
+		if ( $( event.target ).closest( ".ui-dialog" ).length ) {
+			return true;
+		}
+
+		// TODO: Remove hack when datepicker implements
+		// the .ui-front logic (#8989)
+		return !!$( event.target ).closest( ".ui-datepicker" ).length;
+	},
+
+	_createOverlay: function() {
+		if ( !this.options.modal ) {
+			return;
+		}
+
+		// We use a delay in case the overlay is created from an
+		// event that we're going to be cancelling (#2804)
+		var isOpening = true;
+		this._delay( function() {
+			isOpening = false;
+		} );
+
+		if ( !this.document.data( "ui-dialog-overlays" ) ) {
+
+			// Prevent use of anchors and inputs
+			// Using _on() for an event handler shared across many instances is
+			// safe because the dialogs stack and must be closed in reverse order
+			this._on( this.document, {
+				focusin: function( event ) {
+					if ( isOpening ) {
+						return;
+					}
+
+					if ( !this._allowInteraction( event ) ) {
+						event.preventDefault();
+						this._trackingInstances()[ 0 ]._focusTabbable();
+					}
+				}
+			} );
+		}
+
+		this.overlay = $( "<div>" )
+			.appendTo( this._appendTo() );
+
+		this._addClass( this.overlay, null, "ui-widget-overlay ui-front" );
+		this._on( this.overlay, {
+			mousedown: "_keepFocus"
+		} );
+		this.document.data( "ui-dialog-overlays",
+			( this.document.data( "ui-dialog-overlays" ) || 0 ) + 1 );
+	},
+
+	_destroyOverlay: function() {
+		if ( !this.options.modal ) {
+			return;
+		}
+
+		if ( this.overlay ) {
+			var overlays = this.document.data( "ui-dialog-overlays" ) - 1;
+
+			if ( !overlays ) {
+				this._off( this.document, "focusin" );
+				this.document.removeData( "ui-dialog-overlays" );
+			} else {
+				this.document.data( "ui-dialog-overlays", overlays );
+			}
+
+			this.overlay.remove();
+			this.overlay = null;
+		}
+	}
+} );
+
+// DEPRECATED
+// TODO: switch return back to widget declaration at top of file when this is removed
+if ( $.uiBackCompat !== false ) {
+
+	// Backcompat for dialogClass option
+	$.widget( "ui.dialog", $.ui.dialog, {
+		options: {
+			dialogClass: ""
+		},
+		_createWrapper: function() {
+			this._super();
+			this.uiDialog.addClass( this.options.dialogClass );
+		},
+		_setOption: function( key, value ) {
+			if ( key === "dialogClass" ) {
+				this.uiDialog
+					.removeClass( this.options.dialogClass )
+					.addClass( value );
+			}
+			this._superApply( arguments );
+		}
+	} );
+}
+
+var widgetsDialog = $.ui.dialog;
+
+
+/*!
+ * jQuery UI Droppable 1.12.0
+ * http://jqueryui.com
+ *
+ * Copyright jQuery Foundation and other contributors
+ * Released under the MIT license.
+ * http://jquery.org/license
+ */
+
+//>>label: Droppable
+//>>group: Interactions
+//>>description: Enables drop targets for draggable elements.
+//>>docs: http://api.jqueryui.com/droppable/
+//>>demos: http://jqueryui.com/droppable/
+
+
+
+$.widget( "ui.droppable", {
+	version: "1.12.0",
+	widgetEventPrefix: "drop",
+	options: {
+		accept: "*",
+		addClasses: true,
+		greedy: false,
+		scope: "default",
+		tolerance: "intersect",
+
+		// Callbacks
+		activate: null,
+		deactivate: null,
+		drop: null,
+		out: null,
+		over: null
+	},
+	_create: function() {
+
+		var proportions,
+			o = this.options,
+			accept = o.accept;
+
+		this.isover = false;
+		this.isout = true;
+
+		this.accept = $.isFunction( accept ) ? accept : function( d ) {
+			return d.is( accept );
+		};
+
+		this.proportions = function( /* valueToWrite */ ) {
+			if ( arguments.length ) {
+
+				// Store the droppable's proportions
+				proportions = arguments[ 0 ];
+			} else {
+
+				// Retrieve or derive the droppable's proportions
+				return proportions ?
+					proportions :
+					proportions = {
+						width: this.element[ 0 ].offsetWidth,
+						height: this.element[ 0 ].offsetHeight
+					};
+			}
+		};
+
+		this._addToManager( o.scope );
+
+		o.addClasses && this._addClass( "ui-droppable" );
+
+	},
+
+	_addToManager: function( scope ) {
+
+		// Add the reference and positions to the manager
+		$.ui.ddmanager.droppables[ scope ] = $.ui.ddmanager.droppables[ scope ] || [];
+		$.ui.ddmanager.droppables[ scope ].push( this );
+	},
+
+	_splice: function( drop ) {
+		var i = 0;
+		for ( ; i < drop.length; i++ ) {
+			if ( drop[ i ] === this ) {
+				drop.splice( i, 1 );
+			}
+		}
+	},
+
+	_destroy: function() {
+		var drop = $.ui.ddmanager.droppables[ this.options.scope ];
+
+		this._splice( drop );
+	},
+
+	_setOption: function( key, value ) {
+
+		if ( key === "accept" ) {
+			this.accept = $.isFunction( value ) ? value : function( d ) {
+				return d.is( value );
+			};
+		} else if ( key === "scope" ) {
+			var drop = $.ui.ddmanager.droppables[ this.options.scope ];
+
+			this._splice( drop );
+			this._addToManager( value );
+		}
+
+		this._super( key, value );
+	},
+
+	_activate: function( event ) {
+		var draggable = $.ui.ddmanager.current;
+
+		this._addActiveClass();
+		if ( draggable ) {
+			this._trigger( "activate", event, this.ui( draggable ) );
+		}
+	},
+
+	_deactivate: function( event ) {
+		var draggable = $.ui.ddmanager.current;
+
+		this._removeActiveClass();
+		if ( draggable ) {
+			this._trigger( "deactivate", event, this.ui( draggable ) );
+		}
+	},
+
+	_over: function( event ) {
+
+		var draggable = $.ui.ddmanager.current;
+
+		// Bail if draggable and droppable are same element
+		if ( !draggable || ( draggable.currentItem ||
+				draggable.element )[ 0 ] === this.element[ 0 ] ) {
+			return;
+		}
+
+		if ( this.accept.call( this.element[ 0 ], ( draggable.currentItem ||
+				draggable.element ) ) ) {
+			this._addHoverClass();
+			this._trigger( "over", event, this.ui( draggable ) );
+		}
+
+	},
+
+	_out: function( event ) {
+
+		var draggable = $.ui.ddmanager.current;
+
+		// Bail if draggable and droppable are same element
+		if ( !draggable || ( draggable.currentItem ||
+				draggable.element )[ 0 ] === this.element[ 0 ] ) {
+			return;
+		}
+
+		if ( this.accept.call( this.element[ 0 ], ( draggable.currentItem ||
+				draggable.element ) ) ) {
+			this._removeHoverClass();
+			this._trigger( "out", event, this.ui( draggable ) );
+		}
+
+	},
+
+	_drop: function( event, custom ) {
+
+		var draggable = custom || $.ui.ddmanager.current,
+			childrenIntersection = false;
+
+		// Bail if draggable and droppable are same element
+		if ( !draggable || ( draggable.currentItem ||
+				draggable.element )[ 0 ] === this.element[ 0 ] ) {
+			return false;
+		}
+
+		this.element
+			.find( ":data(ui-droppable)" )
+			.not( ".ui-draggable-dragging" )
+			.each( function() {
+				var inst = $( this ).droppable( "instance" );
+				if (
+					inst.options.greedy &&
+					!inst.options.disabled &&
+					inst.options.scope === draggable.options.scope &&
+					inst.accept.call(
+						inst.element[ 0 ], ( draggable.currentItem || draggable.element )
+					) &&
+					intersect(
+						draggable,
+						$.extend( inst, { offset: inst.element.offset() } ),
+						inst.options.tolerance, event
+					)
+				) {
+					childrenIntersection = true;
+					return false; }
+			} );
+		if ( childrenIntersection ) {
+			return false;
+		}
+
+		if ( this.accept.call( this.element[ 0 ],
+				( draggable.currentItem || draggable.element ) ) ) {
+			this._removeActiveClass();
+			this._removeHoverClass();
+
+			this._trigger( "drop", event, this.ui( draggable ) );
+			return this.element;
+		}
+
+		return false;
+
+	},
+
+	ui: function( c ) {
+		return {
+			draggable: ( c.currentItem || c.element ),
+			helper: c.helper,
+			position: c.position,
+			offset: c.positionAbs
+		};
+	},
+
+	// Extension points just to make backcompat sane and avoid duplicating logic
+	// TODO: Remove in 1.13 along with call to it below
+	_addHoverClass: function() {
+		this._addClass( "ui-droppable-hover" );
+	},
+
+	_removeHoverClass: function() {
+		this._removeClass( "ui-droppable-hover" );
+	},
+
+	_addActiveClass: function() {
+		this._addClass( "ui-droppable-active" );
+	},
+
+	_removeActiveClass: function() {
+		this._removeClass( "ui-droppable-active" );
+	}
+} );
+
+var intersect = $.ui.intersect = ( function() {
+	function isOverAxis( x, reference, size ) {
+		return ( x >= reference ) && ( x < ( reference + size ) );
+	}
+
+	return function( draggable, droppable, toleranceMode, event ) {
+
+		if ( !droppable.offset ) {
+			return false;
+		}
+
+		var x1 = ( draggable.positionAbs ||
+				draggable.position.absolute ).left + draggable.margins.left,
+			y1 = ( draggable.positionAbs ||
+				draggable.position.absolute ).top + draggable.margins.top,
+			x2 = x1 + draggable.helperProportions.width,
+			y2 = y1 + draggable.helperProportions.height,
+			l = droppable.offset.left,
+			t = droppable.offset.top,
+			r = l + droppable.proportions().width,
+			b = t + droppable.proportions().height;
+
+		switch ( toleranceMode ) {
+		case "fit":
+			return ( l <= x1 && x2 <= r && t <= y1 && y2 <= b );
+		case "intersect":
+			return ( l < x1 + ( draggable.helperProportions.width / 2 ) && // Right Half
+				x2 - ( draggable.helperProportions.width / 2 ) < r && // Left Half
+				t < y1 + ( draggable.helperProportions.height / 2 ) && // Bottom Half
+				y2 - ( draggable.helperProportions.height / 2 ) < b ); // Top Half
+		case "pointer":
+			return isOverAxis( event.pageY, t, droppable.proportions().height ) &&
+				isOverAxis( event.pageX, l, droppable.proportions().width );
+		case "touch":
+			return (
+				( y1 >= t && y1 <= b ) || // Top edge touching
+				( y2 >= t && y2 <= b ) || // Bottom edge touching
+				( y1 < t && y2 > b ) // Surrounded vertically
+			) && (
+				( x1 >= l && x1 <= r ) || // Left edge touching
+				( x2 >= l && x2 <= r ) || // Right edge touching
+				( x1 < l && x2 > r ) // Surrounded horizontally
+			);
+		default:
+			return false;
+		}
+	};
+} )();
+
+/*
+	This manager tracks offsets of draggables and droppables
+*/
+$.ui.ddmanager = {
+	current: null,
+	droppables: { "default": [] },
+	prepareOffsets: function( t, event ) {
+
+		var i, j,
+			m = $.ui.ddmanager.droppables[ t.options.scope ] || [],
+			type = event ? event.type : null, // workaround for #2317
+			list = ( t.currentItem || t.element ).find( ":data(ui-droppable)" ).addBack();
+
+		droppablesLoop: for ( i = 0; i < m.length; i++ ) {
+
+			// No disabled and non-accepted
+			if ( m[ i ].options.disabled || ( t && !m[ i ].accept.call( m[ i ].element[ 0 ],
+					( t.currentItem || t.element ) ) ) ) {
+				continue;
+			}
+
+			// Filter out elements in the current dragged item
+			for ( j = 0; j < list.length; j++ ) {
+				if ( list[ j ] === m[ i ].element[ 0 ] ) {
+					m[ i ].proportions().height = 0;
+					continue droppablesLoop;
+				}
+			}
+
+			m[ i ].visible = m[ i ].element.css( "display" ) !== "none";
+			if ( !m[ i ].visible ) {
+				continue;
+			}
+
+			// Activate the droppable if used directly from draggables
+			if ( type === "mousedown" ) {
+				m[ i ]._activate.call( m[ i ], event );
+			}
+
+			m[ i ].offset = m[ i ].element.offset();
+			m[ i ].proportions( {
+				width: m[ i ].element[ 0 ].offsetWidth,
+				height: m[ i ].element[ 0 ].offsetHeight
+			} );
+
+		}
+
+	},
+	drop: function( draggable, event ) {
+
+		var dropped = false;
+
+		// Create a copy of the droppables in case the list changes during the drop (#9116)
+		$.each( ( $.ui.ddmanager.droppables[ draggable.options.scope ] || [] ).slice(), function() {
+
+			if ( !this.options ) {
+				return;
+			}
+			if ( !this.options.disabled && this.visible &&
+					intersect( draggable, this, this.options.tolerance, event ) ) {
+				dropped = this._drop.call( this, event ) || dropped;
+			}
+
+			if ( !this.options.disabled && this.visible && this.accept.call( this.element[ 0 ],
+					( draggable.currentItem || draggable.element ) ) ) {
+				this.isout = true;
+				this.isover = false;
+				this._deactivate.call( this, event );
+			}
+
+		} );
+		return dropped;
+
+	},
+	dragStart: function( draggable, event ) {
+
+		// Listen for scrolling so that if the dragging causes scrolling the position of the
+		// droppables can be recalculated (see #5003)
+		draggable.element.parentsUntil( "body" ).on( "scroll.droppable", function() {
+			if ( !draggable.options.refreshPositions ) {
+				$.ui.ddmanager.prepareOffsets( draggable, event );
+			}
+		} );
+	},
+	drag: function( draggable, event ) {
+
+		// If you have a highly dynamic page, you might try this option. It renders positions
+		// every time you move the mouse.
+		if ( draggable.options.refreshPositions ) {
+			$.ui.ddmanager.prepareOffsets( draggable, event );
+		}
+
+		// Run through all droppables and check their positions based on specific tolerance options
+		$.each( $.ui.ddmanager.droppables[ draggable.options.scope ] || [], function() {
+
+			if ( this.options.disabled || this.greedyChild || !this.visible ) {
+				return;
+			}
+
+			var parentInstance, scope, parent,
+				intersects = intersect( draggable, this, this.options.tolerance, event ),
+				c = !intersects && this.isover ?
+					"isout" :
+					( intersects && !this.isover ? "isover" : null );
+			if ( !c ) {
+				return;
+			}
+
+			if ( this.options.greedy ) {
+
+				// find droppable parents with same scope
+				scope = this.options.scope;
+				parent = this.element.parents( ":data(ui-droppable)" ).filter( function() {
+					return $( this ).droppable( "instance" ).options.scope === scope;
+				} );
+
+				if ( parent.length ) {
+					parentInstance = $( parent[ 0 ] ).droppable( "instance" );
+					parentInstance.greedyChild = ( c === "isover" );
+				}
+			}
+
+			// We just moved into a greedy child
+			if ( parentInstance && c === "isover" ) {
+				parentInstance.isover = false;
+				parentInstance.isout = true;
+				parentInstance._out.call( parentInstance, event );
+			}
+
+			this[ c ] = true;
+			this[ c === "isout" ? "isover" : "isout" ] = false;
+			this[ c === "isover" ? "_over" : "_out" ].call( this, event );
+
+			// We just moved out of a greedy child
+			if ( parentInstance && c === "isout" ) {
+				parentInstance.isout = false;
+				parentInstance.isover = true;
+				parentInstance._over.call( parentInstance, event );
+			}
+		} );
+
+	},
+	dragStop: function( draggable, event ) {
+		draggable.element.parentsUntil( "body" ).off( "scroll.droppable" );
+
+		// Call prepareOffsets one final time since IE does not fire return scroll events when
+		// overflow was caused by drag (see #5003)
+		if ( !draggable.options.refreshPositions ) {
+			$.ui.ddmanager.prepareOffsets( draggable, event );
+		}
+	}
+};
+
+// DEPRECATED
+// TODO: switch return back to widget declaration at top of file when this is removed
+if ( $.uiBackCompat !== false ) {
+
+	// Backcompat for activeClass and hoverClass options
+	$.widget( "ui.droppable", $.ui.droppable, {
+		options: {
+			hoverClass: false,
+			activeClass: false
+		},
+		_addActiveClass: function() {
+			this._super();
+			if ( this.options.activeClass ) {
+				this.element.addClass( this.options.activeClass );
+			}
+		},
+		_removeActiveClass: function() {
+			this._super();
+			if ( this.options.activeClass ) {
+				this.element.removeClass( this.options.activeClass );
+			}
+		},
+		_addHoverClass: function() {
+			this._super();
+			if ( this.options.hoverClass ) {
+				this.element.addClass( this.options.hoverClass );
+			}
+		},
+		_removeHoverClass: function() {
+			this._super();
+			if ( this.options.hoverClass ) {
+				this.element.removeClass( this.options.hoverClass );
+			}
+		}
+	} );
+}
+
+var widgetsDroppable = $.ui.droppable;
+
+
+/*!
+ * jQuery UI Progressbar 1.12.0
+ * http://jqueryui.com
+ *
+ * Copyright jQuery Foundation and other contributors
+ * Released under the MIT license.
+ * http://jquery.org/license
+ */
+
+//>>label: Progressbar
+//>>group: Widgets
+// jscs:disable maximumLineLength
+//>>description: Displays a status indicator for loading state, standard percentage, and other progress indicators.
+// jscs:enable maximumLineLength
+//>>docs: http://api.jqueryui.com/progressbar/
+//>>demos: http://jqueryui.com/progressbar/
+//>>css.structure: ../../themes/base/core.css
+//>>css.structure: ../../themes/base/progressbar.css
+//>>css.theme: ../../themes/base/theme.css
+
+
+
+var widgetsProgressbar = $.widget( "ui.progressbar", {
+	version: "1.12.0",
+	options: {
+		classes: {
+			"ui-progressbar": "ui-corner-all",
+			"ui-progressbar-value": "ui-corner-left",
+			"ui-progressbar-complete": "ui-corner-right"
+		},
+		max: 100,
+		value: 0,
+
+		change: null,
+		complete: null
+	},
+
+	min: 0,
+
+	_create: function() {
+
+		// Constrain initial value
+		this.oldValue = this.options.value = this._constrainedValue();
+
+		this.element.attr( {
+
+			// Only set static values; aria-valuenow and aria-valuemax are
+			// set inside _refreshValue()
+			role: "progressbar",
+			"aria-valuemin": this.min
+		} );
+		this._addClass( "ui-progressbar", "ui-widget ui-widget-content" );
+
+		this.valueDiv = $( "<div>" ).appendTo( this.element );
+		this._addClass( this.valueDiv, "ui-progressbar-value", "ui-widget-header" );
+		this._refreshValue();
+	},
+
+	_destroy: function() {
+		this.element.removeAttr( "role aria-valuemin aria-valuemax aria-valuenow" );
+
+		this.valueDiv.remove();
+	},
+
+	value: function( newValue ) {
+		if ( newValue === undefined ) {
+			return this.options.value;
+		}
+
+		this.options.value = this._constrainedValue( newValue );
+		this._refreshValue();
+	},
+
+	_constrainedValue: function( newValue ) {
+		if ( newValue === undefined ) {
+			newValue = this.options.value;
+		}
+
+		this.indeterminate = newValue === false;
+
+		// Sanitize value
+		if ( typeof newValue !== "number" ) {
+			newValue = 0;
+		}
+
+		return this.indeterminate ? false :
+			Math.min( this.options.max, Math.max( this.min, newValue ) );
+	},
+
+	_setOptions: function( options ) {
+
+		// Ensure "value" option is set after other values (like max)
+		var value = options.value;
+		delete options.value;
+
+		this._super( options );
+
+		this.options.value = this._constrainedValue( value );
+		this._refreshValue();
+	},
+
+	_setOption: function( key, value ) {
+		if ( key === "max" ) {
+
+			// Don't allow a max less than min
+			value = Math.max( this.min, value );
+		}
+		this._super( key, value );
+	},
+
+	_setOptionDisabled: function( value ) {
+		this._super( value );
+
+		this.element.attr( "aria-disabled", value );
+		this._toggleClass( null, "ui-state-disabled", !!value );
+	},
+
+	_percentage: function() {
+		return this.indeterminate ?
+			100 :
+			100 * ( this.options.value - this.min ) / ( this.options.max - this.min );
+	},
+
+	_refreshValue: function() {
+		var value = this.options.value,
+			percentage = this._percentage();
+
+		this.valueDiv
+			.toggle( this.indeterminate || value > this.min )
+			.width( percentage.toFixed( 0 ) + "%" );
+
+		this
+			._toggleClass( this.valueDiv, "ui-progressbar-complete", null,
+				value === this.options.max )
+			._toggleClass( "ui-progressbar-indeterminate", null, this.indeterminate );
+
+		if ( this.indeterminate ) {
+			this.element.removeAttr( "aria-valuenow" );
+			if ( !this.overlayDiv ) {
+				this.overlayDiv = $( "<div>" ).appendTo( this.valueDiv );
+				this._addClass( this.overlayDiv, "ui-progressbar-overlay" );
+			}
+		} else {
+			this.element.attr( {
+				"aria-valuemax": this.options.max,
+				"aria-valuenow": value
+			} );
+			if ( this.overlayDiv ) {
+				this.overlayDiv.remove();
+				this.overlayDiv = null;
+			}
+		}
+
+		if ( this.oldValue !== value ) {
+			this.oldValue = value;
+			this._trigger( "change" );
+		}
+		if ( value === this.options.max ) {
+			this._trigger( "complete" );
+		}
+	}
+} );
+
+
+/*!
+ * jQuery UI Selectable 1.12.0
+ * http://jqueryui.com
+ *
+ * Copyright jQuery Foundation and other contributors
+ * Released under the MIT license.
+ * http://jquery.org/license
+ */
+
+//>>label: Selectable
+//>>group: Interactions
+//>>description: Allows groups of elements to be selected with the mouse.
+//>>docs: http://api.jqueryui.com/selectable/
+//>>demos: http://jqueryui.com/selectable/
+//>>css.structure: ../../themes/base/selectable.css
+
+
+
+var widgetsSelectable = $.widget( "ui.selectable", $.ui.mouse, {
+	version: "1.12.0",
+	options: {
+		appendTo: "body",
+		autoRefresh: true,
+		distance: 0,
+		filter: "*",
+		tolerance: "touch",
+
+		// Callbacks
+		selected: null,
+		selecting: null,
+		start: null,
+		stop: null,
+		unselected: null,
+		unselecting: null
+	},
+	_create: function() {
+		var that = this;
+
+		this._addClass( "ui-selectable" );
+
+		this.dragged = false;
+
+		// Cache selectee children based on filter
+		this.refresh = function() {
+			that.elementPos = $( that.element[ 0 ] ).offset();
+			that.selectees = $( that.options.filter, that.element[ 0 ] );
+			that._addClass( that.selectees, "ui-selectee" );
+			that.selectees.each( function() {
+				var $this = $( this ),
+					selecteeOffset = $this.offset(),
+					pos = {
+						left: selecteeOffset.left - that.elementPos.left,
+						top: selecteeOffset.top - that.elementPos.top
+					};
+				$.data( this, "selectable-item", {
+					element: this,
+					$element: $this,
+					left: pos.left,
+					top: pos.top,
+					right: pos.left + $this.outerWidth(),
+					bottom: pos.top + $this.outerHeight(),
+					startselected: false,
+					selected: $this.hasClass( "ui-selected" ),
+					selecting: $this.hasClass( "ui-selecting" ),
+					unselecting: $this.hasClass( "ui-unselecting" )
+				} );
+			} );
+		};
+		this.refresh();
+
+		this._mouseInit();
+
+		this.helper = $( "<div>" );
+		this._addClass( this.helper, "ui-selectable-helper" );
+	},
+
+	_destroy: function() {
+		this.selectees.removeData( "selectable-item" );
+		this._mouseDestroy();
+	},
+
+	_mouseStart: function( event ) {
+		var that = this,
+			options = this.options;
+
+		this.opos = [ event.pageX, event.pageY ];
+		this.elementPos = $( this.element[ 0 ] ).offset();
+
+		if ( this.options.disabled ) {
+			return;
+		}
+
+		this.selectees = $( options.filter, this.element[ 0 ] );
+
+		this._trigger( "start", event );
+
+		$( options.appendTo ).append( this.helper );
+
+		// position helper (lasso)
+		this.helper.css( {
+			"left": event.pageX,
+			"top": event.pageY,
+			"width": 0,
+			"height": 0
+		} );
+
+		if ( options.autoRefresh ) {
+			this.refresh();
+		}
+
+		this.selectees.filter( ".ui-selected" ).each( function() {
+			var selectee = $.data( this, "selectable-item" );
+			selectee.startselected = true;
+			if ( !event.metaKey && !event.ctrlKey ) {
+				that._removeClass( selectee.$element, "ui-selected" );
+				selectee.selected = false;
+				that._addClass( selectee.$element, "ui-unselecting" );
+				selectee.unselecting = true;
+
+				// selectable UNSELECTING callback
+				that._trigger( "unselecting", event, {
+					unselecting: selectee.element
+				} );
+			}
+		} );
+
+		$( event.target ).parents().addBack().each( function() {
+			var doSelect,
+				selectee = $.data( this, "selectable-item" );
+			if ( selectee ) {
+				doSelect = ( !event.metaKey && !event.ctrlKey ) ||
+					!selectee.$element.hasClass( "ui-selected" );
+				that._removeClass( selectee.$element, doSelect ? "ui-unselecting" : "ui-selected" )
+					._addClass( selectee.$element, doSelect ? "ui-selecting" : "ui-unselecting" );
+				selectee.unselecting = !doSelect;
+				selectee.selecting = doSelect;
+				selectee.selected = doSelect;
+
+				// selectable (UN)SELECTING callback
+				if ( doSelect ) {
+					that._trigger( "selecting", event, {
+						selecting: selectee.element
+					} );
+				} else {
+					that._trigger( "unselecting", event, {
+						unselecting: selectee.element
+					} );
+				}
+				return false;
+			}
+		} );
+
+	},
+
+	_mouseDrag: function( event ) {
+
+		this.dragged = true;
+
+		if ( this.options.disabled ) {
+			return;
+		}
+
+		var tmp,
+			that = this,
+			options = this.options,
+			x1 = this.opos[ 0 ],
+			y1 = this.opos[ 1 ],
+			x2 = event.pageX,
+			y2 = event.pageY;
+
+		if ( x1 > x2 ) { tmp = x2; x2 = x1; x1 = tmp; }
+		if ( y1 > y2 ) { tmp = y2; y2 = y1; y1 = tmp; }
+		this.helper.css( { left: x1, top: y1, width: x2 - x1, height: y2 - y1 } );
+
+		this.selectees.each( function() {
+			var selectee = $.data( this, "selectable-item" ),
+				hit = false,
+				offset = {};
+
+			//prevent helper from being selected if appendTo: selectable
+			if ( !selectee || selectee.element === that.element[ 0 ] ) {
+				return;
+			}
+
+			offset.left   = selectee.left   + that.elementPos.left;
+			offset.right  = selectee.right  + that.elementPos.left;
+			offset.top    = selectee.top    + that.elementPos.top;
+			offset.bottom = selectee.bottom + that.elementPos.top;
+
+			if ( options.tolerance === "touch" ) {
+				hit = ( !( offset.left > x2 || offset.right < x1 || offset.top > y2 ||
+                    offset.bottom < y1 ) );
+			} else if ( options.tolerance === "fit" ) {
+				hit = ( offset.left > x1 && offset.right < x2 && offset.top > y1 &&
+                    offset.bottom < y2 );
+			}
+
+			if ( hit ) {
+
+				// SELECT
+				if ( selectee.selected ) {
+					that._removeClass( selectee.$element, "ui-selected" );
+					selectee.selected = false;
+				}
+				if ( selectee.unselecting ) {
+					that._removeClass( selectee.$element, "ui-unselecting" );
+					selectee.unselecting = false;
+				}
+				if ( !selectee.selecting ) {
+					that._addClass( selectee.$element, "ui-selecting" );
+					selectee.selecting = true;
+
+					// selectable SELECTING callback
+					that._trigger( "selecting", event, {
+						selecting: selectee.element
+					} );
+				}
+			} else {
+
+				// UNSELECT
+				if ( selectee.selecting ) {
+					if ( ( event.metaKey || event.ctrlKey ) && selectee.startselected ) {
+						that._removeClass( selectee.$element, "ui-selecting" );
+						selectee.selecting = false;
+						that._addClass( selectee.$element, "ui-selected" );
+						selectee.selected = true;
+					} else {
+						that._removeClass( selectee.$element, "ui-selecting" );
+						selectee.selecting = false;
+						if ( selectee.startselected ) {
+							that._addClass( selectee.$element, "ui-unselecting" );
+							selectee.unselecting = true;
+						}
+
+						// selectable UNSELECTING callback
+						that._trigger( "unselecting", event, {
+							unselecting: selectee.element
+						} );
+					}
+				}
+				if ( selectee.selected ) {
+					if ( !event.metaKey && !event.ctrlKey && !selectee.startselected ) {
+						that._removeClass( selectee.$element, "ui-selected" );
+						selectee.selected = false;
+
+						that._addClass( selectee.$element, "ui-unselecting" );
+						selectee.unselecting = true;
+
+						// selectable UNSELECTING callback
+						that._trigger( "unselecting", event, {
+							unselecting: selectee.element
+						} );
+					}
+				}
+			}
+		} );
+
+		return false;
+	},
+
+	_mouseStop: function( event ) {
+		var that = this;
+
+		this.dragged = false;
+
+		$( ".ui-unselecting", this.element[ 0 ] ).each( function() {
+			var selectee = $.data( this, "selectable-item" );
+			that._removeClass( selectee.$element, "ui-unselecting" );
+			selectee.unselecting = false;
+			selectee.startselected = false;
+			that._trigger( "unselected", event, {
+				unselected: selectee.element
+			} );
+		} );
+		$( ".ui-selecting", this.element[ 0 ] ).each( function() {
+			var selectee = $.data( this, "selectable-item" );
+			that._removeClass( selectee.$element, "ui-selecting" )
+				._addClass( selectee.$element, "ui-selected" );
+			selectee.selecting = false;
+			selectee.selected = true;
+			selectee.startselected = true;
+			that._trigger( "selected", event, {
+				selected: selectee.element
+			} );
+		} );
+		this._trigger( "stop", event );
+
+		this.helper.remove();
+
+		return false;
+	}
+
+} );
+
+
+/*!
+ * jQuery UI Selectmenu 1.12.0
+ * http://jqueryui.com
+ *
+ * Copyright jQuery Foundation and other contributors
+ * Released under the MIT license.
+ * http://jquery.org/license
+ */
+
+//>>label: Selectmenu
+//>>group: Widgets
+// jscs:disable maximumLineLength
+//>>description: Duplicates and extends the functionality of a native HTML select element, allowing it to be customizable in behavior and appearance far beyond the limitations of a native select.
+// jscs:enable maximumLineLength
+//>>docs: http://api.jqueryui.com/selectmenu/
+//>>demos: http://jqueryui.com/selectmenu/
+//>>css.structure: ../../themes/base/core.css
+//>>css.structure: ../../themes/base/selectmenu.css, ../../themes/base/button.css
+//>>css.theme: ../../themes/base/theme.css
+
+
+
+var widgetsSelectmenu = $.widget( "ui.selectmenu", [ $.ui.formResetMixin, {
+	version: "1.12.0",
+	defaultElement: "<select>",
+	options: {
+		appendTo: null,
+		classes: {
+			"ui-selectmenu-button-open": "ui-corner-top",
+			"ui-selectmenu-button-closed": "ui-corner-all"
+		},
+		disabled: null,
+		icons: {
+			button: "ui-icon-triangle-1-s"
+		},
+		position: {
+			my: "left top",
+			at: "left bottom",
+			collision: "none"
+		},
+		width: false,
+
+		// Callbacks
+		change: null,
+		close: null,
+		focus: null,
+		open: null,
+		select: null
+	},
+
+	_create: function() {
+		var selectmenuId = this.element.uniqueId().attr( "id" );
+		this.ids = {
+			element: selectmenuId,
+			button: selectmenuId + "-button",
+			menu: selectmenuId + "-menu"
+		};
+
+		this._drawButton();
+		this._drawMenu();
+		this._bindFormResetHandler();
+
+		this._rendered = false;
+		this.menuItems = $();
+	},
+
+	_drawButton: function() {
+		var icon,
+			that = this,
+			item = this._parseOption(
+				this.element.find( "option:selected" ),
+				this.element[ 0 ].selectedIndex
+			);
+
+		// Associate existing label with the new button
+		this.labels = this.element.labels().attr( "for", this.ids.button );
+		this._on( this.labels, {
+			click: function( event ) {
+				this.button.focus();
+				event.preventDefault();
+			}
+		} );
+
+		// Hide original select element
+		this.element.hide();
+
+		// Create button
+		this.button = $( "<span>", {
+			tabindex: this.options.disabled ? -1 : 0,
+			id: this.ids.button,
+			role: "combobox",
+			"aria-expanded": "false",
+			"aria-autocomplete": "list",
+			"aria-owns": this.ids.menu,
+			"aria-haspopup": "true",
+			title: this.element.attr( "title" )
+		} )
+			.insertAfter( this.element );
+
+		this._addClass( this.button, "ui-selectmenu-button ui-selectmenu-button-closed",
+			"ui-button ui-widget" );
+
+		icon = $( "<span>" ).appendTo( this.button );
+		this._addClass( icon, "ui-selectmenu-icon", "ui-icon " + this.options.icons.button );
+		this.buttonItem = this._renderButtonItem( item )
+			.appendTo( this.button );
+
+		if ( this.options.width !== false ) {
+			this._resizeButton();
+		}
+
+		this._on( this.button, this._buttonEvents );
+		this.button.one( "focusin", function() {
+
+			// Delay rendering the menu items until the button receives focus.
+			// The menu may have already been rendered via a programmatic open.
+			if ( !that._rendered ) {
+				that._refreshMenu();
+			}
+		} );
+	},
+
+	_drawMenu: function() {
+		var that = this;
+
+		// Create menu
+		this.menu = $( "<ul>", {
+			"aria-hidden": "true",
+			"aria-labelledby": this.ids.button,
+			id: this.ids.menu
+		} );
+
+		// Wrap menu
+		this.menuWrap = $( "<div>" ).append( this.menu );
+		this._addClass( this.menuWrap, "ui-selectmenu-menu", "ui-front" );
+		this.menuWrap.appendTo( this._appendTo() );
+
+		// Initialize menu widget
+		this.menuInstance = this.menu
+			.menu( {
+				classes: {
+					"ui-menu": "ui-corner-bottom"
+				},
+				role: "listbox",
+				select: function( event, ui ) {
+					event.preventDefault();
+
+					// Support: IE8
+					// If the item was selected via a click, the text selection
+					// will be destroyed in IE
+					that._setSelection();
+
+					that._select( ui.item.data( "ui-selectmenu-item" ), event );
+				},
+				focus: function( event, ui ) {
+					var item = ui.item.data( "ui-selectmenu-item" );
+
+					// Prevent inital focus from firing and check if its a newly focused item
+					if ( that.focusIndex != null && item.index !== that.focusIndex ) {
+						that._trigger( "focus", event, { item: item } );
+						if ( !that.isOpen ) {
+							that._select( item, event );
+						}
+					}
+					that.focusIndex = item.index;
+
+					that.button.attr( "aria-activedescendant",
+						that.menuItems.eq( item.index ).attr( "id" ) );
+				}
+			} )
+			.menu( "instance" );
+
+		// Don't close the menu on mouseleave
+		this.menuInstance._off( this.menu, "mouseleave" );
+
+		// Cancel the menu's collapseAll on document click
+		this.menuInstance._closeOnDocumentClick = function() {
+			return false;
+		};
+
+		// Selects often contain empty items, but never contain dividers
+		this.menuInstance._isDivider = function() {
+			return false;
+		};
+	},
+
+	refresh: function() {
+		this._refreshMenu();
+		this.buttonItem.replaceWith(
+			this.buttonItem = this._renderButtonItem(
+
+				// Fall back to an empty object in case there are no options
+				this._getSelectedItem().data( "ui-selectmenu-item" ) || {}
+			)
+		);
+		if ( this.options.width === null ) {
+			this._resizeButton();
+		}
+	},
+
+	_refreshMenu: function() {
+		var item,
+			options = this.element.find( "option" );
+
+		this.menu.empty();
+
+		this._parseOptions( options );
+		this._renderMenu( this.menu, this.items );
+
+		this.menuInstance.refresh();
+		this.menuItems = this.menu.find( "li" )
+			.not( ".ui-selectmenu-optgroup" )
+				.find( ".ui-menu-item-wrapper" );
+
+		this._rendered = true;
+
+		if ( !options.length ) {
+			return;
+		}
+
+		item = this._getSelectedItem();
+
+		// Update the menu to have the correct item focused
+		this.menuInstance.focus( null, item );
+		this._setAria( item.data( "ui-selectmenu-item" ) );
+
+		// Set disabled state
+		this._setOption( "disabled", this.element.prop( "disabled" ) );
+	},
+
+	open: function( event ) {
+		if ( this.options.disabled ) {
+			return;
+		}
+
+		// If this is the first time the menu is being opened, render the items
+		if ( !this._rendered ) {
+			this._refreshMenu();
+		} else {
+
+			// Menu clears focus on close, reset focus to selected item
+			this._removeClass( this.menu.find( ".ui-state-active" ), null, "ui-state-active" );
+			this.menuInstance.focus( null, this._getSelectedItem() );
+		}
+
+		// If there are no options, don't open the menu
+		if ( !this.menuItems.length ) {
+			return;
+		}
+
+		this.isOpen = true;
+		this._toggleAttr();
+		this._resizeMenu();
+		this._position();
+
+		this._on( this.document, this._documentClick );
+
+		this._trigger( "open", event );
+	},
+
+	_position: function() {
+		this.menuWrap.position( $.extend( { of: this.button }, this.options.position ) );
+	},
+
+	close: function( event ) {
+		if ( !this.isOpen ) {
+			return;
+		}
+
+		this.isOpen = false;
+		this._toggleAttr();
+
+		this.range = null;
+		this._off( this.document );
+
+		this._trigger( "close", event );
+	},
+
+	widget: function() {
+		return this.button;
+	},
+
+	menuWidget: function() {
+		return this.menu;
+	},
+
+	_renderButtonItem: function( item ) {
+		var buttonItem = $( "<span>" );
+
+		this._setText( buttonItem, item.label );
+		this._addClass( buttonItem, "ui-selectmenu-text" );
+
+		return buttonItem;
+	},
+
+	_renderMenu: function( ul, items ) {
+		var that = this,
+			currentOptgroup = "";
+
+		$.each( items, function( index, item ) {
+			var li;
+
+			if ( item.optgroup !== currentOptgroup ) {
+				li = $( "<li>", {
+					text: item.optgroup
+				} );
+				that._addClass( li, "ui-selectmenu-optgroup", "ui-menu-divider" +
+					( item.element.parent( "optgroup" ).prop( "disabled" ) ?
+						" ui-state-disabled" :
+						"" ) );
+
+				li.appendTo( ul );
+
+				currentOptgroup = item.optgroup;
+			}
+
+			that._renderItemData( ul, item );
+		} );
+	},
+
+	_renderItemData: function( ul, item ) {
+		return this._renderItem( ul, item ).data( "ui-selectmenu-item", item );
+	},
+
+	_renderItem: function( ul, item ) {
+		var li = $( "<li>" ),
+			wrapper = $( "<div>", {
+				title: item.element.attr( "title" )
+			} );
+
+		if ( item.disabled ) {
+			this._addClass( li, null, "ui-state-disabled" );
+		}
+		this._setText( wrapper, item.label );
+
+		return li.append( wrapper ).appendTo( ul );
+	},
+
+	_setText: function( element, value ) {
+		if ( value ) {
+			element.text( value );
+		} else {
+			element.html( "&#160;" );
+		}
+	},
+
+	_move: function( direction, event ) {
+		var item, next,
+			filter = ".ui-menu-item";
+
+		if ( this.isOpen ) {
+			item = this.menuItems.eq( this.focusIndex ).parent( "li" );
+		} else {
+			item = this.menuItems.eq( this.element[ 0 ].selectedIndex ).parent( "li" );
+			filter += ":not(.ui-state-disabled)";
+		}
+
+		if ( direction === "first" || direction === "last" ) {
+			next = item[ direction === "first" ? "prevAll" : "nextAll" ]( filter ).eq( -1 );
+		} else {
+			next = item[ direction + "All" ]( filter ).eq( 0 );
+		}
+
+		if ( next.length ) {
+			this.menuInstance.focus( event, next );
+		}
+	},
+
+	_getSelectedItem: function() {
+		return this.menuItems.eq( this.element[ 0 ].selectedIndex ).parent( "li" );
+	},
+
+	_toggle: function( event ) {
+		this[ this.isOpen ? "close" : "open" ]( event );
+	},
+
+	_setSelection: function() {
+		var selection;
+
+		if ( !this.range ) {
+			return;
+		}
+
+		if ( window.getSelection ) {
+			selection = window.getSelection();
+			selection.removeAllRanges();
+			selection.addRange( this.range );
+
+		// Support: IE8
+		} else {
+			this.range.select();
+		}
+
+		// Support: IE
+		// Setting the text selection kills the button focus in IE, but
+		// restoring the focus doesn't kill the selection.
+		this.button.focus();
+	},
+
+	_documentClick: {
+		mousedown: function( event ) {
+			if ( !this.isOpen ) {
+				return;
+			}
+
+			if ( !$( event.target ).closest( ".ui-selectmenu-menu, #" +
+					$.ui.escapeSelector( this.ids.button ) ).length ) {
+				this.close( event );
+			}
+		}
+	},
+
+	_buttonEvents: {
+
+		// Prevent text selection from being reset when interacting with the selectmenu (#10144)
+		mousedown: function() {
+			var selection;
+
+			if ( window.getSelection ) {
+				selection = window.getSelection();
+				if ( selection.rangeCount ) {
+					this.range = selection.getRangeAt( 0 );
+				}
+
+			// Support: IE8
+			} else {
+				this.range = document.selection.createRange();
+			}
+		},
+
+		click: function( event ) {
+			this._setSelection();
+			this._toggle( event );
+		},
+
+		keydown: function( event ) {
+			var preventDefault = true;
+			switch ( event.keyCode ) {
+			case $.ui.keyCode.TAB:
+			case $.ui.keyCode.ESCAPE:
+				this.close( event );
+				preventDefault = false;
+				break;
+			case $.ui.keyCode.ENTER:
+				if ( this.isOpen ) {
+					this._selectFocusedItem( event );
+				}
+				break;
+			case $.ui.keyCode.UP:
+				if ( event.altKey ) {
+					this._toggle( event );
+				} else {
+					this._move( "prev", event );
+				}
+				break;
+			case $.ui.keyCode.DOWN:
+				if ( event.altKey ) {
+					this._toggle( event );
+				} else {
+					this._move( "next", event );
+				}
+				break;
+			case $.ui.keyCode.SPACE:
+				if ( this.isOpen ) {
+					this._selectFocusedItem( event );
+				} else {
+					this._toggle( event );
+				}
+				break;
+			case $.ui.keyCode.LEFT:
+				this._move( "prev", event );
+				break;
+			case $.ui.keyCode.RIGHT:
+				this._move( "next", event );
+				break;
+			case $.ui.keyCode.HOME:
+			case $.ui.keyCode.PAGE_UP:
+				this._move( "first", event );
+				break;
+			case $.ui.keyCode.END:
+			case $.ui.keyCode.PAGE_DOWN:
+				this._move( "last", event );
+				break;
+			default:
+				this.menu.trigger( event );
+				preventDefault = false;
+			}
+
+			if ( preventDefault ) {
+				event.preventDefault();
+			}
+		}
+	},
+
+	_selectFocusedItem: function( event ) {
+		var item = this.menuItems.eq( this.focusIndex ).parent( "li" );
+		if ( !item.hasClass( "ui-state-disabled" ) ) {
+			this._select( item.data( "ui-selectmenu-item" ), event );
+		}
+	},
+
+	_select: function( item, event ) {
+		var oldIndex = this.element[ 0 ].selectedIndex;
+
+		// Change native select element
+		this.element[ 0 ].selectedIndex = item.index;
+		this.buttonItem.replaceWith( this.buttonItem = this._renderButtonItem( item ) );
+		this._setAria( item );
+		this._trigger( "select", event, { item: item } );
+
+		if ( item.index !== oldIndex ) {
+			this._trigger( "change", event, { item: item } );
+		}
+
+		this.close( event );
+	},
+
+	_setAria: function( item ) {
+		var id = this.menuItems.eq( item.index ).attr( "id" );
+
+		this.button.attr( {
+			"aria-labelledby": id,
+			"aria-activedescendant": id
+		} );
+		this.menu.attr( "aria-activedescendant", id );
+	},
+
+	_setOption: function( key, value ) {
+		if ( key === "icons" ) {
+			var icon = this.button.find( "span.ui-icon" );
+			this._removeClass( icon, null, this.options.icons.button )
+				._addClass( icon, null, value.button );
+		}
+
+		this._super( key, value );
+
+		if ( key === "appendTo" ) {
+			this.menuWrap.appendTo( this._appendTo() );
+		}
+
+		if ( key === "width" ) {
+			this._resizeButton();
+		}
+	},
+
+	_setOptionDisabled: function( value ) {
+		this._super( value );
+
+		this.menuInstance.option( "disabled", value );
+		this.button.attr( "aria-disabled", value );
+		this._toggleClass( this.button, null, "ui-state-disabled", value );
+
+		this.element.prop( "disabled", value );
+		if ( value ) {
+			this.button.attr( "tabindex", -1 );
+			this.close();
+		} else {
+			this.button.attr( "tabindex", 0 );
+		}
+	},
+
+	_appendTo: function() {
+		var element = this.options.appendTo;
+
+		if ( element ) {
+			element = element.jquery || element.nodeType ?
+				$( element ) :
+				this.document.find( element ).eq( 0 );
+		}
+
+		if ( !element || !element[ 0 ] ) {
+			element = this.element.closest( ".ui-front, dialog" );
+		}
+
+		if ( !element.length ) {
+			element = this.document[ 0 ].body;
+		}
+
+		return element;
+	},
+
+	_toggleAttr: function() {
+		this.button.attr( "aria-expanded", this.isOpen );
+
+		// We can't use two _toggleClass() calls here, because we need to make sure
+		// we always remove classes first and add them second, otherwise if both classes have the
+		// same theme class, it will be removed after we add it.
+		this._removeClass( this.button, "ui-selectmenu-button-" +
+			( this.isOpen ? "closed" : "open" ) )
+			._addClass( this.button, "ui-selectmenu-button-" +
+				( this.isOpen ? "open" : "closed" ) )
+			._toggleClass( this.menuWrap, "ui-selectmenu-open", null, this.isOpen );
+
+		this.menu.attr( "aria-hidden", !this.isOpen );
+	},
+
+	_resizeButton: function() {
+		var width = this.options.width;
+
+		// For `width: false`, just remove inline style and stop
+		if ( width === false ) {
+			this.button.css( "width", "" );
+			return;
+		}
+
+		// For `width: null`, match the width of the original element
+		if ( width === null ) {
+			width = this.element.show().outerWidth();
+			this.element.hide();
+		}
+
+		this.button.outerWidth( width );
+	},
+
+	_resizeMenu: function() {
+		this.menu.outerWidth( Math.max(
+			this.button.outerWidth(),
+
+			// Support: IE10
+			// IE10 wraps long text (possibly a rounding bug)
+			// so we add 1px to avoid the wrapping
+			this.menu.width( "" ).outerWidth() + 1
+		) );
+	},
+
+	_getCreateOptions: function() {
+		var options = this._super();
+
+		options.disabled = this.element.prop( "disabled" );
+
+		return options;
+	},
+
+	_parseOptions: function( options ) {
+		var that = this,
+			data = [];
+		options.each( function( index, item ) {
+			data.push( that._parseOption( $( item ), index ) );
+		} );
+		this.items = data;
+	},
+
+	_parseOption: function( option, index ) {
+		var optgroup = option.parent( "optgroup" );
+
+		return {
+			element: option,
+			index: index,
+			value: option.val(),
+			label: option.text(),
+			optgroup: optgroup.attr( "label" ) || "",
+			disabled: optgroup.prop( "disabled" ) || option.prop( "disabled" )
+		};
+	},
+
+	_destroy: function() {
+		this._unbindFormResetHandler();
+		this.menuWrap.remove();
+		this.button.remove();
+		this.element.show();
+		this.element.removeUniqueId();
+		this.labels.attr( "for", this.ids.element );
+	}
+} ] );
+
+
+/*!
+ * jQuery UI Slider 1.12.0
+ * http://jqueryui.com
+ *
+ * Copyright jQuery Foundation and other contributors
+ * Released under the MIT license.
+ * http://jquery.org/license
+ */
+
+//>>label: Slider
+//>>group: Widgets
+//>>description: Displays a flexible slider with ranges and accessibility via keyboard.
+//>>docs: http://api.jqueryui.com/slider/
+//>>demos: http://jqueryui.com/slider/
+//>>css.structure: ../../themes/base/core.css
+//>>css.structure: ../../themes/base/slider.css
+//>>css.theme: ../../themes/base/theme.css
+
+
+
+var widgetsSlider = $.widget( "ui.slider", $.ui.mouse, {
+	version: "1.12.0",
+	widgetEventPrefix: "slide",
+
+	options: {
+		animate: false,
+		classes: {
+			"ui-slider": "ui-corner-all",
+			"ui-slider-handle": "ui-corner-all",
+
+			// Note: ui-widget-header isn't the most fittingly semantic framework class for this
+			// element, but worked best visually with a variety of themes
+			"ui-slider-range": "ui-corner-all ui-widget-header"
+		},
+		distance: 0,
+		max: 100,
+		min: 0,
+		orientation: "horizontal",
+		range: false,
+		step: 1,
+		value: 0,
+		values: null,
+
+		// Callbacks
+		change: null,
+		slide: null,
+		start: null,
+		stop: null
+	},
+
+	// Number of pages in a slider
+	// (how many times can you page up/down to go through the whole range)
+	numPages: 5,
+
+	_create: function() {
+		this._keySliding = false;
+		this._mouseSliding = false;
+		this._animateOff = true;
+		this._handleIndex = null;
+		this._detectOrientation();
+		this._mouseInit();
+		this._calculateNewMax();
+
+		this._addClass( "ui-slider ui-slider-" + this.orientation,
+			"ui-widget ui-widget-content" );
+
+		this._refresh();
+
+		this._animateOff = false;
+	},
+
+	_refresh: function() {
+		this._createRange();
+		this._createHandles();
+		this._setupEvents();
+		this._refreshValue();
+	},
+
+	_createHandles: function() {
+		var i, handleCount,
+			options = this.options,
+			existingHandles = this.element.find( ".ui-slider-handle" ),
+			handle = "<span tabindex='0'></span>",
+			handles = [];
+
+		handleCount = ( options.values && options.values.length ) || 1;
+
+		if ( existingHandles.length > handleCount ) {
+			existingHandles.slice( handleCount ).remove();
+			existingHandles = existingHandles.slice( 0, handleCount );
+		}
+
+		for ( i = existingHandles.length; i < handleCount; i++ ) {
+			handles.push( handle );
+		}
+
+		this.handles = existingHandles.add( $( handles.join( "" ) ).appendTo( this.element ) );
+
+		this._addClass( this.handles, "ui-slider-handle", "ui-state-default" );
+
+		this.handle = this.handles.eq( 0 );
+
+		this.handles.each( function( i ) {
+			$( this ).data( "ui-slider-handle-index", i );
+		} );
+	},
+
+	_createRange: function() {
+		var options = this.options;
+
+		if ( options.range ) {
+			if ( options.range === true ) {
+				if ( !options.values ) {
+					options.values = [ this._valueMin(), this._valueMin() ];
+				} else if ( options.values.length && options.values.length !== 2 ) {
+					options.values = [ options.values[ 0 ], options.values[ 0 ] ];
+				} else if ( $.isArray( options.values ) ) {
+					options.values = options.values.slice( 0 );
+				}
+			}
+
+			if ( !this.range || !this.range.length ) {
+				this.range = $( "<div>" )
+					.appendTo( this.element );
+
+				this._addClass( this.range, "ui-slider-range" );
+			} else {
+				this._removeClass( this.range, "ui-slider-range-min ui-slider-range-max" );
+
+				// Handle range switching from true to min/max
+				this.range.css( {
+					"left": "",
+					"bottom": ""
+				} );
+			}
+			if ( options.range === "min" || options.range === "max" ) {
+				this._addClass( this.range, "ui-slider-range-" + options.range );
+			}
+		} else {
+			if ( this.range ) {
+				this.range.remove();
+			}
+			this.range = null;
+		}
+	},
+
+	_setupEvents: function() {
+		this._off( this.handles );
+		this._on( this.handles, this._handleEvents );
+		this._hoverable( this.handles );
+		this._focusable( this.handles );
+	},
+
+	_destroy: function() {
+		this.handles.remove();
+		if ( this.range ) {
+			this.range.remove();
+		}
+
+		this._mouseDestroy();
+	},
+
+	_mouseCapture: function( event ) {
+		var position, normValue, distance, closestHandle, index, allowed, offset, mouseOverHandle,
+			that = this,
+			o = this.options;
+
+		if ( o.disabled ) {
+			return false;
+		}
+
+		this.elementSize = {
+			width: this.element.outerWidth(),
+			height: this.element.outerHeight()
+		};
+		this.elementOffset = this.element.offset();
+
+		position = { x: event.pageX, y: event.pageY };
+		normValue = this._normValueFromMouse( position );
+		distance = this._valueMax() - this._valueMin() + 1;
+		this.handles.each( function( i ) {
+			var thisDistance = Math.abs( normValue - that.values( i ) );
+			if ( ( distance > thisDistance ) ||
+				( distance === thisDistance &&
+					( i === that._lastChangedValue || that.values( i ) === o.min ) ) ) {
+				distance = thisDistance;
+				closestHandle = $( this );
+				index = i;
+			}
+		} );
+
+		allowed = this._start( event, index );
+		if ( allowed === false ) {
+			return false;
+		}
+		this._mouseSliding = true;
+
+		this._handleIndex = index;
+
+		this._addClass( closestHandle, null, "ui-state-active" );
+		closestHandle.trigger( "focus" );
+
+		offset = closestHandle.offset();
+		mouseOverHandle = !$( event.target ).parents().addBack().is( ".ui-slider-handle" );
+		this._clickOffset = mouseOverHandle ? { left: 0, top: 0 } : {
+			left: event.pageX - offset.left - ( closestHandle.width() / 2 ),
+			top: event.pageY - offset.top -
+				( closestHandle.height() / 2 ) -
+				( parseInt( closestHandle.css( "borderTopWidth" ), 10 ) || 0 ) -
+				( parseInt( closestHandle.css( "borderBottomWidth" ), 10 ) || 0 ) +
+				( parseInt( closestHandle.css( "marginTop" ), 10 ) || 0 )
+		};
+
+		if ( !this.handles.hasClass( "ui-state-hover" ) ) {
+			this._slide( event, index, normValue );
+		}
+		this._animateOff = true;
+		return true;
+	},
+
+	_mouseStart: function() {
+		return true;
+	},
+
+	_mouseDrag: function( event ) {
+		var position = { x: event.pageX, y: event.pageY },
+			normValue = this._normValueFromMouse( position );
+
+		this._slide( event, this._handleIndex, normValue );
+
+		return false;
+	},
+
+	_mouseStop: function( event ) {
+		this._removeClass( this.handles, null, "ui-state-active" );
+		this._mouseSliding = false;
+
+		this._stop( event, this._handleIndex );
+		this._change( event, this._handleIndex );
+
+		this._handleIndex = null;
+		this._clickOffset = null;
+		this._animateOff = false;
+
+		return false;
+	},
+
+	_detectOrientation: function() {
+		this.orientation = ( this.options.orientation === "vertical" ) ? "vertical" : "horizontal";
+	},
+
+	_normValueFromMouse: function( position ) {
+		var pixelTotal,
+			pixelMouse,
+			percentMouse,
+			valueTotal,
+			valueMouse;
+
+		if ( this.orientation === "horizontal" ) {
+			pixelTotal = this.elementSize.width;
+			pixelMouse = position.x - this.elementOffset.left -
+				( this._clickOffset ? this._clickOffset.left : 0 );
+		} else {
+			pixelTotal = this.elementSize.height;
+			pixelMouse = position.y - this.elementOffset.top -
+				( this._clickOffset ? this._clickOffset.top : 0 );
+		}
+
+		percentMouse = ( pixelMouse / pixelTotal );
+		if ( percentMouse > 1 ) {
+			percentMouse = 1;
+		}
+		if ( percentMouse < 0 ) {
+			percentMouse = 0;
+		}
+		if ( this.orientation === "vertical" ) {
+			percentMouse = 1 - percentMouse;
+		}
+
+		valueTotal = this._valueMax() - this._valueMin();
+		valueMouse = this._valueMin() + percentMouse * valueTotal;
+
+		return this._trimAlignValue( valueMouse );
+	},
+
+	_uiHash: function( index, value, values ) {
+		var uiHash = {
+			handle: this.handles[ index ],
+			handleIndex: index,
+			value: value !== undefined ? value : this.value()
+		};
+
+		if ( this._hasMultipleValues() ) {
+			uiHash.value = value !== undefined ? value : this.values( index );
+			uiHash.values = values || this.values();
+		}
+
+		return uiHash;
+	},
+
+	_hasMultipleValues: function() {
+		return this.options.values && this.options.values.length;
+	},
+
+	_start: function( event, index ) {
+		return this._trigger( "start", event, this._uiHash( index ) );
+	},
+
+	_slide: function( event, index, newVal ) {
+		var allowed, otherVal,
+			currentValue = this.value(),
+			newValues = this.values();
+
+		if ( this._hasMultipleValues() ) {
+			otherVal = this.values( index ? 0 : 1 );
+			currentValue = this.values( index );
+
+			if ( this.options.values.length === 2 && this.options.range === true ) {
+				newVal =  index === 0 ? Math.min( otherVal, newVal ) : Math.max( otherVal, newVal );
+			}
+
+			newValues[ index ] = newVal;
+		}
+
+		if ( newVal === currentValue ) {
+			return;
+		}
+
+		allowed = this._trigger( "slide", event, this._uiHash( index, newVal, newValues ) );
+
+		// A slide can be canceled by returning false from the slide callback
+		if ( allowed === false ) {
+			return;
+		}
+
+		if ( this._hasMultipleValues() ) {
+			this.values( index, newVal );
+		} else {
+			this.value( newVal );
+		}
+	},
+
+	_stop: function( event, index ) {
+		this._trigger( "stop", event, this._uiHash( index ) );
+	},
+
+	_change: function( event, index ) {
+		if ( !this._keySliding && !this._mouseSliding ) {
+
+			//store the last changed value index for reference when handles overlap
+			this._lastChangedValue = index;
+			this._trigger( "change", event, this._uiHash( index ) );
+		}
+	},
+
+	value: function( newValue ) {
+		if ( arguments.length ) {
+			this.options.value = this._trimAlignValue( newValue );
+			this._refreshValue();
+			this._change( null, 0 );
+			return;
+		}
+
+		return this._value();
+	},
+
+	values: function( index, newValue ) {
+		var vals,
+			newValues,
+			i;
+
+		if ( arguments.length > 1 ) {
+			this.options.values[ index ] = this._trimAlignValue( newValue );
+			this._refreshValue();
+			this._change( null, index );
+			return;
+		}
+
+		if ( arguments.length ) {
+			if ( $.isArray( arguments[ 0 ] ) ) {
+				vals = this.options.values;
+				newValues = arguments[ 0 ];
+				for ( i = 0; i < vals.length; i += 1 ) {
+					vals[ i ] = this._trimAlignValue( newValues[ i ] );
+					this._change( null, i );
+				}
+				this._refreshValue();
+			} else {
+				if ( this._hasMultipleValues() ) {
+					return this._values( index );
+				} else {
+					return this.value();
+				}
+			}
+		} else {
+			return this._values();
+		}
+	},
+
+	_setOption: function( key, value ) {
+		var i,
+			valsLength = 0;
+
+		if ( key === "range" && this.options.range === true ) {
+			if ( value === "min" ) {
+				this.options.value = this._values( 0 );
+				this.options.values = null;
+			} else if ( value === "max" ) {
+				this.options.value = this._values( this.options.values.length - 1 );
+				this.options.values = null;
+			}
+		}
+
+		if ( $.isArray( this.options.values ) ) {
+			valsLength = this.options.values.length;
+		}
+
+		this._super( key, value );
+
+		switch ( key ) {
+			case "orientation":
+				this._detectOrientation();
+				this._removeClass( "ui-slider-horizontal ui-slider-vertical" )
+					._addClass( "ui-slider-" + this.orientation );
+				this._refreshValue();
+				if ( this.options.range ) {
+					this._refreshRange( value );
+				}
+
+				// Reset positioning from previous orientation
+				this.handles.css( value === "horizontal" ? "bottom" : "left", "" );
+				break;
+			case "value":
+				this._animateOff = true;
+				this._refreshValue();
+				this._change( null, 0 );
+				this._animateOff = false;
+				break;
+			case "values":
+				this._animateOff = true;
+				this._refreshValue();
+
+				// Start from the last handle to prevent unreachable handles (#9046)
+				for ( i = valsLength - 1; i >= 0; i-- ) {
+					this._change( null, i );
+				}
+				this._animateOff = false;
+				break;
+			case "step":
+			case "min":
+			case "max":
+				this._animateOff = true;
+				this._calculateNewMax();
+				this._refreshValue();
+				this._animateOff = false;
+				break;
+			case "range":
+				this._animateOff = true;
+				this._refresh();
+				this._animateOff = false;
+				break;
+		}
+	},
+
+	_setOptionDisabled: function( value ) {
+		this._super( value );
+
+		this._toggleClass( null, "ui-state-disabled", !!value );
+	},
+
+	//internal value getter
+	// _value() returns value trimmed by min and max, aligned by step
+	_value: function() {
+		var val = this.options.value;
+		val = this._trimAlignValue( val );
+
+		return val;
+	},
+
+	//internal values getter
+	// _values() returns array of values trimmed by min and max, aligned by step
+	// _values( index ) returns single value trimmed by min and max, aligned by step
+	_values: function( index ) {
+		var val,
+			vals,
+			i;
+
+		if ( arguments.length ) {
+			val = this.options.values[ index ];
+			val = this._trimAlignValue( val );
+
+			return val;
+		} else if ( this._hasMultipleValues() ) {
+
+			// .slice() creates a copy of the array
+			// this copy gets trimmed by min and max and then returned
+			vals = this.options.values.slice();
+			for ( i = 0; i < vals.length; i += 1 ) {
+				vals[ i ] = this._trimAlignValue( vals[ i ] );
+			}
+
+			return vals;
+		} else {
+			return [];
+		}
+	},
+
+	// Returns the step-aligned value that val is closest to, between (inclusive) min and max
+	_trimAlignValue: function( val ) {
+		if ( val <= this._valueMin() ) {
+			return this._valueMin();
+		}
+		if ( val >= this._valueMax() ) {
+			return this._valueMax();
+		}
+		var step = ( this.options.step > 0 ) ? this.options.step : 1,
+			valModStep = ( val - this._valueMin() ) % step,
+			alignValue = val - valModStep;
+
+		if ( Math.abs( valModStep ) * 2 >= step ) {
+			alignValue += ( valModStep > 0 ) ? step : ( -step );
+		}
+
+		// Since JavaScript has problems with large floats, round
+		// the final value to 5 digits after the decimal point (see #4124)
+		return parseFloat( alignValue.toFixed( 5 ) );
+	},
+
+	_calculateNewMax: function() {
+		var max = this.options.max,
+			min = this._valueMin(),
+			step = this.options.step,
+			aboveMin = Math.round( ( max - min ) / step ) * step;
+		max = aboveMin + min;
+		if ( max > this.options.max ) {
+
+			//If max is not divisible by step, rounding off may increase its value
+			max -= step;
+		}
+		this.max = parseFloat( max.toFixed( this._precision() ) );
+	},
+
+	_precision: function() {
+		var precision = this._precisionOf( this.options.step );
+		if ( this.options.min !== null ) {
+			precision = Math.max( precision, this._precisionOf( this.options.min ) );
+		}
+		return precision;
+	},
+
+	_precisionOf: function( num ) {
+		var str = num.toString(),
+			decimal = str.indexOf( "." );
+		return decimal === -1 ? 0 : str.length - decimal - 1;
+	},
+
+	_valueMin: function() {
+		return this.options.min;
+	},
+
+	_valueMax: function() {
+		return this.max;
+	},
+
+	_refreshRange: function( orientation ) {
+		if ( orientation === "vertical" ) {
+			this.range.css( { "width": "", "left": "" } );
+		}
+		if ( orientation === "horizontal" ) {
+			this.range.css( { "height": "", "bottom": "" } );
+		}
+	},
+
+	_refreshValue: function() {
+		var lastValPercent, valPercent, value, valueMin, valueMax,
+			oRange = this.options.range,
+			o = this.options,
+			that = this,
+			animate = ( !this._animateOff ) ? o.animate : false,
+			_set = {};
+
+		if ( this._hasMultipleValues() ) {
+			this.handles.each( function( i ) {
+				valPercent = ( that.values( i ) - that._valueMin() ) / ( that._valueMax() -
+					that._valueMin() ) * 100;
+				_set[ that.orientation === "horizontal" ? "left" : "bottom" ] = valPercent + "%";
+				$( this ).stop( 1, 1 )[ animate ? "animate" : "css" ]( _set, o.animate );
+				if ( that.options.range === true ) {
+					if ( that.orientation === "horizontal" ) {
+						if ( i === 0 ) {
+							that.range.stop( 1, 1 )[ animate ? "animate" : "css" ]( {
+								left: valPercent + "%"
+							}, o.animate );
+						}
+						if ( i === 1 ) {
+							that.range[ animate ? "animate" : "css" ]( {
+								width: ( valPercent - lastValPercent ) + "%"
+							}, {
+								queue: false,
+								duration: o.animate
+							} );
+						}
+					} else {
+						if ( i === 0 ) {
+							that.range.stop( 1, 1 )[ animate ? "animate" : "css" ]( {
+								bottom: ( valPercent ) + "%"
+							}, o.animate );
+						}
+						if ( i === 1 ) {
+							that.range[ animate ? "animate" : "css" ]( {
+								height: ( valPercent - lastValPercent ) + "%"
+							}, {
+								queue: false,
+								duration: o.animate
+							} );
+						}
+					}
+				}
+				lastValPercent = valPercent;
+			} );
+		} else {
+			value = this.value();
+			valueMin = this._valueMin();
+			valueMax = this._valueMax();
+			valPercent = ( valueMax !== valueMin ) ?
+					( value - valueMin ) / ( valueMax - valueMin ) * 100 :
+					0;
+			_set[ this.orientation === "horizontal" ? "left" : "bottom" ] = valPercent + "%";
+			this.handle.stop( 1, 1 )[ animate ? "animate" : "css" ]( _set, o.animate );
+
+			if ( oRange === "min" && this.orientation === "horizontal" ) {
+				this.range.stop( 1, 1 )[ animate ? "animate" : "css" ]( {
+					width: valPercent + "%"
+				}, o.animate );
+			}
+			if ( oRange === "max" && this.orientation === "horizontal" ) {
+				this.range.stop( 1, 1 )[ animate ? "animate" : "css" ]( {
+					width: ( 100 - valPercent ) + "%"
+				}, o.animate );
+			}
+			if ( oRange === "min" && this.orientation === "vertical" ) {
+				this.range.stop( 1, 1 )[ animate ? "animate" : "css" ]( {
+					height: valPercent + "%"
+				}, o.animate );
+			}
+			if ( oRange === "max" && this.orientation === "vertical" ) {
+				this.range.stop( 1, 1 )[ animate ? "animate" : "css" ]( {
+					height: ( 100 - valPercent ) + "%"
+				}, o.animate );
+			}
+		}
+	},
+
+	_handleEvents: {
+		keydown: function( event ) {
+			var allowed, curVal, newVal, step,
+				index = $( event.target ).data( "ui-slider-handle-index" );
+
+			switch ( event.keyCode ) {
+				case $.ui.keyCode.HOME:
+				case $.ui.keyCode.END:
+				case $.ui.keyCode.PAGE_UP:
+				case $.ui.keyCode.PAGE_DOWN:
+				case $.ui.keyCode.UP:
+				case $.ui.keyCode.RIGHT:
+				case $.ui.keyCode.DOWN:
+				case $.ui.keyCode.LEFT:
+					event.preventDefault();
+					if ( !this._keySliding ) {
+						this._keySliding = true;
+						this._addClass( $( event.target ), null, "ui-state-active" );
+						allowed = this._start( event, index );
+						if ( allowed === false ) {
+							return;
+						}
+					}
+					break;
+			}
+
+			step = this.options.step;
+			if ( this._hasMultipleValues() ) {
+				curVal = newVal = this.values( index );
+			} else {
+				curVal = newVal = this.value();
+			}
+
+			switch ( event.keyCode ) {
+				case $.ui.keyCode.HOME:
+					newVal = this._valueMin();
+					break;
+				case $.ui.keyCode.END:
+					newVal = this._valueMax();
+					break;
+				case $.ui.keyCode.PAGE_UP:
+					newVal = this._trimAlignValue(
+						curVal + ( ( this._valueMax() - this._valueMin() ) / this.numPages )
+					);
+					break;
+				case $.ui.keyCode.PAGE_DOWN:
+					newVal = this._trimAlignValue(
+						curVal - ( ( this._valueMax() - this._valueMin() ) / this.numPages ) );
+					break;
+				case $.ui.keyCode.UP:
+				case $.ui.keyCode.RIGHT:
+					if ( curVal === this._valueMax() ) {
+						return;
+					}
+					newVal = this._trimAlignValue( curVal + step );
+					break;
+				case $.ui.keyCode.DOWN:
+				case $.ui.keyCode.LEFT:
+					if ( curVal === this._valueMin() ) {
+						return;
+					}
+					newVal = this._trimAlignValue( curVal - step );
+					break;
+			}
+
+			this._slide( event, index, newVal );
+		},
+		keyup: function( event ) {
+			var index = $( event.target ).data( "ui-slider-handle-index" );
+
+			if ( this._keySliding ) {
+				this._keySliding = false;
+				this._stop( event, index );
+				this._change( event, index );
+				this._removeClass( $( event.target ), null, "ui-state-active" );
+			}
+		}
+	}
+} );
+
+
+/*!
+ * jQuery UI Sortable 1.12.0
+ * http://jqueryui.com
+ *
+ * Copyright jQuery Foundation and other contributors
+ * Released under the MIT license.
+ * http://jquery.org/license
+ */
+
+//>>label: Sortable
+//>>group: Interactions
+//>>description: Enables items in a list to be sorted using the mouse.
+//>>docs: http://api.jqueryui.com/sortable/
+//>>demos: http://jqueryui.com/sortable/
+//>>css.structure: ../../themes/base/sortable.css
+
+
+
+var widgetsSortable = $.widget( "ui.sortable", $.ui.mouse, {
+	version: "1.12.0",
+	widgetEventPrefix: "sort",
+	ready: false,
+	options: {
+		appendTo: "parent",
+		axis: false,
+		connectWith: false,
+		containment: false,
+		cursor: "auto",
+		cursorAt: false,
+		dropOnEmpty: true,
+		forcePlaceholderSize: false,
+		forceHelperSize: false,
+		grid: false,
+		handle: false,
+		helper: "original",
+		items: "> *",
+		opacity: false,
+		placeholder: false,
+		revert: false,
+		scroll: true,
+		scrollSensitivity: 20,
+		scrollSpeed: 20,
+		scope: "default",
+		tolerance: "intersect",
+		zIndex: 1000,
+
+		// Callbacks
+		activate: null,
+		beforeStop: null,
+		change: null,
+		deactivate: null,
+		out: null,
+		over: null,
+		receive: null,
+		remove: null,
+		sort: null,
+		start: null,
+		stop: null,
+		update: null
+	},
+
+	_isOverAxis: function( x, reference, size ) {
+		return ( x >= reference ) && ( x < ( reference + size ) );
+	},
+
+	_isFloating: function( item ) {
+		return ( /left|right/ ).test( item.css( "float" ) ) ||
+			( /inline|table-cell/ ).test( item.css( "display" ) );
+	},
+
+	_create: function() {
+		this.containerCache = {};
+		this._addClass( "ui-sortable" );
+
+		//Get the items
+		this.refresh();
+
+		//Let's determine the parent's offset
+		this.offset = this.element.offset();
+
+		//Initialize mouse events for interaction
+		this._mouseInit();
+
+		this._setHandleClassName();
+
+		//We're ready to go
+		this.ready = true;
+
+	},
+
+	_setOption: function( key, value ) {
+		this._super( key, value );
+
+		if ( key === "handle" ) {
+			this._setHandleClassName();
+		}
+	},
+
+	_setHandleClassName: function() {
+		var that = this;
+		this._removeClass( this.element.find( ".ui-sortable-handle" ), "ui-sortable-handle" );
+		$.each( this.items, function() {
+			that._addClass(
+				this.instance.options.handle ?
+					this.item.find( this.instance.options.handle ) :
+					this.item,
+				"ui-sortable-handle"
+			);
+		} );
+	},
+
+	_destroy: function() {
+		this._mouseDestroy();
+
+		for ( var i = this.items.length - 1; i >= 0; i-- ) {
+			this.items[ i ].item.removeData( this.widgetName + "-item" );
+		}
+
+		return this;
+	},
+
+	_mouseCapture: function( event, overrideHandle ) {
+		var currentItem = null,
+			validHandle = false,
+			that = this;
+
+		if ( this.reverting ) {
+			return false;
+		}
+
+		if ( this.options.disabled || this.options.type === "static" ) {
+			return false;
+		}
+
+		//We have to refresh the items data once first
+		this._refreshItems( event );
+
+		//Find out if the clicked node (or one of its parents) is a actual item in this.items
+		$( event.target ).parents().each( function() {
+			if ( $.data( this, that.widgetName + "-item" ) === that ) {
+				currentItem = $( this );
+				return false;
+			}
+		} );
+		if ( $.data( event.target, that.widgetName + "-item" ) === that ) {
+			currentItem = $( event.target );
+		}
+
+		if ( !currentItem ) {
+			return false;
+		}
+		if ( this.options.handle && !overrideHandle ) {
+			$( this.options.handle, currentItem ).find( "*" ).addBack().each( function() {
+				if ( this === event.target ) {
+					validHandle = true;
+				}
+			} );
+			if ( !validHandle ) {
+				return false;
+			}
+		}
+
+		this.currentItem = currentItem;
+		this._removeCurrentsFromItems();
+		return true;
+
+	},
+
+	_mouseStart: function( event, overrideHandle, noActivation ) {
+
+		var i, body,
+			o = this.options;
+
+		this.currentContainer = this;
+
+		//We only need to call refreshPositions, because the refreshItems call has been moved to
+		// mouseCapture
+		this.refreshPositions();
+
+		//Create and append the visible helper
+		this.helper = this._createHelper( event );
+
+		//Cache the helper size
+		this._cacheHelperProportions();
+
+		/*
+		 * - Position generation -
+		 * This block generates everything position related - it's the core of draggables.
+		 */
+
+		//Cache the margins of the original element
+		this._cacheMargins();
+
+		//Get the next scrolling parent
+		this.scrollParent = this.helper.scrollParent();
+
+		//The element's absolute position on the page minus margins
+		this.offset = this.currentItem.offset();
+		this.offset = {
+			top: this.offset.top - this.margins.top,
+			left: this.offset.left - this.margins.left
+		};
+
+		$.extend( this.offset, {
+			click: { //Where the click happened, relative to the element
+				left: event.pageX - this.offset.left,
+				top: event.pageY - this.offset.top
+			},
+			parent: this._getParentOffset(),
+
+			// This is a relative to absolute position minus the actual position calculation -
+			// only used for relative positioned helper
+			relative: this._getRelativeOffset()
+		} );
+
+		// Only after we got the offset, we can change the helper's position to absolute
+		// TODO: Still need to figure out a way to make relative sorting possible
+		this.helper.css( "position", "absolute" );
+		this.cssPosition = this.helper.css( "position" );
+
+		//Generate the original position
+		this.originalPosition = this._generatePosition( event );
+		this.originalPageX = event.pageX;
+		this.originalPageY = event.pageY;
+
+		//Adjust the mouse offset relative to the helper if "cursorAt" is supplied
+		( o.cursorAt && this._adjustOffsetFromHelper( o.cursorAt ) );
+
+		//Cache the former DOM position
+		this.domPosition = {
+			prev: this.currentItem.prev()[ 0 ],
+			parent: this.currentItem.parent()[ 0 ]
+		};
+
+		// If the helper is not the original, hide the original so it's not playing any role during
+		// the drag, won't cause anything bad this way
+		if ( this.helper[ 0 ] !== this.currentItem[ 0 ] ) {
+			this.currentItem.hide();
+		}
+
+		//Create the placeholder
+		this._createPlaceholder();
+
+		//Set a containment if given in the options
+		if ( o.containment ) {
+			this._setContainment();
+		}
+
+		if ( o.cursor && o.cursor !== "auto" ) { // cursor option
+			body = this.document.find( "body" );
+
+			// Support: IE
+			this.storedCursor = body.css( "cursor" );
+			body.css( "cursor", o.cursor );
+
+			this.storedStylesheet =
+				$( "<style>*{ cursor: " + o.cursor + " !important; }</style>" ).appendTo( body );
+		}
+
+		if ( o.opacity ) { // opacity option
+			if ( this.helper.css( "opacity" ) ) {
+				this._storedOpacity = this.helper.css( "opacity" );
+			}
+			this.helper.css( "opacity", o.opacity );
+		}
+
+		if ( o.zIndex ) { // zIndex option
+			if ( this.helper.css( "zIndex" ) ) {
+				this._storedZIndex = this.helper.css( "zIndex" );
+			}
+			this.helper.css( "zIndex", o.zIndex );
+		}
+
+		//Prepare scrolling
+		if ( this.scrollParent[ 0 ] !== this.document[ 0 ] &&
+				this.scrollParent[ 0 ].tagName !== "HTML" ) {
+			this.overflowOffset = this.scrollParent.offset();
+		}
+
+		//Call callbacks
+		this._trigger( "start", event, this._uiHash() );
+
+		//Recache the helper size
+		if ( !this._preserveHelperProportions ) {
+			this._cacheHelperProportions();
+		}
+
+		//Post "activate" events to possible containers
+		if ( !noActivation ) {
+			for ( i = this.containers.length - 1; i >= 0; i-- ) {
+				this.containers[ i ]._trigger( "activate", event, this._uiHash( this ) );
+			}
+		}
+
+		//Prepare possible droppables
+		if ( $.ui.ddmanager ) {
+			$.ui.ddmanager.current = this;
+		}
+
+		if ( $.ui.ddmanager && !o.dropBehaviour ) {
+			$.ui.ddmanager.prepareOffsets( this, event );
+		}
+
+		this.dragging = true;
+
+		this._addClass( this.helper, "ui-sortable-helper" );
+
+		// Execute the drag once - this causes the helper not to be visiblebefore getting its
+		// correct position
+		this._mouseDrag( event );
+		return true;
+
+	},
+
+	_mouseDrag: function( event ) {
+		var i, item, itemElement, intersection,
+			o = this.options,
+			scrolled = false;
+
+		//Compute the helpers position
+		this.position = this._generatePosition( event );
+		this.positionAbs = this._convertPositionTo( "absolute" );
+
+		if ( !this.lastPositionAbs ) {
+			this.lastPositionAbs = this.positionAbs;
+		}
+
+		//Do scrolling
+		if ( this.options.scroll ) {
+			if ( this.scrollParent[ 0 ] !== this.document[ 0 ] &&
+					this.scrollParent[ 0 ].tagName !== "HTML" ) {
+
+				if ( ( this.overflowOffset.top + this.scrollParent[ 0 ].offsetHeight ) -
+						event.pageY < o.scrollSensitivity ) {
+					this.scrollParent[ 0 ].scrollTop =
+						scrolled = this.scrollParent[ 0 ].scrollTop + o.scrollSpeed;
+				} else if ( event.pageY - this.overflowOffset.top < o.scrollSensitivity ) {
+					this.scrollParent[ 0 ].scrollTop =
+						scrolled = this.scrollParent[ 0 ].scrollTop - o.scrollSpeed;
+				}
+
+				if ( ( this.overflowOffset.left + this.scrollParent[ 0 ].offsetWidth ) -
+						event.pageX < o.scrollSensitivity ) {
+					this.scrollParent[ 0 ].scrollLeft = scrolled =
+						this.scrollParent[ 0 ].scrollLeft + o.scrollSpeed;
+				} else if ( event.pageX - this.overflowOffset.left < o.scrollSensitivity ) {
+					this.scrollParent[ 0 ].scrollLeft = scrolled =
+						this.scrollParent[ 0 ].scrollLeft - o.scrollSpeed;
+				}
+
+			} else {
+
+				if ( event.pageY - this.document.scrollTop() < o.scrollSensitivity ) {
+					scrolled = this.document.scrollTop( this.document.scrollTop() - o.scrollSpeed );
+				} else if ( this.window.height() - ( event.pageY - this.document.scrollTop() ) <
+						o.scrollSensitivity ) {
+					scrolled = this.document.scrollTop( this.document.scrollTop() + o.scrollSpeed );
+				}
+
+				if ( event.pageX - this.document.scrollLeft() < o.scrollSensitivity ) {
+					scrolled = this.document.scrollLeft(
+						this.document.scrollLeft() - o.scrollSpeed
+					);
+				} else if ( this.window.width() - ( event.pageX - this.document.scrollLeft() ) <
+						o.scrollSensitivity ) {
+					scrolled = this.document.scrollLeft(
+						this.document.scrollLeft() + o.scrollSpeed
+					);
+				}
+
+			}
+
+			if ( scrolled !== false && $.ui.ddmanager && !o.dropBehaviour ) {
+				$.ui.ddmanager.prepareOffsets( this, event );
+			}
+		}
+
+		//Regenerate the absolute position used for position checks
+		this.positionAbs = this._convertPositionTo( "absolute" );
+
+		//Set the helper position
+		if ( !this.options.axis || this.options.axis !== "y" ) {
+			this.helper[ 0 ].style.left = this.position.left + "px";
+		}
+		if ( !this.options.axis || this.options.axis !== "x" ) {
+			this.helper[ 0 ].style.top = this.position.top + "px";
+		}
+
+		//Rearrange
+		for ( i = this.items.length - 1; i >= 0; i-- ) {
+
+			//Cache variables and intersection, continue if no intersection
+			item = this.items[ i ];
+			itemElement = item.item[ 0 ];
+			intersection = this._intersectsWithPointer( item );
+			if ( !intersection ) {
+				continue;
+			}
+
+			// Only put the placeholder inside the current Container, skip all
+			// items from other containers. This works because when moving
+			// an item from one container to another the
+			// currentContainer is switched before the placeholder is moved.
+			//
+			// Without this, moving items in "sub-sortables" can cause
+			// the placeholder to jitter between the outer and inner container.
+			if ( item.instance !== this.currentContainer ) {
+				continue;
+			}
+
+			// Cannot intersect with itself
+			// no useless actions that have been done before
+			// no action if the item moved is the parent of the item checked
+			if ( itemElement !== this.currentItem[ 0 ] &&
+				this.placeholder[ intersection === 1 ? "next" : "prev" ]()[ 0 ] !== itemElement &&
+				!$.contains( this.placeholder[ 0 ], itemElement ) &&
+				( this.options.type === "semi-dynamic" ?
+					!$.contains( this.element[ 0 ], itemElement ) :
+					true
+				)
+			) {
+
+				this.direction = intersection === 1 ? "down" : "up";
+
+				if ( this.options.tolerance === "pointer" || this._intersectsWithSides( item ) ) {
+					this._rearrange( event, item );
+				} else {
+					break;
+				}
+
+				this._trigger( "change", event, this._uiHash() );
+				break;
+			}
+		}
+
+		//Post events to containers
+		this._contactContainers( event );
+
+		//Interconnect with droppables
+		if ( $.ui.ddmanager ) {
+			$.ui.ddmanager.drag( this, event );
+		}
+
+		//Call callbacks
+		this._trigger( "sort", event, this._uiHash() );
+
+		this.lastPositionAbs = this.positionAbs;
+		return false;
+
+	},
+
+	_mouseStop: function( event, noPropagation ) {
+
+		if ( !event ) {
+			return;
+		}
+
+		//If we are using droppables, inform the manager about the drop
+		if ( $.ui.ddmanager && !this.options.dropBehaviour ) {
+			$.ui.ddmanager.drop( this, event );
+		}
+
+		if ( this.options.revert ) {
+			var that = this,
+				cur = this.placeholder.offset(),
+				axis = this.options.axis,
+				animation = {};
+
+			if ( !axis || axis === "x" ) {
+				animation.left = cur.left - this.offset.parent.left - this.margins.left +
+					( this.offsetParent[ 0 ] === this.document[ 0 ].body ?
+						0 :
+						this.offsetParent[ 0 ].scrollLeft
+					);
+			}
+			if ( !axis || axis === "y" ) {
+				animation.top = cur.top - this.offset.parent.top - this.margins.top +
+					( this.offsetParent[ 0 ] === this.document[ 0 ].body ?
+						0 :
+						this.offsetParent[ 0 ].scrollTop
+					);
+			}
+			this.reverting = true;
+			$( this.helper ).animate(
+				animation,
+				parseInt( this.options.revert, 10 ) || 500,
+				function() {
+					that._clear( event );
+				}
+			);
+		} else {
+			this._clear( event, noPropagation );
+		}
+
+		return false;
+
+	},
+
+	cancel: function() {
+
+		if ( this.dragging ) {
+
+			this._mouseUp( { target: null } );
+
+			if ( this.options.helper === "original" ) {
+				this.currentItem.css( this._storedCSS );
+				this._removeClass( this.currentItem, "ui-sortable-helper" );
+			} else {
+				this.currentItem.show();
+			}
+
+			//Post deactivating events to containers
+			for ( var i = this.containers.length - 1; i >= 0; i-- ) {
+				this.containers[ i ]._trigger( "deactivate", null, this._uiHash( this ) );
+				if ( this.containers[ i ].containerCache.over ) {
+					this.containers[ i ]._trigger( "out", null, this._uiHash( this ) );
+					this.containers[ i ].containerCache.over = 0;
+				}
+			}
+
+		}
+
+		if ( this.placeholder ) {
+
+			//$(this.placeholder[0]).remove(); would have been the jQuery way - unfortunately,
+			// it unbinds ALL events from the original node!
+			if ( this.placeholder[ 0 ].parentNode ) {
+				this.placeholder[ 0 ].parentNode.removeChild( this.placeholder[ 0 ] );
+			}
+			if ( this.options.helper !== "original" && this.helper &&
+					this.helper[ 0 ].parentNode ) {
+				this.helper.remove();
+			}
+
+			$.extend( this, {
+				helper: null,
+				dragging: false,
+				reverting: false,
+				_noFinalSort: null
+			} );
+
+			if ( this.domPosition.prev ) {
+				$( this.domPosition.prev ).after( this.currentItem );
+			} else {
+				$( this.domPosition.parent ).prepend( this.currentItem );
+			}
+		}
+
+		return this;
+
+	},
+
+	serialize: function( o ) {
+
+		var items = this._getItemsAsjQuery( o && o.connected ),
+			str = [];
+		o = o || {};
+
+		$( items ).each( function() {
+			var res = ( $( o.item || this ).attr( o.attribute || "id" ) || "" )
+				.match( o.expression || ( /(.+)[\-=_](.+)/ ) );
+			if ( res ) {
+				str.push(
+					( o.key || res[ 1 ] + "[]" ) +
+					"=" + ( o.key && o.expression ? res[ 1 ] : res[ 2 ] ) );
+			}
+		} );
+
+		if ( !str.length && o.key ) {
+			str.push( o.key + "=" );
+		}
+
+		return str.join( "&" );
+
+	},
+
+	toArray: function( o ) {
+
+		var items = this._getItemsAsjQuery( o && o.connected ),
+			ret = [];
+
+		o = o || {};
+
+		items.each( function() {
+			ret.push( $( o.item || this ).attr( o.attribute || "id" ) || "" );
+		} );
+		return ret;
+
+	},
+
+	/* Be careful with the following core functions */
+	_intersectsWith: function( item ) {
+
+		var x1 = this.positionAbs.left,
+			x2 = x1 + this.helperProportions.width,
+			y1 = this.positionAbs.top,
+			y2 = y1 + this.helperProportions.height,
+			l = item.left,
+			r = l + item.width,
+			t = item.top,
+			b = t + item.height,
+			dyClick = this.offset.click.top,
+			dxClick = this.offset.click.left,
+			isOverElementHeight = ( this.options.axis === "x" ) || ( ( y1 + dyClick ) > t &&
+				( y1 + dyClick ) < b ),
+			isOverElementWidth = ( this.options.axis === "y" ) || ( ( x1 + dxClick ) > l &&
+				( x1 + dxClick ) < r ),
+			isOverElement = isOverElementHeight && isOverElementWidth;
+
+		if ( this.options.tolerance === "pointer" ||
+			this.options.forcePointerForContainers ||
+			( this.options.tolerance !== "pointer" &&
+				this.helperProportions[ this.floating ? "width" : "height" ] >
+				item[ this.floating ? "width" : "height" ] )
+		) {
+			return isOverElement;
+		} else {
+
+			return ( l < x1 + ( this.helperProportions.width / 2 ) && // Right Half
+				x2 - ( this.helperProportions.width / 2 ) < r && // Left Half
+				t < y1 + ( this.helperProportions.height / 2 ) && // Bottom Half
+				y2 - ( this.helperProportions.height / 2 ) < b ); // Top Half
+
+		}
+	},
+
+	_intersectsWithPointer: function( item ) {
+		var verticalDirection, horizontalDirection,
+			isOverElementHeight = ( this.options.axis === "x" ) ||
+				this._isOverAxis(
+					this.positionAbs.top + this.offset.click.top, item.top, item.height ),
+			isOverElementWidth = ( this.options.axis === "y" ) ||
+				this._isOverAxis(
+					this.positionAbs.left + this.offset.click.left, item.left, item.width ),
+			isOverElement = isOverElementHeight && isOverElementWidth;
+
+		if ( !isOverElement ) {
+			return false;
+		}
+
+		verticalDirection = this._getDragVerticalDirection();
+		horizontalDirection = this._getDragHorizontalDirection();
+
+		return this.floating ?
+			( ( horizontalDirection === "right" || verticalDirection === "down" ) ? 2 : 1 )
+			: ( verticalDirection && ( verticalDirection === "down" ? 2 : 1 ) );
+
+	},
+
+	_intersectsWithSides: function( item ) {
+
+		var isOverBottomHalf = this._isOverAxis( this.positionAbs.top +
+				this.offset.click.top, item.top + ( item.height / 2 ), item.height ),
+			isOverRightHalf = this._isOverAxis( this.positionAbs.left +
+				this.offset.click.left, item.left + ( item.width / 2 ), item.width ),
+			verticalDirection = this._getDragVerticalDirection(),
+			horizontalDirection = this._getDragHorizontalDirection();
+
+		if ( this.floating && horizontalDirection ) {
+			return ( ( horizontalDirection === "right" && isOverRightHalf ) ||
+				( horizontalDirection === "left" && !isOverRightHalf ) );
+		} else {
+			return verticalDirection && ( ( verticalDirection === "down" && isOverBottomHalf ) ||
+				( verticalDirection === "up" && !isOverBottomHalf ) );
+		}
+
+	},
+
+	_getDragVerticalDirection: function() {
+		var delta = this.positionAbs.top - this.lastPositionAbs.top;
+		return delta !== 0 && ( delta > 0 ? "down" : "up" );
+	},
+
+	_getDragHorizontalDirection: function() {
+		var delta = this.positionAbs.left - this.lastPositionAbs.left;
+		return delta !== 0 && ( delta > 0 ? "right" : "left" );
+	},
+
+	refresh: function( event ) {
+		this._refreshItems( event );
+		this._setHandleClassName();
+		this.refreshPositions();
+		return this;
+	},
+
+	_connectWith: function() {
+		var options = this.options;
+		return options.connectWith.constructor === String ?
+			[ options.connectWith ] :
+			options.connectWith;
+	},
+
+	_getItemsAsjQuery: function( connected ) {
+
+		var i, j, cur, inst,
+			items = [],
+			queries = [],
+			connectWith = this._connectWith();
+
+		if ( connectWith && connected ) {
+			for ( i = connectWith.length - 1; i >= 0; i-- ) {
+				cur = $( connectWith[ i ], this.document[ 0 ] );
+				for ( j = cur.length - 1; j >= 0; j-- ) {
+					inst = $.data( cur[ j ], this.widgetFullName );
+					if ( inst && inst !== this && !inst.options.disabled ) {
+						queries.push( [ $.isFunction( inst.options.items ) ?
+							inst.options.items.call( inst.element ) :
+							$( inst.options.items, inst.element )
+								.not( ".ui-sortable-helper" )
+								.not( ".ui-sortable-placeholder" ), inst ] );
+					}
+				}
+			}
+		}
+
+		queries.push( [ $.isFunction( this.options.items ) ?
+			this.options.items
+				.call( this.element, null, { options: this.options, item: this.currentItem } ) :
+			$( this.options.items, this.element )
+				.not( ".ui-sortable-helper" )
+				.not( ".ui-sortable-placeholder" ), this ] );
+
+		function addItems() {
+			items.push( this );
+		}
+		for ( i = queries.length - 1; i >= 0; i-- ) {
+			queries[ i ][ 0 ].each( addItems );
+		}
+
+		return $( items );
+
+	},
+
+	_removeCurrentsFromItems: function() {
+
+		var list = this.currentItem.find( ":data(" + this.widgetName + "-item)" );
+
+		this.items = $.grep( this.items, function( item ) {
+			for ( var j = 0; j < list.length; j++ ) {
+				if ( list[ j ] === item.item[ 0 ] ) {
+					return false;
+				}
+			}
+			return true;
+		} );
+
+	},
+
+	_refreshItems: function( event ) {
+
+		this.items = [];
+		this.containers = [ this ];
+
+		var i, j, cur, inst, targetData, _queries, item, queriesLength,
+			items = this.items,
+			queries = [ [ $.isFunction( this.options.items ) ?
+				this.options.items.call( this.element[ 0 ], event, { item: this.currentItem } ) :
+				$( this.options.items, this.element ), this ] ],
+			connectWith = this._connectWith();
+
+		//Shouldn't be run the first time through due to massive slow-down
+		if ( connectWith && this.ready ) {
+			for ( i = connectWith.length - 1; i >= 0; i-- ) {
+				cur = $( connectWith[ i ], this.document[ 0 ] );
+				for ( j = cur.length - 1; j >= 0; j-- ) {
+					inst = $.data( cur[ j ], this.widgetFullName );
+					if ( inst && inst !== this && !inst.options.disabled ) {
+						queries.push( [ $.isFunction( inst.options.items ) ?
+							inst.options.items
+								.call( inst.element[ 0 ], event, { item: this.currentItem } ) :
+							$( inst.options.items, inst.element ), inst ] );
+						this.containers.push( inst );
+					}
+				}
+			}
+		}
+
+		for ( i = queries.length - 1; i >= 0; i-- ) {
+			targetData = queries[ i ][ 1 ];
+			_queries = queries[ i ][ 0 ];
+
+			for ( j = 0, queriesLength = _queries.length; j < queriesLength; j++ ) {
+				item = $( _queries[ j ] );
+
+				// Data for target checking (mouse manager)
+				item.data( this.widgetName + "-item", targetData );
+
+				items.push( {
+					item: item,
+					instance: targetData,
+					width: 0, height: 0,
+					left: 0, top: 0
+				} );
+			}
+		}
+
+	},
+
+	refreshPositions: function( fast ) {
+
+		// Determine whether items are being displayed horizontally
+		this.floating = this.items.length ?
+			this.options.axis === "x" || this._isFloating( this.items[ 0 ].item ) :
+			false;
+
+		//This has to be redone because due to the item being moved out/into the offsetParent,
+		// the offsetParent's position will change
+		if ( this.offsetParent && this.helper ) {
+			this.offset.parent = this._getParentOffset();
+		}
+
+		var i, item, t, p;
+
+		for ( i = this.items.length - 1; i >= 0; i-- ) {
+			item = this.items[ i ];
+
+			//We ignore calculating positions of all connected containers when we're not over them
+			if ( item.instance !== this.currentContainer && this.currentContainer &&
+					item.item[ 0 ] !== this.currentItem[ 0 ] ) {
+				continue;
+			}
+
+			t = this.options.toleranceElement ?
+				$( this.options.toleranceElement, item.item ) :
+				item.item;
+
+			if ( !fast ) {
+				item.width = t.outerWidth();
+				item.height = t.outerHeight();
+			}
+
+			p = t.offset();
+			item.left = p.left;
+			item.top = p.top;
+		}
+
+		if ( this.options.custom && this.options.custom.refreshContainers ) {
+			this.options.custom.refreshContainers.call( this );
+		} else {
+			for ( i = this.containers.length - 1; i >= 0; i-- ) {
+				p = this.containers[ i ].element.offset();
+				this.containers[ i ].containerCache.left = p.left;
+				this.containers[ i ].containerCache.top = p.top;
+				this.containers[ i ].containerCache.width =
+					this.containers[ i ].element.outerWidth();
+				this.containers[ i ].containerCache.height =
+					this.containers[ i ].element.outerHeight();
+			}
+		}
+
+		return this;
+	},
+
+	_createPlaceholder: function( that ) {
+		that = that || this;
+		var className,
+			o = that.options;
+
+		if ( !o.placeholder || o.placeholder.constructor === String ) {
+			className = o.placeholder;
+			o.placeholder = {
+				element: function() {
+
+					var nodeName = that.currentItem[ 0 ].nodeName.toLowerCase(),
+						element = $( "<" + nodeName + ">", that.document[ 0 ] );
+
+						that._addClass( element, "ui-sortable-placeholder",
+								className || that.currentItem[ 0 ].className )
+							._removeClass( element, "ui-sortable-helper" );
+
+					if ( nodeName === "tbody" ) {
+						that._createTrPlaceholder(
+							that.currentItem.find( "tr" ).eq( 0 ),
+							$( "<tr>", that.document[ 0 ] ).appendTo( element )
+						);
+					} else if ( nodeName === "tr" ) {
+						that._createTrPlaceholder( that.currentItem, element );
+					} else if ( nodeName === "img" ) {
+						element.attr( "src", that.currentItem.attr( "src" ) );
+					}
+
+					if ( !className ) {
+						element.css( "visibility", "hidden" );
+					}
+
+					return element;
+				},
+				update: function( container, p ) {
+
+					// 1. If a className is set as 'placeholder option, we don't force sizes -
+					// the class is responsible for that
+					// 2. The option 'forcePlaceholderSize can be enabled to force it even if a
+					// class name is specified
+					if ( className && !o.forcePlaceholderSize ) {
+						return;
+					}
+
+					//If the element doesn't have a actual height by itself (without styles coming
+					// from a stylesheet), it receives the inline height from the dragged item
+					if ( !p.height() ) {
+						p.height(
+							that.currentItem.innerHeight() -
+							parseInt( that.currentItem.css( "paddingTop" ) || 0, 10 ) -
+							parseInt( that.currentItem.css( "paddingBottom" ) || 0, 10 ) );
+					}
+					if ( !p.width() ) {
+						p.width(
+							that.currentItem.innerWidth() -
+							parseInt( that.currentItem.css( "paddingLeft" ) || 0, 10 ) -
+							parseInt( that.currentItem.css( "paddingRight" ) || 0, 10 ) );
+					}
+				}
+			};
+		}
+
+		//Create the placeholder
+		that.placeholder = $( o.placeholder.element.call( that.element, that.currentItem ) );
+
+		//Append it after the actual current item
+		that.currentItem.after( that.placeholder );
+
+		//Update the size of the placeholder (TODO: Logic to fuzzy, see line 316/317)
+		o.placeholder.update( that, that.placeholder );
+
+	},
+
+	_createTrPlaceholder: function( sourceTr, targetTr ) {
+		var that = this;
+
+		sourceTr.children().each( function() {
+			$( "<td>&#160;</td>", that.document[ 0 ] )
+				.attr( "colspan", $( this ).attr( "colspan" ) || 1 )
+				.appendTo( targetTr );
+		} );
+	},
+
+	_contactContainers: function( event ) {
+		var i, j, dist, itemWithLeastDistance, posProperty, sizeProperty, cur, nearBottom,
+			floating, axis,
+			innermostContainer = null,
+			innermostIndex = null;
+
+		// Get innermost container that intersects with item
+		for ( i = this.containers.length - 1; i >= 0; i-- ) {
+
+			// Never consider a container that's located within the item itself
+			if ( $.contains( this.currentItem[ 0 ], this.containers[ i ].element[ 0 ] ) ) {
+				continue;
+			}
+
+			if ( this._intersectsWith( this.containers[ i ].containerCache ) ) {
+
+				// If we've already found a container and it's more "inner" than this, then continue
+				if ( innermostContainer &&
+						$.contains(
+							this.containers[ i ].element[ 0 ],
+							innermostContainer.element[ 0 ] ) ) {
+					continue;
+				}
+
+				innermostContainer = this.containers[ i ];
+				innermostIndex = i;
+
+			} else {
+
+				// container doesn't intersect. trigger "out" event if necessary
+				if ( this.containers[ i ].containerCache.over ) {
+					this.containers[ i ]._trigger( "out", event, this._uiHash( this ) );
+					this.containers[ i ].containerCache.over = 0;
+				}
+			}
+
+		}
+
+		// If no intersecting containers found, return
+		if ( !innermostContainer ) {
+			return;
+		}
+
+		// Move the item into the container if it's not there already
+		if ( this.containers.length === 1 ) {
+			if ( !this.containers[ innermostIndex ].containerCache.over ) {
+				this.containers[ innermostIndex ]._trigger( "over", event, this._uiHash( this ) );
+				this.containers[ innermostIndex ].containerCache.over = 1;
+			}
+		} else {
+
+			// When entering a new container, we will find the item with the least distance and
+			// append our item near it
+			dist = 10000;
+			itemWithLeastDistance = null;
+			floating = innermostContainer.floating || this._isFloating( this.currentItem );
+			posProperty = floating ? "left" : "top";
+			sizeProperty = floating ? "width" : "height";
+			axis = floating ? "pageX" : "pageY";
+
+			for ( j = this.items.length - 1; j >= 0; j-- ) {
+				if ( !$.contains(
+						this.containers[ innermostIndex ].element[ 0 ], this.items[ j ].item[ 0 ] )
+				) {
+					continue;
+				}
+				if ( this.items[ j ].item[ 0 ] === this.currentItem[ 0 ] ) {
+					continue;
+				}
+
+				cur = this.items[ j ].item.offset()[ posProperty ];
+				nearBottom = false;
+				if ( event[ axis ] - cur > this.items[ j ][ sizeProperty ] / 2 ) {
+					nearBottom = true;
+				}
+
+				if ( Math.abs( event[ axis ] - cur ) < dist ) {
+					dist = Math.abs( event[ axis ] - cur );
+					itemWithLeastDistance = this.items[ j ];
+					this.direction = nearBottom ? "up" : "down";
+				}
+			}
+
+			//Check if dropOnEmpty is enabled
+			if ( !itemWithLeastDistance && !this.options.dropOnEmpty ) {
+				return;
+			}
+
+			if ( this.currentContainer === this.containers[ innermostIndex ] ) {
+				if ( !this.currentContainer.containerCache.over ) {
+					this.containers[ innermostIndex ]._trigger( "over", event, this._uiHash() );
+					this.currentContainer.containerCache.over = 1;
+				}
+				return;
+			}
+
+			itemWithLeastDistance ?
+				this._rearrange( event, itemWithLeastDistance, null, true ) :
+				this._rearrange( event, null, this.containers[ innermostIndex ].element, true );
+			this._trigger( "change", event, this._uiHash() );
+			this.containers[ innermostIndex ]._trigger( "change", event, this._uiHash( this ) );
+			this.currentContainer = this.containers[ innermostIndex ];
+
+			//Update the placeholder
+			this.options.placeholder.update( this.currentContainer, this.placeholder );
+
+			this.containers[ innermostIndex ]._trigger( "over", event, this._uiHash( this ) );
+			this.containers[ innermostIndex ].containerCache.over = 1;
+		}
+
+	},
+
+	_createHelper: function( event ) {
+
+		var o = this.options,
+			helper = $.isFunction( o.helper ) ?
+				$( o.helper.apply( this.element[ 0 ], [ event, this.currentItem ] ) ) :
+				( o.helper === "clone" ? this.currentItem.clone() : this.currentItem );
+
+		//Add the helper to the DOM if that didn't happen already
+		if ( !helper.parents( "body" ).length ) {
+			$( o.appendTo !== "parent" ?
+				o.appendTo :
+				this.currentItem[ 0 ].parentNode )[ 0 ].appendChild( helper[ 0 ] );
+		}
+
+		if ( helper[ 0 ] === this.currentItem[ 0 ] ) {
+			this._storedCSS = {
+				width: this.currentItem[ 0 ].style.width,
+				height: this.currentItem[ 0 ].style.height,
+				position: this.currentItem.css( "position" ),
+				top: this.currentItem.css( "top" ),
+				left: this.currentItem.css( "left" )
+			};
+		}
+
+		if ( !helper[ 0 ].style.width || o.forceHelperSize ) {
+			helper.width( this.currentItem.width() );
+		}
+		if ( !helper[ 0 ].style.height || o.forceHelperSize ) {
+			helper.height( this.currentItem.height() );
+		}
+
+		return helper;
+
+	},
+
+	_adjustOffsetFromHelper: function( obj ) {
+		if ( typeof obj === "string" ) {
+			obj = obj.split( " " );
+		}
+		if ( $.isArray( obj ) ) {
+			obj = { left: +obj[ 0 ], top: +obj[ 1 ] || 0 };
+		}
+		if ( "left" in obj ) {
+			this.offset.click.left = obj.left + this.margins.left;
+		}
+		if ( "right" in obj ) {
+			this.offset.click.left = this.helperProportions.width - obj.right + this.margins.left;
+		}
+		if ( "top" in obj ) {
+			this.offset.click.top = obj.top + this.margins.top;
+		}
+		if ( "bottom" in obj ) {
+			this.offset.click.top = this.helperProportions.height - obj.bottom + this.margins.top;
+		}
+	},
+
+	_getParentOffset: function() {
+
+		//Get the offsetParent and cache its position
+		this.offsetParent = this.helper.offsetParent();
+		var po = this.offsetParent.offset();
+
+		// This is a special case where we need to modify a offset calculated on start, since the
+		// following happened:
+		// 1. The position of the helper is absolute, so it's position is calculated based on the
+		// next positioned parent
+		// 2. The actual offset parent is a child of the scroll parent, and the scroll parent isn't
+		// the document, which means that the scroll is included in the initial calculation of the
+		// offset of the parent, and never recalculated upon drag
+		if ( this.cssPosition === "absolute" && this.scrollParent[ 0 ] !== this.document[ 0 ] &&
+				$.contains( this.scrollParent[ 0 ], this.offsetParent[ 0 ] ) ) {
+			po.left += this.scrollParent.scrollLeft();
+			po.top += this.scrollParent.scrollTop();
+		}
+
+		// This needs to be actually done for all browsers, since pageX/pageY includes this
+		// information with an ugly IE fix
+		if ( this.offsetParent[ 0 ] === this.document[ 0 ].body ||
+				( this.offsetParent[ 0 ].tagName &&
+				this.offsetParent[ 0 ].tagName.toLowerCase() === "html" && $.ui.ie ) ) {
+			po = { top: 0, left: 0 };
+		}
+
+		return {
+			top: po.top + ( parseInt( this.offsetParent.css( "borderTopWidth" ), 10 ) || 0 ),
+			left: po.left + ( parseInt( this.offsetParent.css( "borderLeftWidth" ), 10 ) || 0 )
+		};
+
+	},
+
+	_getRelativeOffset: function() {
+
+		if ( this.cssPosition === "relative" ) {
+			var p = this.currentItem.position();
+			return {
+				top: p.top - ( parseInt( this.helper.css( "top" ), 10 ) || 0 ) +
+					this.scrollParent.scrollTop(),
+				left: p.left - ( parseInt( this.helper.css( "left" ), 10 ) || 0 ) +
+					this.scrollParent.scrollLeft()
+			};
+		} else {
+			return { top: 0, left: 0 };
+		}
+
+	},
+
+	_cacheMargins: function() {
+		this.margins = {
+			left: ( parseInt( this.currentItem.css( "marginLeft" ), 10 ) || 0 ),
+			top: ( parseInt( this.currentItem.css( "marginTop" ), 10 ) || 0 )
+		};
+	},
+
+	_cacheHelperProportions: function() {
+		this.helperProportions = {
+			width: this.helper.outerWidth(),
+			height: this.helper.outerHeight()
+		};
+	},
+
+	_setContainment: function() {
+
+		var ce, co, over,
+			o = this.options;
+		if ( o.containment === "parent" ) {
+			o.containment = this.helper[ 0 ].parentNode;
+		}
+		if ( o.containment === "document" || o.containment === "window" ) {
+			this.containment = [
+				0 - this.offset.relative.left - this.offset.parent.left,
+				0 - this.offset.relative.top - this.offset.parent.top,
+				o.containment === "document" ?
+					this.document.width() :
+					this.window.width() - this.helperProportions.width - this.margins.left,
+				( o.containment === "document" ?
+					( this.document.height() || document.body.parentNode.scrollHeight ) :
+					this.window.height() || this.document[ 0 ].body.parentNode.scrollHeight
+				) - this.helperProportions.height - this.margins.top
+			];
+		}
+
+		if ( !( /^(document|window|parent)$/ ).test( o.containment ) ) {
+			ce = $( o.containment )[ 0 ];
+			co = $( o.containment ).offset();
+			over = ( $( ce ).css( "overflow" ) !== "hidden" );
+
+			this.containment = [
+				co.left + ( parseInt( $( ce ).css( "borderLeftWidth" ), 10 ) || 0 ) +
+					( parseInt( $( ce ).css( "paddingLeft" ), 10 ) || 0 ) - this.margins.left,
+				co.top + ( parseInt( $( ce ).css( "borderTopWidth" ), 10 ) || 0 ) +
+					( parseInt( $( ce ).css( "paddingTop" ), 10 ) || 0 ) - this.margins.top,
+				co.left + ( over ? Math.max( ce.scrollWidth, ce.offsetWidth ) : ce.offsetWidth ) -
+					( parseInt( $( ce ).css( "borderLeftWidth" ), 10 ) || 0 ) -
+					( parseInt( $( ce ).css( "paddingRight" ), 10 ) || 0 ) -
+					this.helperProportions.width - this.margins.left,
+				co.top + ( over ? Math.max( ce.scrollHeight, ce.offsetHeight ) : ce.offsetHeight ) -
+					( parseInt( $( ce ).css( "borderTopWidth" ), 10 ) || 0 ) -
+					( parseInt( $( ce ).css( "paddingBottom" ), 10 ) || 0 ) -
+					this.helperProportions.height - this.margins.top
+			];
+		}
+
+	},
+
+	_convertPositionTo: function( d, pos ) {
+
+		if ( !pos ) {
+			pos = this.position;
+		}
+		var mod = d === "absolute" ? 1 : -1,
+			scroll = this.cssPosition === "absolute" &&
+				!( this.scrollParent[ 0 ] !== this.document[ 0 ] &&
+				$.contains( this.scrollParent[ 0 ], this.offsetParent[ 0 ] ) ) ?
+					this.offsetParent :
+					this.scrollParent,
+			scrollIsRootNode = ( /(html|body)/i ).test( scroll[ 0 ].tagName );
+
+		return {
+			top: (
+
+				// The absolute mouse position
+				pos.top	+
+
+				// Only for relative positioned nodes: Relative offset from element to offset parent
+				this.offset.relative.top * mod +
+
+				// The offsetParent's offset without borders (offset + border)
+				this.offset.parent.top * mod -
+				( ( this.cssPosition === "fixed" ?
+					-this.scrollParent.scrollTop() :
+					( scrollIsRootNode ? 0 : scroll.scrollTop() ) ) * mod )
+			),
+			left: (
+
+				// The absolute mouse position
+				pos.left +
+
+				// Only for relative positioned nodes: Relative offset from element to offset parent
+				this.offset.relative.left * mod +
+
+				// The offsetParent's offset without borders (offset + border)
+				this.offset.parent.left * mod	-
+				( ( this.cssPosition === "fixed" ?
+					-this.scrollParent.scrollLeft() : scrollIsRootNode ? 0 :
+					scroll.scrollLeft() ) * mod )
+			)
+		};
+
+	},
+
+	_generatePosition: function( event ) {
+
+		var top, left,
+			o = this.options,
+			pageX = event.pageX,
+			pageY = event.pageY,
+			scroll = this.cssPosition === "absolute" &&
+				!( this.scrollParent[ 0 ] !== this.document[ 0 ] &&
+				$.contains( this.scrollParent[ 0 ], this.offsetParent[ 0 ] ) ) ?
+					this.offsetParent :
+					this.scrollParent,
+				scrollIsRootNode = ( /(html|body)/i ).test( scroll[ 0 ].tagName );
+
+		// This is another very weird special case that only happens for relative elements:
+		// 1. If the css position is relative
+		// 2. and the scroll parent is the document or similar to the offset parent
+		// we have to refresh the relative offset during the scroll so there are no jumps
+		if ( this.cssPosition === "relative" && !( this.scrollParent[ 0 ] !== this.document[ 0 ] &&
+				this.scrollParent[ 0 ] !== this.offsetParent[ 0 ] ) ) {
+			this.offset.relative = this._getRelativeOffset();
+		}
+
+		/*
+		 * - Position constraining -
+		 * Constrain the position to a mix of grid, containment.
+		 */
+
+		if ( this.originalPosition ) { //If we are not dragging yet, we won't check for options
+
+			if ( this.containment ) {
+				if ( event.pageX - this.offset.click.left < this.containment[ 0 ] ) {
+					pageX = this.containment[ 0 ] + this.offset.click.left;
+				}
+				if ( event.pageY - this.offset.click.top < this.containment[ 1 ] ) {
+					pageY = this.containment[ 1 ] + this.offset.click.top;
+				}
+				if ( event.pageX - this.offset.click.left > this.containment[ 2 ] ) {
+					pageX = this.containment[ 2 ] + this.offset.click.left;
+				}
+				if ( event.pageY - this.offset.click.top > this.containment[ 3 ] ) {
+					pageY = this.containment[ 3 ] + this.offset.click.top;
+				}
+			}
+
+			if ( o.grid ) {
+				top = this.originalPageY + Math.round( ( pageY - this.originalPageY ) /
+					o.grid[ 1 ] ) * o.grid[ 1 ];
+				pageY = this.containment ?
+					( ( top - this.offset.click.top >= this.containment[ 1 ] &&
+						top - this.offset.click.top <= this.containment[ 3 ] ) ?
+							top :
+							( ( top - this.offset.click.top >= this.containment[ 1 ] ) ?
+								top - o.grid[ 1 ] : top + o.grid[ 1 ] ) ) :
+								top;
+
+				left = this.originalPageX + Math.round( ( pageX - this.originalPageX ) /
+					o.grid[ 0 ] ) * o.grid[ 0 ];
+				pageX = this.containment ?
+					( ( left - this.offset.click.left >= this.containment[ 0 ] &&
+						left - this.offset.click.left <= this.containment[ 2 ] ) ?
+							left :
+							( ( left - this.offset.click.left >= this.containment[ 0 ] ) ?
+								left - o.grid[ 0 ] : left + o.grid[ 0 ] ) ) :
+								left;
+			}
+
+		}
+
+		return {
+			top: (
+
+				// The absolute mouse position
+				pageY -
+
+				// Click offset (relative to the element)
+				this.offset.click.top -
+
+				// Only for relative positioned nodes: Relative offset from element to offset parent
+				this.offset.relative.top -
+
+				// The offsetParent's offset without borders (offset + border)
+				this.offset.parent.top +
+				( ( this.cssPosition === "fixed" ?
+					-this.scrollParent.scrollTop() :
+					( scrollIsRootNode ? 0 : scroll.scrollTop() ) ) )
+			),
+			left: (
+
+				// The absolute mouse position
+				pageX -
+
+				// Click offset (relative to the element)
+				this.offset.click.left -
+
+				// Only for relative positioned nodes: Relative offset from element to offset parent
+				this.offset.relative.left -
+
+				// The offsetParent's offset without borders (offset + border)
+				this.offset.parent.left +
+				( ( this.cssPosition === "fixed" ?
+					-this.scrollParent.scrollLeft() :
+					scrollIsRootNode ? 0 : scroll.scrollLeft() ) )
+			)
+		};
+
+	},
+
+	_rearrange: function( event, i, a, hardRefresh ) {
+
+		a ? a[ 0 ].appendChild( this.placeholder[ 0 ] ) :
+			i.item[ 0 ].parentNode.insertBefore( this.placeholder[ 0 ],
+				( this.direction === "down" ? i.item[ 0 ] : i.item[ 0 ].nextSibling ) );
+
+		//Various things done here to improve the performance:
+		// 1. we create a setTimeout, that calls refreshPositions
+		// 2. on the instance, we have a counter variable, that get's higher after every append
+		// 3. on the local scope, we copy the counter variable, and check in the timeout,
+		// if it's still the same
+		// 4. this lets only the last addition to the timeout stack through
+		this.counter = this.counter ? ++this.counter : 1;
+		var counter = this.counter;
+
+		this._delay( function() {
+			if ( counter === this.counter ) {
+
+				//Precompute after each DOM insertion, NOT on mousemove
+				this.refreshPositions( !hardRefresh );
+			}
+		} );
+
+	},
+
+	_clear: function( event, noPropagation ) {
+
+		this.reverting = false;
+
+		// We delay all events that have to be triggered to after the point where the placeholder
+		// has been removed and everything else normalized again
+		var i,
+			delayedTriggers = [];
+
+		// We first have to update the dom position of the actual currentItem
+		// Note: don't do it if the current item is already removed (by a user), or it gets
+		// reappended (see #4088)
+		if ( !this._noFinalSort && this.currentItem.parent().length ) {
+			this.placeholder.before( this.currentItem );
+		}
+		this._noFinalSort = null;
+
+		if ( this.helper[ 0 ] === this.currentItem[ 0 ] ) {
+			for ( i in this._storedCSS ) {
+				if ( this._storedCSS[ i ] === "auto" || this._storedCSS[ i ] === "static" ) {
+					this._storedCSS[ i ] = "";
+				}
+			}
+			this.currentItem.css( this._storedCSS );
+			this._removeClass( this.currentItem, "ui-sortable-helper" );
+		} else {
+			this.currentItem.show();
+		}
+
+		if ( this.fromOutside && !noPropagation ) {
+			delayedTriggers.push( function( event ) {
+				this._trigger( "receive", event, this._uiHash( this.fromOutside ) );
+			} );
+		}
+		if ( ( this.fromOutside ||
+				this.domPosition.prev !==
+				this.currentItem.prev().not( ".ui-sortable-helper" )[ 0 ] ||
+				this.domPosition.parent !== this.currentItem.parent()[ 0 ] ) && !noPropagation ) {
+
+			// Trigger update callback if the DOM position has changed
+			delayedTriggers.push( function( event ) {
+				this._trigger( "update", event, this._uiHash() );
+			} );
+		}
+
+		// Check if the items Container has Changed and trigger appropriate
+		// events.
+		if ( this !== this.currentContainer ) {
+			if ( !noPropagation ) {
+				delayedTriggers.push( function( event ) {
+					this._trigger( "remove", event, this._uiHash() );
+				} );
+				delayedTriggers.push( ( function( c ) {
+					return function( event ) {
+						c._trigger( "receive", event, this._uiHash( this ) );
+					};
+				} ).call( this, this.currentContainer ) );
+				delayedTriggers.push( ( function( c ) {
+					return function( event ) {
+						c._trigger( "update", event, this._uiHash( this ) );
+					};
+				} ).call( this, this.currentContainer ) );
+			}
+		}
+
+		//Post events to containers
+		function delayEvent( type, instance, container ) {
+			return function( event ) {
+				container._trigger( type, event, instance._uiHash( instance ) );
+			};
+		}
+		for ( i = this.containers.length - 1; i >= 0; i-- ) {
+			if ( !noPropagation ) {
+				delayedTriggers.push( delayEvent( "deactivate", this, this.containers[ i ] ) );
+			}
+			if ( this.containers[ i ].containerCache.over ) {
+				delayedTriggers.push( delayEvent( "out", this, this.containers[ i ] ) );
+				this.containers[ i ].containerCache.over = 0;
+			}
+		}
+
+		//Do what was originally in plugins
+		if ( this.storedCursor ) {
+			this.document.find( "body" ).css( "cursor", this.storedCursor );
+			this.storedStylesheet.remove();
+		}
+		if ( this._storedOpacity ) {
+			this.helper.css( "opacity", this._storedOpacity );
+		}
+		if ( this._storedZIndex ) {
+			this.helper.css( "zIndex", this._storedZIndex === "auto" ? "" : this._storedZIndex );
+		}
+
+		this.dragging = false;
+
+		if ( !noPropagation ) {
+			this._trigger( "beforeStop", event, this._uiHash() );
+		}
+
+		//$(this.placeholder[0]).remove(); would have been the jQuery way - unfortunately,
+		// it unbinds ALL events from the original node!
+		this.placeholder[ 0 ].parentNode.removeChild( this.placeholder[ 0 ] );
+
+		if ( !this.cancelHelperRemoval ) {
+			if ( this.helper[ 0 ] !== this.currentItem[ 0 ] ) {
+				this.helper.remove();
+			}
+			this.helper = null;
+		}
+
+		if ( !noPropagation ) {
+			for ( i = 0; i < delayedTriggers.length; i++ ) {
+
+				// Trigger all delayed events
+				delayedTriggers[ i ].call( this, event );
+			}
+			this._trigger( "stop", event, this._uiHash() );
+		}
+
+		this.fromOutside = false;
+		return !this.cancelHelperRemoval;
+
+	},
+
+	_trigger: function() {
+		if ( $.Widget.prototype._trigger.apply( this, arguments ) === false ) {
+			this.cancel();
+		}
+	},
+
+	_uiHash: function( _inst ) {
+		var inst = _inst || this;
+		return {
+			helper: inst.helper,
+			placeholder: inst.placeholder || $( [] ),
+			position: inst.position,
+			originalPosition: inst.originalPosition,
+			offset: inst.positionAbs,
+			item: inst.currentItem,
+			sender: _inst ? _inst.element : null
+		};
+	}
+
+} );
+
+
+/*!
+ * jQuery UI Spinner 1.12.0
+ * http://jqueryui.com
+ *
+ * Copyright jQuery Foundation and other contributors
+ * Released under the MIT license.
+ * http://jquery.org/license
+ */
+
+//>>label: Spinner
+//>>group: Widgets
+//>>description: Displays buttons to easily input numbers via the keyboard or mouse.
+//>>docs: http://api.jqueryui.com/spinner/
+//>>demos: http://jqueryui.com/spinner/
+//>>css.structure: ../../themes/base/core.css
+//>>css.structure: ../../themes/base/spinner.css
+//>>css.theme: ../../themes/base/theme.css
+
+
+
+function spinnerModifer( fn ) {
+	return function() {
+		var previous = this.element.val();
+		fn.apply( this, arguments );
+		this._refresh();
+		if ( previous !== this.element.val() ) {
+			this._trigger( "change" );
+		}
+	};
+}
+
+$.widget( "ui.spinner", {
+	version: "1.12.0",
+	defaultElement: "<input>",
+	widgetEventPrefix: "spin",
+	options: {
+		classes: {
+			"ui-spinner": "ui-corner-all",
+			"ui-spinner-down": "ui-corner-br",
+			"ui-spinner-up": "ui-corner-tr"
+		},
+		culture: null,
+		icons: {
+			down: "ui-icon-triangle-1-s",
+			up: "ui-icon-triangle-1-n"
+		},
+		incremental: true,
+		max: null,
+		min: null,
+		numberFormat: null,
+		page: 10,
+		step: 1,
+
+		change: null,
+		spin: null,
+		start: null,
+		stop: null
+	},
+
+	_create: function() {
+
+		// handle string values that need to be parsed
+		this._setOption( "max", this.options.max );
+		this._setOption( "min", this.options.min );
+		this._setOption( "step", this.options.step );
+
+		// Only format if there is a value, prevents the field from being marked
+		// as invalid in Firefox, see #9573.
+		if ( this.value() !== "" ) {
+
+			// Format the value, but don't constrain.
+			this._value( this.element.val(), true );
+		}
+
+		this._draw();
+		this._on( this._events );
+		this._refresh();
+
+		// Turning off autocomplete prevents the browser from remembering the
+		// value when navigating through history, so we re-enable autocomplete
+		// if the page is unloaded before the widget is destroyed. #7790
+		this._on( this.window, {
+			beforeunload: function() {
+				this.element.removeAttr( "autocomplete" );
+			}
+		} );
+	},
+
+	_getCreateOptions: function() {
+		var options = this._super();
+		var element = this.element;
+
+		$.each( [ "min", "max", "step" ], function( i, option ) {
+			var value = element.attr( option );
+			if ( value != null && value.length ) {
+				options[ option ] = value;
+			}
+		} );
+
+		return options;
+	},
+
+	_events: {
+		keydown: function( event ) {
+			if ( this._start( event ) && this._keydown( event ) ) {
+				event.preventDefault();
+			}
+		},
+		keyup: "_stop",
+		focus: function() {
+			this.previous = this.element.val();
+		},
+		blur: function( event ) {
+			if ( this.cancelBlur ) {
+				delete this.cancelBlur;
+				return;
+			}
+
+			this._stop();
+			this._refresh();
+			if ( this.previous !== this.element.val() ) {
+				this._trigger( "change", event );
+			}
+		},
+		mousewheel: function( event, delta ) {
+			if ( !delta ) {
+				return;
+			}
+			if ( !this.spinning && !this._start( event ) ) {
+				return false;
+			}
+
+			this._spin( ( delta > 0 ? 1 : -1 ) * this.options.step, event );
+			clearTimeout( this.mousewheelTimer );
+			this.mousewheelTimer = this._delay( function() {
+				if ( this.spinning ) {
+					this._stop( event );
+				}
+			}, 100 );
+			event.preventDefault();
+		},
+		"mousedown .ui-spinner-button": function( event ) {
+			var previous;
+
+			// We never want the buttons to have focus; whenever the user is
+			// interacting with the spinner, the focus should be on the input.
+			// If the input is focused then this.previous is properly set from
+			// when the input first received focus. If the input is not focused
+			// then we need to set this.previous based on the value before spinning.
+			previous = this.element[ 0 ] === $.ui.safeActiveElement( this.document[ 0 ] ) ?
+				this.previous : this.element.val();
+			function checkFocus() {
+				var isActive = this.element[ 0 ] === $.ui.safeActiveElement( this.document[ 0 ] );
+				if ( !isActive ) {
+					this.element.trigger( "focus" );
+					this.previous = previous;
+
+					// support: IE
+					// IE sets focus asynchronously, so we need to check if focus
+					// moved off of the input because the user clicked on the button.
+					this._delay( function() {
+						this.previous = previous;
+					} );
+				}
+			}
+
+			// Ensure focus is on (or stays on) the text field
+			event.preventDefault();
+			checkFocus.call( this );
+
+			// Support: IE
+			// IE doesn't prevent moving focus even with event.preventDefault()
+			// so we set a flag to know when we should ignore the blur event
+			// and check (again) if focus moved off of the input.
+			this.cancelBlur = true;
+			this._delay( function() {
+				delete this.cancelBlur;
+				checkFocus.call( this );
+			} );
+
+			if ( this._start( event ) === false ) {
+				return;
+			}
+
+			this._repeat( null, $( event.currentTarget )
+				.hasClass( "ui-spinner-up" ) ? 1 : -1, event );
+		},
+		"mouseup .ui-spinner-button": "_stop",
+		"mouseenter .ui-spinner-button": function( event ) {
+
+			// button will add ui-state-active if mouse was down while mouseleave and kept down
+			if ( !$( event.currentTarget ).hasClass( "ui-state-active" ) ) {
+				return;
+			}
+
+			if ( this._start( event ) === false ) {
+				return false;
+			}
+			this._repeat( null, $( event.currentTarget )
+				.hasClass( "ui-spinner-up" ) ? 1 : -1, event );
+		},
+
+		// TODO: do we really want to consider this a stop?
+		// shouldn't we just stop the repeater and wait until mouseup before
+		// we trigger the stop event?
+		"mouseleave .ui-spinner-button": "_stop"
+	},
+
+	// Support mobile enhanced option and make backcompat more sane
+	_enhance: function() {
+		this.uiSpinner = this.element
+			.attr( "autocomplete", "off" )
+			.wrap( "<span>" )
+			.parent()
+
+				// Add buttons
+				.append(
+					"<a></a><a></a>"
+				);
+	},
+
+	_draw: function() {
+		this._enhance();
+
+		this._addClass( this.uiSpinner, "ui-spinner", "ui-widget ui-widget-content" );
+		this._addClass( "ui-spinner-input" );
+
+		this.element.attr( "role", "spinbutton" );
+
+		// Button bindings
+		this.buttons = this.uiSpinner.children( "a" )
+			.attr( "tabIndex", -1 )
+			.attr( "aria-hidden", true )
+			.button( {
+				classes: {
+					"ui-button": ""
+				}
+			} );
+
+		// TODO: Right now button does not support classes this is already updated in button PR
+		this._removeClass( this.buttons, "ui-corner-all" );
+
+		this._addClass( this.buttons.first(), "ui-spinner-button ui-spinner-up" );
+		this._addClass( this.buttons.last(), "ui-spinner-button ui-spinner-down" );
+		this.buttons.first().button( {
+			"icon": this.options.icons.up,
+			"showLabel": false
+		} );
+		this.buttons.last().button( {
+			"icon": this.options.icons.down,
+			"showLabel": false
+		} );
+
+		// IE 6 doesn't understand height: 50% for the buttons
+		// unless the wrapper has an explicit height
+		if ( this.buttons.height() > Math.ceil( this.uiSpinner.height() * 0.5 ) &&
+				this.uiSpinner.height() > 0 ) {
+			this.uiSpinner.height( this.uiSpinner.height() );
+		}
+	},
+
+	_keydown: function( event ) {
+		var options = this.options,
+			keyCode = $.ui.keyCode;
+
+		switch ( event.keyCode ) {
+		case keyCode.UP:
+			this._repeat( null, 1, event );
+			return true;
+		case keyCode.DOWN:
+			this._repeat( null, -1, event );
+			return true;
+		case keyCode.PAGE_UP:
+			this._repeat( null, options.page, event );
+			return true;
+		case keyCode.PAGE_DOWN:
+			this._repeat( null, -options.page, event );
+			return true;
+		}
+
+		return false;
+	},
+
+	_start: function( event ) {
+		if ( !this.spinning && this._trigger( "start", event ) === false ) {
+			return false;
+		}
+
+		if ( !this.counter ) {
+			this.counter = 1;
+		}
+		this.spinning = true;
+		return true;
+	},
+
+	_repeat: function( i, steps, event ) {
+		i = i || 500;
+
+		clearTimeout( this.timer );
+		this.timer = this._delay( function() {
+			this._repeat( 40, steps, event );
+		}, i );
+
+		this._spin( steps * this.options.step, event );
+	},
+
+	_spin: function( step, event ) {
+		var value = this.value() || 0;
+
+		if ( !this.counter ) {
+			this.counter = 1;
+		}
+
+		value = this._adjustValue( value + step * this._increment( this.counter ) );
+
+		if ( !this.spinning || this._trigger( "spin", event, { value: value } ) !== false ) {
+			this._value( value );
+			this.counter++;
+		}
+	},
+
+	_increment: function( i ) {
+		var incremental = this.options.incremental;
+
+		if ( incremental ) {
+			return $.isFunction( incremental ) ?
+				incremental( i ) :
+				Math.floor( i * i * i / 50000 - i * i / 500 + 17 * i / 200 + 1 );
+		}
+
+		return 1;
+	},
+
+	_precision: function() {
+		var precision = this._precisionOf( this.options.step );
+		if ( this.options.min !== null ) {
+			precision = Math.max( precision, this._precisionOf( this.options.min ) );
+		}
+		return precision;
+	},
+
+	_precisionOf: function( num ) {
+		var str = num.toString(),
+			decimal = str.indexOf( "." );
+		return decimal === -1 ? 0 : str.length - decimal - 1;
+	},
+
+	_adjustValue: function( value ) {
+		var base, aboveMin,
+			options = this.options;
+
+		// Make sure we're at a valid step
+		// - find out where we are relative to the base (min or 0)
+		base = options.min !== null ? options.min : 0;
+		aboveMin = value - base;
+
+		// - round to the nearest step
+		aboveMin = Math.round( aboveMin / options.step ) * options.step;
+
+		// - rounding is based on 0, so adjust back to our base
+		value = base + aboveMin;
+
+		// Fix precision from bad JS floating point math
+		value = parseFloat( value.toFixed( this._precision() ) );
+
+		// Clamp the value
+		if ( options.max !== null && value > options.max ) {
+			return options.max;
+		}
+		if ( options.min !== null && value < options.min ) {
+			return options.min;
+		}
+
+		return value;
+	},
+
+	_stop: function( event ) {
+		if ( !this.spinning ) {
+			return;
+		}
+
+		clearTimeout( this.timer );
+		clearTimeout( this.mousewheelTimer );
+		this.counter = 0;
+		this.spinning = false;
+		this._trigger( "stop", event );
+	},
+
+	_setOption: function( key, value ) {
+		var prevValue, first, last;
+
+		if ( key === "culture" || key === "numberFormat" ) {
+			prevValue = this._parse( this.element.val() );
+			this.options[ key ] = value;
+			this.element.val( this._format( prevValue ) );
+			return;
+		}
+
+		if ( key === "max" || key === "min" || key === "step" ) {
+			if ( typeof value === "string" ) {
+				value = this._parse( value );
+			}
+		}
+		if ( key === "icons" ) {
+			first = this.buttons.first().find( ".ui-icon" );
+			this._removeClass( first, null, this.options.icons.up );
+			this._addClass( first, null, value.up );
+			last = this.buttons.last().find( ".ui-icon" );
+			this._removeClass( last, null, this.options.icons.down );
+			this._addClass( last, null, value.down );
+		}
+
+		this._super( key, value );
+	},
+
+	_setOptionDisabled: function( value ) {
+		this._super( value );
+
+		this._toggleClass( this.uiSpinner, null, "ui-state-disabled", !!value );
+		this.element.prop( "disabled", !!value );
+		this.buttons.button( value ? "disable" : "enable" );
+	},
+
+	_setOptions: spinnerModifer( function( options ) {
+		this._super( options );
+	} ),
+
+	_parse: function( val ) {
+		if ( typeof val === "string" && val !== "" ) {
+			val = window.Globalize && this.options.numberFormat ?
+				Globalize.parseFloat( val, 10, this.options.culture ) : +val;
+		}
+		return val === "" || isNaN( val ) ? null : val;
+	},
+
+	_format: function( value ) {
+		if ( value === "" ) {
+			return "";
+		}
+		return window.Globalize && this.options.numberFormat ?
+			Globalize.format( value, this.options.numberFormat, this.options.culture ) :
+			value;
+	},
+
+	_refresh: function() {
+		this.element.attr( {
+			"aria-valuemin": this.options.min,
+			"aria-valuemax": this.options.max,
+
+			// TODO: what should we do with values that can't be parsed?
+			"aria-valuenow": this._parse( this.element.val() )
+		} );
+	},
+
+	isValid: function() {
+		var value = this.value();
+
+		// Null is invalid
+		if ( value === null ) {
+			return false;
+		}
+
+		// If value gets adjusted, it's invalid
+		return value === this._adjustValue( value );
+	},
+
+	// Update the value without triggering change
+	_value: function( value, allowAny ) {
+		var parsed;
+		if ( value !== "" ) {
+			parsed = this._parse( value );
+			if ( parsed !== null ) {
+				if ( !allowAny ) {
+					parsed = this._adjustValue( parsed );
+				}
+				value = this._format( parsed );
+			}
+		}
+		this.element.val( value );
+		this._refresh();
+	},
+
+	_destroy: function() {
+		this.element
+			.prop( "disabled", false )
+			.removeAttr( "autocomplete role aria-valuemin aria-valuemax aria-valuenow" );
+
+		this.uiSpinner.replaceWith( this.element );
+	},
+
+	stepUp: spinnerModifer( function( steps ) {
+		this._stepUp( steps );
+	} ),
+	_stepUp: function( steps ) {
+		if ( this._start() ) {
+			this._spin( ( steps || 1 ) * this.options.step );
+			this._stop();
+		}
+	},
+
+	stepDown: spinnerModifer( function( steps ) {
+		this._stepDown( steps );
+	} ),
+	_stepDown: function( steps ) {
+		if ( this._start() ) {
+			this._spin( ( steps || 1 ) * -this.options.step );
+			this._stop();
+		}
+	},
+
+	pageUp: spinnerModifer( function( pages ) {
+		this._stepUp( ( pages || 1 ) * this.options.page );
+	} ),
+
+	pageDown: spinnerModifer( function( pages ) {
+		this._stepDown( ( pages || 1 ) * this.options.page );
+	} ),
+
+	value: function( newVal ) {
+		if ( !arguments.length ) {
+			return this._parse( this.element.val() );
+		}
+		spinnerModifer( this._value ).call( this, newVal );
+	},
+
+	widget: function() {
+		return this.uiSpinner;
+	}
+} );
+
+// DEPRECATED
+// TODO: switch return back to widget declaration at top of file when this is removed
+if ( $.uiBackCompat !== false ) {
+
+	// Backcompat for spinner html extension points
+	$.widget( "ui.spinner", $.ui.spinner, {
+		_enhance: function() {
+			this.uiSpinner = this.element
+				.attr( "autocomplete", "off" )
+				.wrap( this._uiSpinnerHtml() )
+				.parent()
+
+					// Add buttons
+					.append( this._buttonHtml() );
+		},
+		_uiSpinnerHtml: function() {
+			return "<span>";
+		},
+
+		_buttonHtml: function() {
+			return "<a></a><a></a>";
+		}
+	} );
+}
+
+var widgetsSpinner = $.ui.spinner;
+
+
+/*!
+ * jQuery UI Tabs 1.12.0
+ * http://jqueryui.com
+ *
+ * Copyright jQuery Foundation and other contributors
+ * Released under the MIT license.
+ * http://jquery.org/license
+ */
+
+//>>label: Tabs
+//>>group: Widgets
+//>>description: Transforms a set of container elements into a tab structure.
+//>>docs: http://api.jqueryui.com/tabs/
+//>>demos: http://jqueryui.com/tabs/
+//>>css.structure: ../../themes/base/core.css
+//>>css.structure: ../../themes/base/tabs.css
+//>>css.theme: ../../themes/base/theme.css
+
+
+
+$.widget( "ui.tabs", {
+	version: "1.12.0",
+	delay: 300,
+	options: {
+		active: null,
+		classes: {
+			"ui-tabs": "ui-corner-all",
+			"ui-tabs-nav": "ui-corner-all",
+			"ui-tabs-panel": "ui-corner-bottom",
+			"ui-tabs-tab": "ui-corner-top"
+		},
+		collapsible: false,
+		event: "click",
+		heightStyle: "content",
+		hide: null,
+		show: null,
+
+		// Callbacks
+		activate: null,
+		beforeActivate: null,
+		beforeLoad: null,
+		load: null
+	},
+
+	_isLocal: ( function() {
+		var rhash = /#.*$/;
+
+		return function( anchor ) {
+			var anchorUrl, locationUrl;
+
+			anchorUrl = anchor.href.replace( rhash, "" );
+			locationUrl = location.href.replace( rhash, "" );
+
+			// Decoding may throw an error if the URL isn't UTF-8 (#9518)
+			try {
+				anchorUrl = decodeURIComponent( anchorUrl );
+			} catch ( error ) {}
+			try {
+				locationUrl = decodeURIComponent( locationUrl );
+			} catch ( error ) {}
+
+			return anchor.hash.length > 1 && anchorUrl === locationUrl;
+		};
+	} )(),
+
+	_create: function() {
+		var that = this,
+			options = this.options;
+
+		this.running = false;
+
+		this._addClass( "ui-tabs", "ui-widget ui-widget-content" );
+		this._toggleClass( "ui-tabs-collapsible", null, options.collapsible );
+
+		this._processTabs();
+		options.active = this._initialActive();
+
+		// Take disabling tabs via class attribute from HTML
+		// into account and update option properly.
+		if ( $.isArray( options.disabled ) ) {
+			options.disabled = $.unique( options.disabled.concat(
+				$.map( this.tabs.filter( ".ui-state-disabled" ), function( li ) {
+					return that.tabs.index( li );
+				} )
+			) ).sort();
+		}
+
+		// Check for length avoids error when initializing empty list
+		if ( this.options.active !== false && this.anchors.length ) {
+			this.active = this._findActive( options.active );
+		} else {
+			this.active = $();
+		}
+
+		this._refresh();
+
+		if ( this.active.length ) {
+			this.load( options.active );
+		}
+	},
+
+	_initialActive: function() {
+		var active = this.options.active,
+			collapsible = this.options.collapsible,
+			locationHash = location.hash.substring( 1 );
+
+		if ( active === null ) {
+
+			// check the fragment identifier in the URL
+			if ( locationHash ) {
+				this.tabs.each( function( i, tab ) {
+					if ( $( tab ).attr( "aria-controls" ) === locationHash ) {
+						active = i;
+						return false;
+					}
+				} );
+			}
+
+			// Check for a tab marked active via a class
+			if ( active === null ) {
+				active = this.tabs.index( this.tabs.filter( ".ui-tabs-active" ) );
+			}
+
+			// No active tab, set to false
+			if ( active === null || active === -1 ) {
+				active = this.tabs.length ? 0 : false;
+			}
+		}
+
+		// Handle numbers: negative, out of range
+		if ( active !== false ) {
+			active = this.tabs.index( this.tabs.eq( active ) );
+			if ( active === -1 ) {
+				active = collapsible ? false : 0;
+			}
+		}
+
+		// Don't allow collapsible: false and active: false
+		if ( !collapsible && active === false && this.anchors.length ) {
+			active = 0;
+		}
+
+		return active;
+	},
+
+	_getCreateEventData: function() {
+		return {
+			tab: this.active,
+			panel: !this.active.length ? $() : this._getPanelForTab( this.active )
+		};
+	},
+
+	_tabKeydown: function( event ) {
+		var focusedTab = $( $.ui.safeActiveElement( this.document[ 0 ] ) ).closest( "li" ),
+			selectedIndex = this.tabs.index( focusedTab ),
+			goingForward = true;
+
+		if ( this._handlePageNav( event ) ) {
+			return;
+		}
+
+		switch ( event.keyCode ) {
+		case $.ui.keyCode.RIGHT:
+		case $.ui.keyCode.DOWN:
+			selectedIndex++;
+			break;
+		case $.ui.keyCode.UP:
+		case $.ui.keyCode.LEFT:
+			goingForward = false;
+			selectedIndex--;
+			break;
+		case $.ui.keyCode.END:
+			selectedIndex = this.anchors.length - 1;
+			break;
+		case $.ui.keyCode.HOME:
+			selectedIndex = 0;
+			break;
+		case $.ui.keyCode.SPACE:
+
+			// Activate only, no collapsing
+			event.preventDefault();
+			clearTimeout( this.activating );
+			this._activate( selectedIndex );
+			return;
+		case $.ui.keyCode.ENTER:
+
+			// Toggle (cancel delayed activation, allow collapsing)
+			event.preventDefault();
+			clearTimeout( this.activating );
+
+			// Determine if we should collapse or activate
+			this._activate( selectedIndex === this.options.active ? false : selectedIndex );
+			return;
+		default:
+			return;
+		}
+
+		// Focus the appropriate tab, based on which key was pressed
+		event.preventDefault();
+		clearTimeout( this.activating );
+		selectedIndex = this._focusNextTab( selectedIndex, goingForward );
+
+		// Navigating with control/command key will prevent automatic activation
+		if ( !event.ctrlKey && !event.metaKey ) {
+
+			// Update aria-selected immediately so that AT think the tab is already selected.
+			// Otherwise AT may confuse the user by stating that they need to activate the tab,
+			// but the tab will already be activated by the time the announcement finishes.
+			focusedTab.attr( "aria-selected", "false" );
+			this.tabs.eq( selectedIndex ).attr( "aria-selected", "true" );
+
+			this.activating = this._delay( function() {
+				this.option( "active", selectedIndex );
+			}, this.delay );
+		}
+	},
+
+	_panelKeydown: function( event ) {
+		if ( this._handlePageNav( event ) ) {
+			return;
+		}
+
+		// Ctrl+up moves focus to the current tab
+		if ( event.ctrlKey && event.keyCode === $.ui.keyCode.UP ) {
+			event.preventDefault();
+			this.active.trigger( "focus" );
+		}
+	},
+
+	// Alt+page up/down moves focus to the previous/next tab (and activates)
+	_handlePageNav: function( event ) {
+		if ( event.altKey && event.keyCode === $.ui.keyCode.PAGE_UP ) {
+			this._activate( this._focusNextTab( this.options.active - 1, false ) );
+			return true;
+		}
+		if ( event.altKey && event.keyCode === $.ui.keyCode.PAGE_DOWN ) {
+			this._activate( this._focusNextTab( this.options.active + 1, true ) );
+			return true;
+		}
+	},
+
+	_findNextTab: function( index, goingForward ) {
+		var lastTabIndex = this.tabs.length - 1;
+
+		function constrain() {
+			if ( index > lastTabIndex ) {
+				index = 0;
+			}
+			if ( index < 0 ) {
+				index = lastTabIndex;
+			}
+			return index;
+		}
+
+		while ( $.inArray( constrain(), this.options.disabled ) !== -1 ) {
+			index = goingForward ? index + 1 : index - 1;
+		}
+
+		return index;
+	},
+
+	_focusNextTab: function( index, goingForward ) {
+		index = this._findNextTab( index, goingForward );
+		this.tabs.eq( index ).trigger( "focus" );
+		return index;
+	},
+
+	_setOption: function( key, value ) {
+		if ( key === "active" ) {
+
+			// _activate() will handle invalid values and update this.options
+			this._activate( value );
+			return;
+		}
+
+		this._super( key, value );
+
+		if ( key === "collapsible" ) {
+			this._toggleClass( "ui-tabs-collapsible", null, value );
+
+			// Setting collapsible: false while collapsed; open first panel
+			if ( !value && this.options.active === false ) {
+				this._activate( 0 );
+			}
+		}
+
+		if ( key === "event" ) {
+			this._setupEvents( value );
+		}
+
+		if ( key === "heightStyle" ) {
+			this._setupHeightStyle( value );
+		}
+	},
+
+	_sanitizeSelector: function( hash ) {
+		return hash ? hash.replace( /[!"$%&'()*+,.\/:;<=>?@\[\]\^`{|}~]/g, "\\$&" ) : "";
+	},
+
+	refresh: function() {
+		var options = this.options,
+			lis = this.tablist.children( ":has(a[href])" );
+
+		// Get disabled tabs from class attribute from HTML
+		// this will get converted to a boolean if needed in _refresh()
+		options.disabled = $.map( lis.filter( ".ui-state-disabled" ), function( tab ) {
+			return lis.index( tab );
+		} );
+
+		this._processTabs();
+
+		// Was collapsed or no tabs
+		if ( options.active === false || !this.anchors.length ) {
+			options.active = false;
+			this.active = $();
+
+		// was active, but active tab is gone
+		} else if ( this.active.length && !$.contains( this.tablist[ 0 ], this.active[ 0 ] ) ) {
+
+			// all remaining tabs are disabled
+			if ( this.tabs.length === options.disabled.length ) {
+				options.active = false;
+				this.active = $();
+
+			// activate previous tab
+			} else {
+				this._activate( this._findNextTab( Math.max( 0, options.active - 1 ), false ) );
+			}
+
+		// was active, active tab still exists
+		} else {
+
+			// make sure active index is correct
+			options.active = this.tabs.index( this.active );
+		}
+
+		this._refresh();
+	},
+
+	_refresh: function() {
+		this._setOptionDisabled( this.options.disabled );
+		this._setupEvents( this.options.event );
+		this._setupHeightStyle( this.options.heightStyle );
+
+		this.tabs.not( this.active ).attr( {
+			"aria-selected": "false",
+			"aria-expanded": "false",
+			tabIndex: -1
+		} );
+		this.panels.not( this._getPanelForTab( this.active ) )
+			.hide()
+			.attr( {
+				"aria-hidden": "true"
+			} );
+
+		// Make sure one tab is in the tab order
+		if ( !this.active.length ) {
+			this.tabs.eq( 0 ).attr( "tabIndex", 0 );
+		} else {
+			this.active
+				.attr( {
+					"aria-selected": "true",
+					"aria-expanded": "true",
+					tabIndex: 0
+				} );
+			this._addClass( this.active, "ui-tabs-active", "ui-state-active" );
+			this._getPanelForTab( this.active )
+				.show()
+				.attr( {
+					"aria-hidden": "false"
+				} );
+		}
+	},
+
+	_processTabs: function() {
+		var that = this,
+			prevTabs = this.tabs,
+			prevAnchors = this.anchors,
+			prevPanels = this.panels;
+
+		this.tablist = this._getList().attr( "role", "tablist" );
+		this._addClass( this.tablist, "ui-tabs-nav",
+			"ui-helper-reset ui-helper-clearfix ui-widget-header" );
+
+		// Prevent users from focusing disabled tabs via click
+		this.tablist
+			.on( "mousedown" + this.eventNamespace, "> li", function( event ) {
+				if ( $( this ).is( ".ui-state-disabled" ) ) {
+					event.preventDefault();
+				}
+			} )
+
+			// Support: IE <9
+			// Preventing the default action in mousedown doesn't prevent IE
+			// from focusing the element, so if the anchor gets focused, blur.
+			// We don't have to worry about focusing the previously focused
+			// element since clicking on a non-focusable element should focus
+			// the body anyway.
+			.on( "focus" + this.eventNamespace, ".ui-tabs-anchor", function() {
+				if ( $( this ).closest( "li" ).is( ".ui-state-disabled" ) ) {
+					this.blur();
+				}
+			} );
+
+		this.tabs = this.tablist.find( "> li:has(a[href])" )
+			.attr( {
+				role: "tab",
+				tabIndex: -1
+			} );
+		this._addClass( this.tabs, "ui-tabs-tab", "ui-state-default" );
+
+		this.anchors = this.tabs.map( function() {
+			return $( "a", this )[ 0 ];
+		} )
+			.attr( {
+				role: "presentation",
+				tabIndex: -1
+			} );
+		this._addClass( this.anchors, "ui-tabs-anchor" );
+
+		this.panels = $();
+
+		this.anchors.each( function( i, anchor ) {
+			var selector, panel, panelId,
+				anchorId = $( anchor ).uniqueId().attr( "id" ),
+				tab = $( anchor ).closest( "li" ),
+				originalAriaControls = tab.attr( "aria-controls" );
+
+			// Inline tab
+			if ( that._isLocal( anchor ) ) {
+				selector = anchor.hash;
+				panelId = selector.substring( 1 );
+				panel = that.element.find( that._sanitizeSelector( selector ) );
+
+			// remote tab
+			} else {
+
+				// If the tab doesn't already have aria-controls,
+				// generate an id by using a throw-away element
+				panelId = tab.attr( "aria-controls" ) || $( {} ).uniqueId()[ 0 ].id;
+				selector = "#" + panelId;
+				panel = that.element.find( selector );
+				if ( !panel.length ) {
+					panel = that._createPanel( panelId );
+					panel.insertAfter( that.panels[ i - 1 ] || that.tablist );
+				}
+				panel.attr( "aria-live", "polite" );
+			}
+
+			if ( panel.length ) {
+				that.panels = that.panels.add( panel );
+			}
+			if ( originalAriaControls ) {
+				tab.data( "ui-tabs-aria-controls", originalAriaControls );
+			}
+			tab.attr( {
+				"aria-controls": panelId,
+				"aria-labelledby": anchorId
+			} );
+			panel.attr( "aria-labelledby", anchorId );
+		} );
+
+		this.panels.attr( "role", "tabpanel" );
+		this._addClass( this.panels, "ui-tabs-panel", "ui-widget-content" );
+
+		// Avoid memory leaks (#10056)
+		if ( prevTabs ) {
+			this._off( prevTabs.not( this.tabs ) );
+			this._off( prevAnchors.not( this.anchors ) );
+			this._off( prevPanels.not( this.panels ) );
+		}
+	},
+
+	// Allow overriding how to find the list for rare usage scenarios (#7715)
+	_getList: function() {
+		return this.tablist || this.element.find( "ol, ul" ).eq( 0 );
+	},
+
+	_createPanel: function( id ) {
+		return $( "<div>" )
+			.attr( "id", id )
+			.data( "ui-tabs-destroy", true );
+	},
+
+	_setOptionDisabled: function( disabled ) {
+		var currentItem, li, i;
+
+		if ( $.isArray( disabled ) ) {
+			if ( !disabled.length ) {
+				disabled = false;
+			} else if ( disabled.length === this.anchors.length ) {
+				disabled = true;
+			}
+		}
+
+		// Disable tabs
+		for ( i = 0; ( li = this.tabs[ i ] ); i++ ) {
+			currentItem = $( li );
+			if ( disabled === true || $.inArray( i, disabled ) !== -1 ) {
+				currentItem.attr( "aria-disabled", "true" );
+				this._addClass( currentItem, null, "ui-state-disabled" );
+			} else {
+				currentItem.removeAttr( "aria-disabled" );
+				this._removeClass( currentItem, null, "ui-state-disabled" );
+			}
+		}
+
+		this.options.disabled = disabled;
+
+		this._toggleClass( this.widget(), this.widgetFullName + "-disabled", null,
+			disabled === true );
+	},
+
+	_setupEvents: function( event ) {
+		var events = {};
+		if ( event ) {
+			$.each( event.split( " " ), function( index, eventName ) {
+				events[ eventName ] = "_eventHandler";
+			} );
+		}
+
+		this._off( this.anchors.add( this.tabs ).add( this.panels ) );
+
+		// Always prevent the default action, even when disabled
+		this._on( true, this.anchors, {
+			click: function( event ) {
+				event.preventDefault();
+			}
+		} );
+		this._on( this.anchors, events );
+		this._on( this.tabs, { keydown: "_tabKeydown" } );
+		this._on( this.panels, { keydown: "_panelKeydown" } );
+
+		this._focusable( this.tabs );
+		this._hoverable( this.tabs );
+	},
+
+	_setupHeightStyle: function( heightStyle ) {
+		var maxHeight,
+			parent = this.element.parent();
+
+		if ( heightStyle === "fill" ) {
+			maxHeight = parent.height();
+			maxHeight -= this.element.outerHeight() - this.element.height();
+
+			this.element.siblings( ":visible" ).each( function() {
+				var elem = $( this ),
+					position = elem.css( "position" );
+
+				if ( position === "absolute" || position === "fixed" ) {
+					return;
+				}
+				maxHeight -= elem.outerHeight( true );
+			} );
+
+			this.element.children().not( this.panels ).each( function() {
+				maxHeight -= $( this ).outerHeight( true );
+			} );
+
+			this.panels.each( function() {
+				$( this ).height( Math.max( 0, maxHeight -
+					$( this ).innerHeight() + $( this ).height() ) );
+			} )
+				.css( "overflow", "auto" );
+		} else if ( heightStyle === "auto" ) {
+			maxHeight = 0;
+			this.panels.each( function() {
+				maxHeight = Math.max( maxHeight, $( this ).height( "" ).height() );
+			} ).height( maxHeight );
+		}
+	},
+
+	_eventHandler: function( event ) {
+		var options = this.options,
+			active = this.active,
+			anchor = $( event.currentTarget ),
+			tab = anchor.closest( "li" ),
+			clickedIsActive = tab[ 0 ] === active[ 0 ],
+			collapsing = clickedIsActive && options.collapsible,
+			toShow = collapsing ? $() : this._getPanelForTab( tab ),
+			toHide = !active.length ? $() : this._getPanelForTab( active ),
+			eventData = {
+				oldTab: active,
+				oldPanel: toHide,
+				newTab: collapsing ? $() : tab,
+				newPanel: toShow
+			};
+
+		event.preventDefault();
+
+		if ( tab.hasClass( "ui-state-disabled" ) ||
+
+				// tab is already loading
+				tab.hasClass( "ui-tabs-loading" ) ||
+
+				// can't switch durning an animation
+				this.running ||
+
+				// click on active header, but not collapsible
+				( clickedIsActive && !options.collapsible ) ||
+
+				// allow canceling activation
+				( this._trigger( "beforeActivate", event, eventData ) === false ) ) {
+			return;
+		}
+
+		options.active = collapsing ? false : this.tabs.index( tab );
+
+		this.active = clickedIsActive ? $() : tab;
+		if ( this.xhr ) {
+			this.xhr.abort();
+		}
+
+		if ( !toHide.length && !toShow.length ) {
+			$.error( "jQuery UI Tabs: Mismatching fragment identifier." );
+		}
+
+		if ( toShow.length ) {
+			this.load( this.tabs.index( tab ), event );
+		}
+		this._toggle( event, eventData );
+	},
+
+	// Handles show/hide for selecting tabs
+	_toggle: function( event, eventData ) {
+		var that = this,
+			toShow = eventData.newPanel,
+			toHide = eventData.oldPanel;
+
+		this.running = true;
+
+		function complete() {
+			that.running = false;
+			that._trigger( "activate", event, eventData );
+		}
+
+		function show() {
+			that._addClass( eventData.newTab.closest( "li" ), "ui-tabs-active", "ui-state-active" );
+
+			if ( toShow.length && that.options.show ) {
+				that._show( toShow, that.options.show, complete );
+			} else {
+				toShow.show();
+				complete();
+			}
+		}
+
+		// Start out by hiding, then showing, then completing
+		if ( toHide.length && this.options.hide ) {
+			this._hide( toHide, this.options.hide, function() {
+				that._removeClass( eventData.oldTab.closest( "li" ),
+					"ui-tabs-active", "ui-state-active" );
+				show();
+			} );
+		} else {
+			this._removeClass( eventData.oldTab.closest( "li" ),
+				"ui-tabs-active", "ui-state-active" );
+			toHide.hide();
+			show();
+		}
+
+		toHide.attr( "aria-hidden", "true" );
+		eventData.oldTab.attr( {
+			"aria-selected": "false",
+			"aria-expanded": "false"
+		} );
+
+		// If we're switching tabs, remove the old tab from the tab order.
+		// If we're opening from collapsed state, remove the previous tab from the tab order.
+		// If we're collapsing, then keep the collapsing tab in the tab order.
+		if ( toShow.length && toHide.length ) {
+			eventData.oldTab.attr( "tabIndex", -1 );
+		} else if ( toShow.length ) {
+			this.tabs.filter( function() {
+				return $( this ).attr( "tabIndex" ) === 0;
+			} )
+				.attr( "tabIndex", -1 );
+		}
+
+		toShow.attr( "aria-hidden", "false" );
+		eventData.newTab.attr( {
+			"aria-selected": "true",
+			"aria-expanded": "true",
+			tabIndex: 0
+		} );
+	},
+
+	_activate: function( index ) {
+		var anchor,
+			active = this._findActive( index );
+
+		// Trying to activate the already active panel
+		if ( active[ 0 ] === this.active[ 0 ] ) {
+			return;
+		}
+
+		// Trying to collapse, simulate a click on the current active header
+		if ( !active.length ) {
+			active = this.active;
+		}
+
+		anchor = active.find( ".ui-tabs-anchor" )[ 0 ];
+		this._eventHandler( {
+			target: anchor,
+			currentTarget: anchor,
+			preventDefault: $.noop
+		} );
+	},
+
+	_findActive: function( index ) {
+		return index === false ? $() : this.tabs.eq( index );
+	},
+
+	_getIndex: function( index ) {
+
+		// meta-function to give users option to provide a href string instead of a numerical index.
+		if ( typeof index === "string" ) {
+			index = this.anchors.index( this.anchors.filter( "[href$='" +
+				$.ui.escapeSelector( index ) + "']" ) );
+		}
+
+		return index;
+	},
+
+	_destroy: function() {
+		if ( this.xhr ) {
+			this.xhr.abort();
+		}
+
+		this.tablist
+			.removeAttr( "role" )
+			.off( this.eventNamespace );
+
+		this.anchors
+			.removeAttr( "role tabIndex" )
+			.removeUniqueId();
+
+		this.tabs.add( this.panels ).each( function() {
+			if ( $.data( this, "ui-tabs-destroy" ) ) {
+				$( this ).remove();
+			} else {
+				$( this ).removeAttr( "role tabIndex " +
+					"aria-live aria-busy aria-selected aria-labelledby aria-hidden aria-expanded" );
+			}
+		} );
+
+		this.tabs.each( function() {
+			var li = $( this ),
+				prev = li.data( "ui-tabs-aria-controls" );
+			if ( prev ) {
+				li
+					.attr( "aria-controls", prev )
+					.removeData( "ui-tabs-aria-controls" );
+			} else {
+				li.removeAttr( "aria-controls" );
+			}
+		} );
+
+		this.panels.show();
+
+		if ( this.options.heightStyle !== "content" ) {
+			this.panels.css( "height", "" );
+		}
+	},
+
+	enable: function( index ) {
+		var disabled = this.options.disabled;
+		if ( disabled === false ) {
+			return;
+		}
+
+		if ( index === undefined ) {
+			disabled = false;
+		} else {
+			index = this._getIndex( index );
+			if ( $.isArray( disabled ) ) {
+				disabled = $.map( disabled, function( num ) {
+					return num !== index ? num : null;
+				} );
+			} else {
+				disabled = $.map( this.tabs, function( li, num ) {
+					return num !== index ? num : null;
+				} );
+			}
+		}
+		this._setOptionDisabled( disabled );
+	},
+
+	disable: function( index ) {
+		var disabled = this.options.disabled;
+		if ( disabled === true ) {
+			return;
+		}
+
+		if ( index === undefined ) {
+			disabled = true;
+		} else {
+			index = this._getIndex( index );
+			if ( $.inArray( index, disabled ) !== -1 ) {
+				return;
+			}
+			if ( $.isArray( disabled ) ) {
+				disabled = $.merge( [ index ], disabled ).sort();
+			} else {
+				disabled = [ index ];
+			}
+		}
+		this._setOptionDisabled( disabled );
+	},
+
+	load: function( index, event ) {
+		index = this._getIndex( index );
+		var that = this,
+			tab = this.tabs.eq( index ),
+			anchor = tab.find( ".ui-tabs-anchor" ),
+			panel = this._getPanelForTab( tab ),
+			eventData = {
+				tab: tab,
+				panel: panel
+			},
+			complete = function( jqXHR, status ) {
+				if ( status === "abort" ) {
+					that.panels.stop( false, true );
+				}
+
+				that._removeClass( tab, "ui-tabs-loading" );
+				panel.removeAttr( "aria-busy" );
+
+				if ( jqXHR === that.xhr ) {
+					delete that.xhr;
+				}
+			};
+
+		// Not remote
+		if ( this._isLocal( anchor[ 0 ] ) ) {
+			return;
+		}
+
+		this.xhr = $.ajax( this._ajaxSettings( anchor, event, eventData ) );
+
+		// Support: jQuery <1.8
+		// jQuery <1.8 returns false if the request is canceled in beforeSend,
+		// but as of 1.8, $.ajax() always returns a jqXHR object.
+		if ( this.xhr && this.xhr.statusText !== "canceled" ) {
+			this._addClass( tab, "ui-tabs-loading" );
+			panel.attr( "aria-busy", "true" );
+
+			this.xhr
+				.done( function( response, status, jqXHR ) {
+
+					// support: jQuery <1.8
+					// http://bugs.jquery.com/ticket/11778
+					setTimeout( function() {
+						panel.html( response );
+						that._trigger( "load", event, eventData );
+
+						complete( jqXHR, status );
+					}, 1 );
+				} )
+				.fail( function( jqXHR, status ) {
+
+					// support: jQuery <1.8
+					// http://bugs.jquery.com/ticket/11778
+					setTimeout( function() {
+						complete( jqXHR, status );
+					}, 1 );
+				} );
+		}
+	},
+
+	_ajaxSettings: function( anchor, event, eventData ) {
+		var that = this;
+		return {
+			url: anchor.attr( "href" ),
+			beforeSend: function( jqXHR, settings ) {
+				return that._trigger( "beforeLoad", event,
+					$.extend( { jqXHR: jqXHR, ajaxSettings: settings }, eventData ) );
+			}
+		};
+	},
+
+	_getPanelForTab: function( tab ) {
+		var id = $( tab ).attr( "aria-controls" );
+		return this.element.find( this._sanitizeSelector( "#" + id ) );
+	}
+} );
+
+// DEPRECATED
+// TODO: Switch return back to widget declaration at top of file when this is removed
+if ( $.uiBackCompat !== false ) {
+
+	// Backcompat for ui-tab class (now ui-tabs-tab)
+	$.widget( "ui.tabs", $.ui.tabs, {
+		_processTabs: function() {
+			this._superApply( arguments );
+			this._addClass( this.tabs, "ui-tab" );
+		}
+	} );
+}
+
+var widgetsTabs = $.ui.tabs;
+
+
+/*!
+ * jQuery UI Tooltip 1.12.0
+ * http://jqueryui.com
+ *
+ * Copyright jQuery Foundation and other contributors
+ * Released under the MIT license.
+ * http://jquery.org/license
+ */
+
+//>>label: Tooltip
+//>>group: Widgets
+//>>description: Shows additional information for any element on hover or focus.
+//>>docs: http://api.jqueryui.com/tooltip/
+//>>demos: http://jqueryui.com/tooltip/
+//>>css.structure: ../../themes/base/core.css
+//>>css.structure: ../../themes/base/tooltip.css
+//>>css.theme: ../../themes/base/theme.css
+
+
+
+$.widget( "ui.tooltip", {
+	version: "1.12.0",
+	options: {
+		classes: {
+			"ui-tooltip": "ui-corner-all ui-widget-shadow"
+		},
+		content: function() {
+
+			// support: IE<9, Opera in jQuery <1.7
+			// .text() can't accept undefined, so coerce to a string
+			var title = $( this ).attr( "title" ) || "";
+
+			// Escape title, since we're going from an attribute to raw HTML
+			return $( "<a>" ).text( title ).html();
+		},
+		hide: true,
+
+		// Disabled elements have inconsistent behavior across browsers (#8661)
+		items: "[title]:not([disabled])",
+		position: {
+			my: "left top+15",
+			at: "left bottom",
+			collision: "flipfit flip"
+		},
+		show: true,
+		track: false,
+
+		// Callbacks
+		close: null,
+		open: null
+	},
+
+	_addDescribedBy: function( elem, id ) {
+		var describedby = ( elem.attr( "aria-describedby" ) || "" ).split( /\s+/ );
+		describedby.push( id );
+		elem
+			.data( "ui-tooltip-id", id )
+			.attr( "aria-describedby", $.trim( describedby.join( " " ) ) );
+	},
+
+	_removeDescribedBy: function( elem ) {
+		var id = elem.data( "ui-tooltip-id" ),
+			describedby = ( elem.attr( "aria-describedby" ) || "" ).split( /\s+/ ),
+			index = $.inArray( id, describedby );
+
+		if ( index !== -1 ) {
+			describedby.splice( index, 1 );
+		}
+
+		elem.removeData( "ui-tooltip-id" );
+		describedby = $.trim( describedby.join( " " ) );
+		if ( describedby ) {
+			elem.attr( "aria-describedby", describedby );
+		} else {
+			elem.removeAttr( "aria-describedby" );
+		}
+	},
+
+	_create: function() {
+		this._on( {
+			mouseover: "open",
+			focusin: "open"
+		} );
+
+		// IDs of generated tooltips, needed for destroy
+		this.tooltips = {};
+
+		// IDs of parent tooltips where we removed the title attribute
+		this.parents = {};
+
+		// Append the aria-live region so tooltips announce correctly
+		this.liveRegion = $( "<div>" )
+			.attr( {
+				role: "log",
+				"aria-live": "assertive",
+				"aria-relevant": "additions"
+			} )
+			.appendTo( this.document[ 0 ].body );
+		this._addClass( this.liveRegion, null, "ui-helper-hidden-accessible" );
+
+		this.disabledTitles = $( [] );
+	},
+
+	_setOption: function( key, value ) {
+		var that = this;
+
+		this._super( key, value );
+
+		if ( key === "content" ) {
+			$.each( this.tooltips, function( id, tooltipData ) {
+				that._updateContent( tooltipData.element );
+			} );
+		}
+	},
+
+	_setOptionDisabled: function( value ) {
+		this[ value ? "_disable" : "_enable" ]();
+	},
+
+	_disable: function() {
+		var that = this;
+
+		// Close open tooltips
+		$.each( this.tooltips, function( id, tooltipData ) {
+			var event = $.Event( "blur" );
+			event.target = event.currentTarget = tooltipData.element[ 0 ];
+			that.close( event, true );
+		} );
+
+		// Remove title attributes to prevent native tooltips
+		this.disabledTitles = this.disabledTitles.add(
+			this.element.find( this.options.items ).addBack()
+				.filter( function() {
+					var element = $( this );
+					if ( element.is( "[title]" ) ) {
+						return element
+							.data( "ui-tooltip-title", element.attr( "title" ) )
+							.removeAttr( "title" );
+					}
+				} )
+		);
+	},
+
+	_enable: function() {
+
+		// restore title attributes
+		this.disabledTitles.each( function() {
+			var element = $( this );
+			if ( element.data( "ui-tooltip-title" ) ) {
+				element.attr( "title", element.data( "ui-tooltip-title" ) );
+			}
+		} );
+		this.disabledTitles = $( [] );
+	},
+
+	open: function( event ) {
+		var that = this,
+			target = $( event ? event.target : this.element )
+
+				// we need closest here due to mouseover bubbling,
+				// but always pointing at the same event target
+				.closest( this.options.items );
+
+		// No element to show a tooltip for or the tooltip is already open
+		if ( !target.length || target.data( "ui-tooltip-id" ) ) {
+			return;
+		}
+
+		if ( target.attr( "title" ) ) {
+			target.data( "ui-tooltip-title", target.attr( "title" ) );
+		}
+
+		target.data( "ui-tooltip-open", true );
+
+		// Kill parent tooltips, custom or native, for hover
+		if ( event && event.type === "mouseover" ) {
+			target.parents().each( function() {
+				var parent = $( this ),
+					blurEvent;
+				if ( parent.data( "ui-tooltip-open" ) ) {
+					blurEvent = $.Event( "blur" );
+					blurEvent.target = blurEvent.currentTarget = this;
+					that.close( blurEvent, true );
+				}
+				if ( parent.attr( "title" ) ) {
+					parent.uniqueId();
+					that.parents[ this.id ] = {
+						element: this,
+						title: parent.attr( "title" )
+					};
+					parent.attr( "title", "" );
+				}
+			} );
+		}
+
+		this._registerCloseHandlers( event, target );
+		this._updateContent( target, event );
+	},
+
+	_updateContent: function( target, event ) {
+		var content,
+			contentOption = this.options.content,
+			that = this,
+			eventType = event ? event.type : null;
+
+		if ( typeof contentOption === "string" || contentOption.nodeType ||
+				contentOption.jquery ) {
+			return this._open( event, target, contentOption );
+		}
+
+		content = contentOption.call( target[ 0 ], function( response ) {
+
+			// IE may instantly serve a cached response for ajax requests
+			// delay this call to _open so the other call to _open runs first
+			that._delay( function() {
+
+				// Ignore async response if tooltip was closed already
+				if ( !target.data( "ui-tooltip-open" ) ) {
+					return;
+				}
+
+				// JQuery creates a special event for focusin when it doesn't
+				// exist natively. To improve performance, the native event
+				// object is reused and the type is changed. Therefore, we can't
+				// rely on the type being correct after the event finished
+				// bubbling, so we set it back to the previous value. (#8740)
+				if ( event ) {
+					event.type = eventType;
+				}
+				this._open( event, target, response );
+			} );
+		} );
+		if ( content ) {
+			this._open( event, target, content );
+		}
+	},
+
+	_open: function( event, target, content ) {
+		var tooltipData, tooltip, delayedShow, a11yContent,
+			positionOption = $.extend( {}, this.options.position );
+
+		if ( !content ) {
+			return;
+		}
+
+		// Content can be updated multiple times. If the tooltip already
+		// exists, then just update the content and bail.
+		tooltipData = this._find( target );
+		if ( tooltipData ) {
+			tooltipData.tooltip.find( ".ui-tooltip-content" ).html( content );
+			return;
+		}
+
+		// If we have a title, clear it to prevent the native tooltip
+		// we have to check first to avoid defining a title if none exists
+		// (we don't want to cause an element to start matching [title])
+		//
+		// We use removeAttr only for key events, to allow IE to export the correct
+		// accessible attributes. For mouse events, set to empty string to avoid
+		// native tooltip showing up (happens only when removing inside mouseover).
+		if ( target.is( "[title]" ) ) {
+			if ( event && event.type === "mouseover" ) {
+				target.attr( "title", "" );
+			} else {
+				target.removeAttr( "title" );
+			}
+		}
+
+		tooltipData = this._tooltip( target );
+		tooltip = tooltipData.tooltip;
+		this._addDescribedBy( target, tooltip.attr( "id" ) );
+		tooltip.find( ".ui-tooltip-content" ).html( content );
+
+		// Support: Voiceover on OS X, JAWS on IE <= 9
+		// JAWS announces deletions even when aria-relevant="additions"
+		// Voiceover will sometimes re-read the entire log region's contents from the beginning
+		this.liveRegion.children().hide();
+		a11yContent = $( "<div>" ).html( tooltip.find( ".ui-tooltip-content" ).html() );
+		a11yContent.removeAttr( "name" ).find( "[name]" ).removeAttr( "name" );
+		a11yContent.removeAttr( "id" ).find( "[id]" ).removeAttr( "id" );
+		a11yContent.appendTo( this.liveRegion );
+
+		function position( event ) {
+			positionOption.of = event;
+			if ( tooltip.is( ":hidden" ) ) {
+				return;
+			}
+			tooltip.position( positionOption );
+		}
+		if ( this.options.track && event && /^mouse/.test( event.type ) ) {
+			this._on( this.document, {
+				mousemove: position
+			} );
+
+			// trigger once to override element-relative positioning
+			position( event );
+		} else {
+			tooltip.position( $.extend( {
+				of: target
+			}, this.options.position ) );
+		}
+
+		tooltip.hide();
+
+		this._show( tooltip, this.options.show );
+
+		// Handle tracking tooltips that are shown with a delay (#8644). As soon
+		// as the tooltip is visible, position the tooltip using the most recent
+		// event.
+		// Adds the check to add the timers only when both delay and track options are set (#14682)
+		if ( this.options.track && this.options.show && this.options.show.delay ) {
+			delayedShow = this.delayedShow = setInterval( function() {
+				if ( tooltip.is( ":visible" ) ) {
+					position( positionOption.of );
+					clearInterval( delayedShow );
+				}
+			}, $.fx.interval );
+		}
+
+		this._trigger( "open", event, { tooltip: tooltip } );
+	},
+
+	_registerCloseHandlers: function( event, target ) {
+		var events = {
+			keyup: function( event ) {
+				if ( event.keyCode === $.ui.keyCode.ESCAPE ) {
+					var fakeEvent = $.Event( event );
+					fakeEvent.currentTarget = target[ 0 ];
+					this.close( fakeEvent, true );
+				}
+			}
+		};
+
+		// Only bind remove handler for delegated targets. Non-delegated
+		// tooltips will handle this in destroy.
+		if ( target[ 0 ] !== this.element[ 0 ] ) {
+			events.remove = function() {
+				this._removeTooltip( this._find( target ).tooltip );
+			};
+		}
+
+		if ( !event || event.type === "mouseover" ) {
+			events.mouseleave = "close";
+		}
+		if ( !event || event.type === "focusin" ) {
+			events.focusout = "close";
+		}
+		this._on( true, target, events );
+	},
+
+	close: function( event ) {
+		var tooltip,
+			that = this,
+			target = $( event ? event.currentTarget : this.element ),
+			tooltipData = this._find( target );
+
+		// The tooltip may already be closed
+		if ( !tooltipData ) {
+
+			// We set ui-tooltip-open immediately upon open (in open()), but only set the
+			// additional data once there's actually content to show (in _open()). So even if the
+			// tooltip doesn't have full data, we always remove ui-tooltip-open in case we're in
+			// the period between open() and _open().
+			target.removeData( "ui-tooltip-open" );
+			return;
+		}
+
+		tooltip = tooltipData.tooltip;
+
+		// Disabling closes the tooltip, so we need to track when we're closing
+		// to avoid an infinite loop in case the tooltip becomes disabled on close
+		if ( tooltipData.closing ) {
+			return;
+		}
+
+		// Clear the interval for delayed tracking tooltips
+		clearInterval( this.delayedShow );
+
+		// Only set title if we had one before (see comment in _open())
+		// If the title attribute has changed since open(), don't restore
+		if ( target.data( "ui-tooltip-title" ) && !target.attr( "title" ) ) {
+			target.attr( "title", target.data( "ui-tooltip-title" ) );
+		}
+
+		this._removeDescribedBy( target );
+
+		tooltipData.hiding = true;
+		tooltip.stop( true );
+		this._hide( tooltip, this.options.hide, function() {
+			that._removeTooltip( $( this ) );
+		} );
+
+		target.removeData( "ui-tooltip-open" );
+		this._off( target, "mouseleave focusout keyup" );
+
+		// Remove 'remove' binding only on delegated targets
+		if ( target[ 0 ] !== this.element[ 0 ] ) {
+			this._off( target, "remove" );
+		}
+		this._off( this.document, "mousemove" );
+
+		if ( event && event.type === "mouseleave" ) {
+			$.each( this.parents, function( id, parent ) {
+				$( parent.element ).attr( "title", parent.title );
+				delete that.parents[ id ];
+			} );
+		}
+
+		tooltipData.closing = true;
+		this._trigger( "close", event, { tooltip: tooltip } );
+		if ( !tooltipData.hiding ) {
+			tooltipData.closing = false;
+		}
+	},
+
+	_tooltip: function( element ) {
+		var tooltip = $( "<div>" ).attr( "role", "tooltip" ),
+			content = $( "<div>" ).appendTo( tooltip ),
+			id = tooltip.uniqueId().attr( "id" );
+
+		this._addClass( content, "ui-tooltip-content" );
+		this._addClass( tooltip, "ui-tooltip", "ui-widget ui-widget-content" );
+
+		tooltip.appendTo( this._appendTo( element ) );
+
+		return this.tooltips[ id ] = {
+			element: element,
+			tooltip: tooltip
+		};
+	},
+
+	_find: function( target ) {
+		var id = target.data( "ui-tooltip-id" );
+		return id ? this.tooltips[ id ] : null;
+	},
+
+	_removeTooltip: function( tooltip ) {
+		tooltip.remove();
+		delete this.tooltips[ tooltip.attr( "id" ) ];
+	},
+
+	_appendTo: function( target ) {
+		var element = target.closest( ".ui-front, dialog" );
+
+		if ( !element.length ) {
+			element = this.document[ 0 ].body;
+		}
+
+		return element;
+	},
+
+	_destroy: function() {
+		var that = this;
+
+		// Close open tooltips
+		$.each( this.tooltips, function( id, tooltipData ) {
+
+			// Delegate to close method to handle common cleanup
+			var event = $.Event( "blur" ),
+				element = tooltipData.element;
+			event.target = event.currentTarget = element[ 0 ];
+			that.close( event, true );
+
+			// Remove immediately; destroying an open tooltip doesn't use the
+			// hide animation
+			$( "#" + id ).remove();
+
+			// Restore the title
+			if ( element.data( "ui-tooltip-title" ) ) {
+
+				// If the title attribute has changed since open(), don't restore
+				if ( !element.attr( "title" ) ) {
+					element.attr( "title", element.data( "ui-tooltip-title" ) );
+				}
+				element.removeData( "ui-tooltip-title" );
+			}
+		} );
+		this.liveRegion.remove();
+	}
+} );
+
+// DEPRECATED
+// TODO: Switch return back to widget declaration at top of file when this is removed
+if ( $.uiBackCompat !== false ) {
+
+	// Backcompat for tooltipClass option
+	$.widget( "ui.tooltip", $.ui.tooltip, {
+		options: {
+			tooltipClass: null
+		},
+		_tooltip: function() {
+			var tooltipData = this._superApply( arguments );
+			if ( this.options.tooltipClass ) {
+				tooltipData.tooltip.addClass( this.options.tooltipClass );
+			}
+			return tooltipData;
+		}
+	} );
+}
+
+var widgetsTooltip = $.ui.tooltip;
+
+
+
+
+}));
\ No newline at end of file
