Copying //UE4/Dev-Core to //UE4/Dev-Main (Source: //UE4/Dev-Core @ 3208226)

diff --git a/Engine/Source/Editor/UnrealEd/Private/StaticLightingSystem/StaticLightingSystem.cpp b/Engine/Source/Editor/UnrealEd/Private/StaticLightingSystem/StaticLightingSystem.cpp
--- a/Engine/Source/Editor/UnrealEd/Private/StaticLightingSystem/StaticLightingSystem.cpp
+++ b/Engine/Source/Editor/UnrealEd/Private/StaticLightingSystem/StaticLightingSystem.cpp
@@ -208,7 +208,7 @@ void FStaticLightingManager::SendBuildDoneNotification( bool AutoApplyFailed )
 {
 	FText CompletedText = LOCTEXT("LightBuildDoneMessage", "Lighting build completed");
 
-	if (ActiveStaticLightingSystem != StaticLightingSystems.Last().GetOwnedPointer() && ActiveStaticLightingSystem->LightingScenario)
+	if (ActiveStaticLightingSystem != StaticLightingSystems.Last().Get() && ActiveStaticLightingSystem->LightingScenario)
 	{
 		FString PackageName = FPackageName::GetShortName(ActiveStaticLightingSystem->LightingScenario->GetOutermost()->GetName());
 		CompletedText = FText::Format(LOCTEXT("LightScenarioBuildDoneMessage", "{0} Lighting Scenario completed"), FText::FromString(PackageName));
@@ -252,18 +252,16 @@ void FStaticLightingManager::CreateStaticLightingSystem(const FLightingBuildOpti
 		{
 			if (Level->bIsLightingScenario && Level->bIsVisible)
 	{
-				StaticLightingSystems.AddZeroed();
-				StaticLightingSystems.Last() = new FStaticLightingSystem(Options, GWorld, Level);
+				StaticLightingSystems.Emplace(new FStaticLightingSystem(Options, GWorld, Level));
 			}
 		}
 
 		if (StaticLightingSystems.Num() == 0)
 		{
-			StaticLightingSystems.AddZeroed();
-			StaticLightingSystems.Last() = new FStaticLightingSystem(Options, GWorld, NULL);
+			StaticLightingSystems.Emplace(new FStaticLightingSystem(Options, GWorld, nullptr));
 		}
 
-		ActiveStaticLightingSystem = StaticLightingSystems[0];
+		ActiveStaticLightingSystem = StaticLightingSystems[0].Get();
 
 		bool bSuccess = ActiveStaticLightingSystem->BeginLightmassProcess();
 
@@ -298,12 +296,12 @@ void FStaticLightingManager::UpdateBuildLighting()
 
 		if (ActiveStaticLightingSystem && ActiveStaticLightingSystem->CurrentBuildStage == FStaticLightingSystem::Finished)
 	{
-			ActiveStaticLightingSystem = NULL;
+			ActiveStaticLightingSystem = nullptr;
 			StaticLightingSystems.RemoveAt(0);
 
 			if (StaticLightingSystems.Num() > 0)
 			{
-				ActiveStaticLightingSystem = StaticLightingSystems[0];
+				ActiveStaticLightingSystem = StaticLightingSystems[0].Get();
 
 				bool bSuccess = ActiveStaticLightingSystem->BeginLightmassProcess();
 
diff --git a/Engine/Source/Editor/UnrealEd/Private/UnrealEdSrv.cpp b/Engine/Source/Editor/UnrealEd/Private/UnrealEdSrv.cpp
--- a/Engine/Source/Editor/UnrealEd/Private/UnrealEdSrv.cpp
+++ b/Engine/Source/Editor/UnrealEd/Private/UnrealEdSrv.cpp
@@ -236,9 +236,9 @@ UPackage* UUnrealEdEngine::GeneratePackageThumbnailsIfRequired( const TCHAR* Str
 			Packages.Add( Pkg );
 
 			// Allocate a new thumbnail map if we need one
-			if( !Pkg->ThumbnailMap.IsValid() )
+			if( !Pkg->ThumbnailMap )
 			{
-				Pkg->ThumbnailMap.Reset( new FThumbnailMap() );
+				Pkg->ThumbnailMap = MakeUnique<FThumbnailMap>();
 			}
 
 			// OK, now query all of the browsable objects in the package we're about to save
diff --git a/Engine/Source/Editor/UnrealEd/Public/Kismet2/EnumEditorUtils.h b/Engine/Source/Editor/UnrealEd/Public/Kismet2/EnumEditorUtils.h
--- a/Engine/Source/Editor/UnrealEd/Public/Kismet2/EnumEditorUtils.h
+++ b/Engine/Source/Editor/UnrealEd/Public/Kismet2/EnumEditorUtils.h
@@ -12,10 +12,10 @@ class UUserDefinedEnum;
 class UNREALED_API FEnumEditorUtils
 {
 	static void PrepareForChange(const UUserDefinedEnum* Enum);
-	static void BroadcastChanges(const UUserDefinedEnum* Enum, const TArray<TPair<FName, uint8>>& OldNames, bool bResolveData = true);
+	static void BroadcastChanges(const UUserDefinedEnum* Enum, const TArray<TPair<FName, int64>>& OldNames, bool bResolveData = true);
 
 	/** copy full enumeratos names from given enum to OutEnumNames, the last '_MAX' enumerator is skipped */
-	static void CopyEnumeratorsWithoutMax(const UEnum* Enum, TArray<TPair<FName, uint8>>& OutEnumNames);
+	static void CopyEnumeratorsWithoutMax(const UEnum* Enum, TArray<TPair<FName, int64>>& OutEnumNames);
 public:
 
 	enum EEnumEditorChangeInfo
@@ -75,7 +75,7 @@ public:
 	 *
 	 *	@return new enum 
 	 */
-	static int32 ResolveEnumerator(const UEnum* Enum, FArchive& Ar, int32 EnumeratorValue);
+	static int64 ResolveEnumerator(const UEnum* Enum, FArchive& Ar, int64 EnumeratorValue);
 
 	//DISPLAY NAME
 	static FString GetEnumeratorDisplayName(const UUserDefinedEnum* Enum, int32 EnumeratorIndex);
diff --git a/Engine/Source/Editor/WorldBrowser/Private/Tiles/SWorldComposition.cpp b/Engine/Source/Editor/WorldBrowser/Private/Tiles/SWorldComposition.cpp
--- a/Engine/Source/Editor/WorldBrowser/Private/Tiles/SWorldComposition.cpp
+++ b/Engine/Source/Editor/WorldBrowser/Private/Tiles/SWorldComposition.cpp
@@ -92,7 +92,7 @@ public:
 
 	void Construct(const FArguments& InArgs)
 	{
-		ZoomLevels = new FWorldZoomLevelsContainer();
+		ZoomLevels = MakeUnique<FWorldZoomLevelsContainer>();
 
 		SNodePanel::Construct();
 	
diff --git a/Engine/Source/Editor/WorldBrowser/Private/Tiles/WorldTileCollectionModel.cpp b/Engine/Source/Editor/WorldBrowser/Private/Tiles/WorldTileCollectionModel.cpp
--- a/Engine/Source/Editor/WorldBrowser/Private/Tiles/WorldTileCollectionModel.cpp
+++ b/Engine/Source/Editor/WorldBrowser/Private/Tiles/WorldTileCollectionModel.cpp
@@ -589,7 +589,7 @@ bool FWorldTileCollectionModel::GetPlayerView(FVector& Location, FRotator& Rotat
 		UWorld* SimulationWorld = GetSimulationWorld();
 		for (FConstPlayerControllerIterator Iterator = SimulationWorld->GetPlayerControllerIterator(); Iterator; ++Iterator)
 		{
-			APlayerController* PlayerActor = *Iterator;
+			APlayerController* PlayerActor = Iterator->Get();
 			PlayerActor->GetPlayerViewPoint(Location, Rotation);
 			return true;
 		}
diff --git a/Engine/Source/Programs/AutomationTool/AutomationUtils/ProjectParams.cs b/Engine/Source/Programs/AutomationTool/AutomationUtils/ProjectParams.cs
--- a/Engine/Source/Programs/AutomationTool/AutomationUtils/ProjectParams.cs
+++ b/Engine/Source/Programs/AutomationTool/AutomationUtils/ProjectParams.cs
@@ -257,7 +257,8 @@ namespace AutomationTool
             this.CookOnTheFlyStreaming = InParams.CookOnTheFlyStreaming;
             this.UnversionedCookedContent = InParams.UnversionedCookedContent;
 			this.EncryptIniFiles = InParams.EncryptIniFiles;
-            this.SkipCookingEditorContent = InParams.SkipCookingEditorContent;
+			this.EncryptEverything = InParams.EncryptEverything;
+			this.SkipCookingEditorContent = InParams.SkipCookingEditorContent;
             this.NumCookersToSpawn = InParams.NumCookersToSpawn;
 			this.FileServer = InParams.FileServer;
 			this.DedicatedServer = InParams.DedicatedServer;
@@ -377,7 +378,8 @@ namespace AutomationTool
             bool? CookOnTheFlyStreaming = null,
             bool? UnversionedCookedContent = null,
 			bool? EncryptIniFiles = null,
-            bool? SkipCookingEditorContent = null,
+			bool? EncryptEverything = null,
+			bool? SkipCookingEditorContent = null,
             int? NumCookersToSpawn = null,
             string AdditionalCookerOptions = null,
             string BasedOnReleaseVersion = null,
@@ -551,6 +553,7 @@ namespace AutomationTool
             this.CookOnTheFlyStreaming = GetParamValueIfNotSpecified(Command, CookOnTheFlyStreaming, this.CookOnTheFlyStreaming, "cookontheflystreaming");
             this.UnversionedCookedContent = GetParamValueIfNotSpecified(Command, UnversionedCookedContent, this.UnversionedCookedContent, "UnversionedCookedContent");
 			this.EncryptIniFiles = GetParamValueIfNotSpecified(Command, EncryptIniFiles, this.EncryptIniFiles, "EncryptIniFiles");
+			this.EncryptEverything = GetParamValueIfNotSpecified(Command, EncryptEverything, this.EncryptEverything, "EncryptEverything");
 			this.SkipCookingEditorContent = GetParamValueIfNotSpecified(Command, SkipCookingEditorContent, this.SkipCookingEditorContent, "SkipCookingEditorContent");
             if (NumCookersToSpawn.HasValue)
             {
@@ -1302,6 +1305,11 @@ namespace AutomationTool
 		/// </summary>
 		public bool EncryptIniFiles;
 
+		/// <summary>
+		/// Encrypt all files which are packaged into the pak file.  Only valid when encryption keys and building pak file specified. 
+		/// </summary>
+		public bool EncryptEverything;
+
 		/// <summary>
 		/// put -debug on the editorexe commandline
 		/// </summary>
@@ -2307,6 +2315,7 @@ namespace AutomationTool
 				CommandUtils.LogLog("CookOnTheFlyStreaming={0}", CookOnTheFlyStreaming);
 				CommandUtils.LogLog("UnversionedCookedContent={0}", UnversionedCookedContent);
 				CommandUtils.LogLog("EncryptIniFiles={0}", EncryptIniFiles);
+				CommandUtils.LogLog("EncryptEverything={0}", EncryptEverything);
 				CommandUtils.LogLog("SkipCookingEditorContent={0}", SkipCookingEditorContent);
                 CommandUtils.LogLog("NumCookersToSpawn={0}", NumCookersToSpawn);
                 CommandUtils.LogLog("GeneratePatch={0}", GeneratePatch);
diff --git a/Engine/Source/Programs/AutomationTool/Scripts/CopyBuildToStagingDirectory.Automation.cs b/Engine/Source/Programs/AutomationTool/Scripts/CopyBuildToStagingDirectory.Automation.cs
--- a/Engine/Source/Programs/AutomationTool/Scripts/CopyBuildToStagingDirectory.Automation.cs
+++ b/Engine/Source/Programs/AutomationTool/Scripts/CopyBuildToStagingDirectory.Automation.cs
@@ -25,6 +25,16 @@ public partial class Project : CommandUtils
 
 	private static readonly object SyncLock = new object();
 
+	private static readonly object EncryptionKeysLock = new object();
+	
+	private static void GetEncryptionKeys(ProjectParams InParams, UnrealTargetPlatform InTargetPlatform, out string[] OutRSAKeys, out string OutAESKey)
+	{
+		lock (EncryptionKeysLock)
+		{
+			UEBuildTarget.ParseEncryptionIni(new DirectoryReference(CommandUtils.GetDirectoryName(InParams.RawProjectPath.FullName)), InTargetPlatform, out OutRSAKeys, out OutAESKey);
+		}
+	}
+
 	/// <returns>The path for the BuildPatchTool executable depending on host platform.</returns>
 	private static string GetBuildPatchToolExecutable()
 	{
@@ -65,7 +75,7 @@ public partial class Project : CommandUtils
 	/// </summary>
 	/// <param name="Filename"></param>
 	/// <param name="ResponseFile"></param>
-	private static void WritePakResponseFile(string Filename, Dictionary<string, string> ResponseFile, bool Compressed, bool EncryptIniFiles)
+	private static void WritePakResponseFile(string Filename, Dictionary<string, string> ResponseFile, bool Compressed, bool EncryptIniFiles, bool EncryptEverything)
 	{
 		using (var Writer = new StreamWriter(Filename, false, new System.Text.UTF8Encoding(true)))
 		{
@@ -77,7 +87,7 @@ public partial class Project : CommandUtils
 					Line += " -compress";
 				}
 				
-				if (Path.GetExtension(Entry.Key).Contains(".ini") && EncryptIniFiles)
+				if (EncryptEverything || (Path.GetExtension(Entry.Key).Contains(".ini") && EncryptIniFiles))
 				{
 					Line += " -encrypt";
 				}
@@ -98,7 +108,7 @@ public partial class Project : CommandUtils
 		return Result;
 	}
 
-    static public void RunUnrealPak(Dictionary<string, string> UnrealPakResponseFile, string OutputLocation, string EncryptionKeys, string PakOrderFileLocation, string PlatformOptions, bool Compressed, bool EncryptIniFiles, String PatchSourceContentPath)
+    static public void RunUnrealPak(Dictionary<string, string> UnrealPakResponseFile, string OutputLocation, string PakOrderFileLocation, string PlatformOptions, bool Compressed, bool EncryptIniFiles, bool EncryptEverything, String PatchSourceContentPath, String EngineDir, String ProjectDir, String Platform)
 	{
 		if (UnrealPakResponseFile.Count < 1)
 		{
@@ -106,16 +116,14 @@ public partial class Project : CommandUtils
 		}
 		string PakName = Path.GetFileNameWithoutExtension(OutputLocation);
 		string UnrealPakResponseFileName = CombinePaths(CmdEnv.LogFolder, "PakList_" + PakName + ".txt");
-		WritePakResponseFile(UnrealPakResponseFileName, UnrealPakResponseFile, Compressed, EncryptIniFiles);
+		WritePakResponseFile(UnrealPakResponseFileName, UnrealPakResponseFile, Compressed, EncryptIniFiles, EncryptEverything);
 
 		var UnrealPakExe = CombinePaths(CmdEnv.LocalRoot, "Engine/Binaries/Win64/UnrealPak.exe");
-
 		Log("Running UnrealPak *******");
 		string CmdLine = CommandUtils.MakePathSafeToUseWithCommandLine(OutputLocation) + " -create=" + CommandUtils.MakePathSafeToUseWithCommandLine(UnrealPakResponseFileName);
-		if (!String.IsNullOrEmpty(EncryptionKeys))
-		{
-			CmdLine += " -sign=" + CommandUtils.MakePathSafeToUseWithCommandLine(EncryptionKeys);
-		}
+
+		CmdLine += String.Format(" -encryptionini -enginedir=\"{0}\" -projectdir=\"{1}\" -platform={2}", EngineDir, ProjectDir, Platform);
+
 		if (GlobalCommandLine.Installed)
 		{
 			CmdLine += " -installed";
@@ -838,7 +846,8 @@ public partial class Project : CommandUtils
             string CloudDir = MakePathSafeToUseWithCommandLine(CombinePaths(ChunkInstallBasePath, "CloudDir"));
             InternalUtils.SafeDeleteDirectory(CloudDir, true);
         }
-        CreatePak(Params, SC, UnrealPakResponseFile, SC.ShortProjectName, Params.SignPak);
+
+        CreatePak(Params, SC, UnrealPakResponseFile, SC.ShortProjectName);
 	}
 
 	/// <summary>
@@ -937,7 +946,7 @@ public partial class Project : CommandUtils
 	/// <param name="SC"></param>
 	/// <param name="UnrealPakResponseFile"></param>
 	/// <param name="PakName"></param>
-	private static void CreatePak(ProjectParams Params, DeploymentContext SC, Dictionary<string, string> UnrealPakResponseFile, string PakName, string EncryptionKeys)
+	private static void CreatePak(ProjectParams Params, DeploymentContext SC, Dictionary<string, string> UnrealPakResponseFile, string PakName)
 	{
         bool bShouldGeneratePatch = Params.IsGeneratingPatch && SC.StageTargetPlatform.GetPlatformPatchesWithDiffPak(Params, SC);
 
@@ -1027,7 +1036,10 @@ public partial class Project : CommandUtils
 			}
 			if (!bCopiedExistingPak)
 			{
-				RunUnrealPak(UnrealPakResponseFile, OutputLocation, EncryptionKeys, PakOrderFileLocation, SC.StageTargetPlatform.GetPlatformPakCommandLine(), Params.Compressed, Params.EncryptIniFiles || PackageSettingsEncryptIniFiles, PatchSourceContentPath);
+				String EngineDir = CommandUtils.CombinePaths(CommandUtils.CmdEnv.LocalRoot, "Engine");
+				String ProjectDir = CommandUtils.GetDirectoryName(Params.RawProjectPath.FullName);
+				String Platform = ConfigCacheIni.GetIniPlatformName(SC.StageTargetPlatform.IniPlatformType);
+				RunUnrealPak(UnrealPakResponseFile, OutputLocation, PakOrderFileLocation, SC.StageTargetPlatform.GetPlatformPakCommandLine(), Params.Compressed, Params.EncryptIniFiles || PackageSettingsEncryptIniFiles, Params.EncryptEverything, PatchSourceContentPath, EngineDir, ProjectDir, Platform);
 			}
 		}
 
@@ -1210,11 +1222,11 @@ public partial class Project : CommandUtils
             string CloudDir = MakePathSafeToUseWithCommandLine(CombinePaths(ChunkInstallBasePath, "CloudDir"));
             InternalUtils.SafeDeleteDirectory(CloudDir, true);
         }
-
-        System.Threading.Tasks.Parallel.ForEach(PakPairs, (PakPair) =>
+		
+		System.Threading.Tasks.Parallel.ForEach(PakPairs, (PakPair) =>
 		{
 			var ChunkName = Path.GetFileNameWithoutExtension(PakPair.Item2);
-			CreatePak(Params, SC, PakPair.Item1, ChunkName, Params.SignPak);
+			CreatePak(Params, SC, PakPair.Item1, ChunkName);
 		});
 
 		String ChunkLayerFilename = CombinePaths(GetTmpPackagingPath(Params, SC), GetChunkPakLayerListName());
diff --git a/Engine/Source/Programs/CrashReporter/CrashReportCommon/AmazonClient.cs b/Engine/Source/Programs/CrashReporter/CrashReportCommon/AmazonClient.cs
--- a/Engine/Source/Programs/CrashReporter/CrashReportCommon/AmazonClient.cs
+++ b/Engine/Source/Programs/CrashReporter/CrashReportCommon/AmazonClient.cs
@@ -1,11 +1,6 @@
 ﻿using System;
 using System.Collections.Generic;
-using System.IO;
-using System.Linq;
 using System.Net;
-using System.Text;
-using System.Threading.Tasks;
-using Amazon.Route53.Model;
 using Amazon.Runtime;
 using Amazon.S3;
 using Amazon.S3.Model;
diff --git a/Engine/Source/Programs/CrashReporter/CrashReportCommon/Common.cs b/Engine/Source/Programs/CrashReporter/CrashReportCommon/Common.cs
--- a/Engine/Source/Programs/CrashReporter/CrashReportCommon/Common.cs
+++ b/Engine/Source/Programs/CrashReporter/CrashReportCommon/Common.cs
@@ -4,12 +4,6 @@ using System;
 using System.IO;
 using System.Runtime.Serialization;
 using System.Xml.Serialization;
-using System.Collections;
-using System.Collections.Generic;
-using Tools.DotNETCommon.XmlHandler;
-using System.Xml;
-using System.Linq;
-using System.Security;
 
 namespace Tools.CrashReporter.CrashReportCommon
 {
@@ -29,11 +23,6 @@ namespace Tools.CrashReporter.CrashReportCommon
 		/// <summary>A optional message to explain any return values.</summary>
 		[XmlAttribute]
 		public string Message = "";
-
-		/// <summary>A simple default constructor to allow Xml serialisation.</summary>
-		public CrashReporterResult()
-		{
-		}
 	}
 
 	/// <summary>
diff --git a/Engine/Source/Programs/CrashReporter/CrashReportCommon/CrashContext.cs b/Engine/Source/Programs/CrashReporter/CrashReportCommon/CrashContext.cs
--- a/Engine/Source/Programs/CrashReporter/CrashReportCommon/CrashContext.cs
+++ b/Engine/Source/Programs/CrashReporter/CrashReportCommon/CrashContext.cs
@@ -8,8 +8,6 @@ using System.Collections.Generic;
 using Tools.DotNETCommon.XmlHandler;
 using System.Xml;
 using System.Linq;
-using System.Security;
-using System.Diagnostics;
 
 namespace Tools.CrashReporter.CrashReportCommon
 {
diff --git a/Engine/Source/Programs/CrashReporter/CrashReportCommon/CrashReporterConstants.cs b/Engine/Source/Programs/CrashReporter/CrashReportCommon/CrashReporterConstants.cs
--- a/Engine/Source/Programs/CrashReporter/CrashReportCommon/CrashReporterConstants.cs
+++ b/Engine/Source/Programs/CrashReporter/CrashReportCommon/CrashReporterConstants.cs
@@ -1,8 +1,5 @@
 ﻿// Copyright 1998-2016 Epic Games, Inc. All Rights Reserved.
 
-using System;
-using System.Xml.Serialization;
-
 namespace Tools.CrashReporter.CrashReportCommon
 {
 	/// <summary>
diff --git a/Engine/Source/Programs/CrashReporter/CrashReportCommon/JiraConnection.cs b/Engine/Source/Programs/CrashReporter/CrashReportCommon/JiraConnection.cs
--- a/Engine/Source/Programs/CrashReporter/CrashReportCommon/JiraConnection.cs
+++ b/Engine/Source/Programs/CrashReporter/CrashReportCommon/JiraConnection.cs
@@ -2,14 +2,9 @@
 
 using System;
 using System.Collections.Generic;
-using System.Diagnostics;
 using System.IO;
-using System.Linq;
 using System.Net;
 using System.Text;
-using System.Threading; 
-using System.Threading.Tasks;
-using System.Web;
 using System.Web.Script.Serialization;
 
 /// <summary>
diff --git a/Engine/Source/Programs/CrashReporter/CrashReportCommon/SlackWriter.cs b/Engine/Source/Programs/CrashReporter/CrashReportCommon/SlackWriter.cs
--- a/Engine/Source/Programs/CrashReporter/CrashReportCommon/SlackWriter.cs
+++ b/Engine/Source/Programs/CrashReporter/CrashReportCommon/SlackWriter.cs
@@ -1,11 +1,7 @@
 ﻿using System;
-using System.Collections.Generic;
-using System.Collections.Specialized;
-using System.Linq;
 using System.Net;
 using System.Text;
 using System.Threading.Tasks;
-using System.Web;
 
 namespace Tools.CrashReporter.CrashReportCommon
 {
diff --git a/Engine/Source/Programs/CrashReporter/CrashReportCommon/Stats.cs b/Engine/Source/Programs/CrashReporter/CrashReportCommon/Stats.cs
--- a/Engine/Source/Programs/CrashReporter/CrashReportCommon/Stats.cs
+++ b/Engine/Source/Programs/CrashReporter/CrashReportCommon/Stats.cs
@@ -1,9 +1,5 @@
 ﻿using System;
 using System.Collections.Generic;
-using System.Linq;
-using System.Text;
-using System.Threading.Tasks;
-using System.Web.ModelBinding;
 
 namespace Tools.CrashReporter.CrashReportCommon
 {
diff --git a/Engine/Source/Programs/CrashReporter/CrashReportCommon/WerReport.cs b/Engine/Source/Programs/CrashReporter/CrashReportCommon/WerReport.cs
--- a/Engine/Source/Programs/CrashReporter/CrashReportCommon/WerReport.cs
+++ b/Engine/Source/Programs/CrashReporter/CrashReportCommon/WerReport.cs
@@ -1,15 +1,5 @@
 // Copyright 1998-2016 Epic Games, Inc. All Rights Reserved.
 
-using System;
-using System.IO;
-using System.Xml.Serialization;
-using System.Collections;
-using System.Collections.Generic;
-using Tools.DotNETCommon.XmlHandler;
-using System.Xml;
-using System.Linq;
-using System.Security;
-
 namespace Tools.CrashReporter.CrashReportCommon
 {
 	/// <summary>
@@ -27,11 +17,6 @@ namespace Tools.CrashReporter.CrashReportCommon
 		public WERReportMetadataDynamicSignatures DynamicSignatures;
 		/// <summary>Information about the system the crash occurred on.</summary>
 		public WERReportMetadataSystemInformation SystemInformation;
-
-		/// <summary>A simple default constructor to allow Xml serialisation.</summary>
-		public WERReportMetadata()
-		{
-		}
 	}
 
 	/// <summary>
diff --git a/Engine/Source/Programs/CrashReporter/CrashReportProcess/Config.cs b/Engine/Source/Programs/CrashReporter/CrashReportProcess/Config.cs
--- a/Engine/Source/Programs/CrashReporter/CrashReportProcess/Config.cs
+++ b/Engine/Source/Programs/CrashReporter/CrashReportProcess/Config.cs
@@ -1,9 +1,4 @@
-﻿using System;
-using System.Collections.Generic;
-using System.IO;
-using System.Linq;
-using System.Text;
-using System.Threading.Tasks;
+﻿using System.IO;
 using System.Xml;
 using System.Xml.Serialization;
 using System.Reflection;
@@ -88,12 +83,6 @@ namespace Tools.CrashReporter.CrashReportProcess
 		[XmlElement]
 		public string PS4LandingZone { get; set; }
 
-		/// <summary>
-		/// Folder where failed reports are moved.
-		/// </summary>
-		[XmlElement]
-		public string InvalidReportsDirectory { get; set; }
-
 		/// <summary>
 		/// Number of reports in a queue at which reports will start to be discarded to stop a backlog from growing uncontrollably.
 		/// </summary>
@@ -113,34 +102,10 @@ namespace Tools.CrashReporter.CrashReportProcess
 		public string DepotIndex { get; set; }
 
 		/// <summary>
-		/// Whether MDD is synched from source control to get the latest files when the CRP starts. Switch off to manually deploy MDD.
-		/// </summary>
-		[XmlElement]
-		public bool bSyncMinidumpDiagnostics { get; set; }
-
-		/// <summary>
-		/// String specifying the binaries that will be synched from source control to get the latest MinidumpDiagnostics.
-		/// </summary>
-		[XmlElement]
-		public string SyncBinariesFromDepot { get; set; }
-
-		/// <summary>
-		/// String specifying the config files that will be synched from source control for MinidumpDiagnostics.
+		/// String specifying the local path of the MinidumpDiagnostics exe
 		/// </summary>
 		[XmlElement]
-		public string SyncConfigFromDepot { get; set; }
-
-		/// <summary>
-		/// String specifying the third-party files that will be synched from source control for MinidumpDiagnostics. e.g. OpenSSL
-		/// </summary>
-		[XmlElement]
-		public string SyncThirdPartyFromDepot { get; set; }
-
-		/// <summary>
-		/// String specifying the local path of the folder containing the MinidumpDiagnostics exe
-		/// </summary>
-		[XmlElement]
-		public string MDDBinariesFolderInDepot { get; set; }
+		public string MDDExecutablePath { get; set; }
 
 		/// <summary>
 		/// String specifying the path to the folder used by MinidumpDiagnostics for the PDB cache
@@ -311,6 +276,12 @@ namespace Tools.CrashReporter.CrashReportProcess
 		[XmlElement]
 		public bool CrashFilesToAWS { get; set; }
 
+		/// <summary>
+		/// Should we save invalid reports that fail to process to S3?
+		/// </summary>
+		[XmlElement]
+		public bool InvalidReportsToAWS { get; set; }
+
 		/// <summary>
 		/// AWSSDK AWS S3 bucket used for output of crash reporter files (optional)
 		/// </summary>
@@ -323,6 +294,12 @@ namespace Tools.CrashReporter.CrashReportProcess
 		[XmlElement]
 		public string AWSS3OutputKeyPrefix { get; set; }
 
+		/// <summary>
+		/// AWSSDK AWS S3 path/key prefix used for writing invalid reports (optional)
+		/// </summary>
+		[XmlElement]
+		public string AWSS3InvalidKeyPrefix { get; set; }
+
 		/// <summary>
 		/// Buffer size used to decompress zlib archives taken from S3
 		/// </summary>
@@ -359,6 +336,12 @@ namespace Tools.CrashReporter.CrashReportProcess
 		[XmlElement]
 		public float DiskSpaceAlertPercent { get; set; }
 
+		/// <summary>
+		/// Switch on perf monitoring via StatusReporting. Adds an extra report with perf info.
+		/// </summary>
+		[XmlElement]
+		public bool MonitorPerformance { get; set; }
+
 		/// <summary>
 		/// Get the default config object (lazy loads it on first access)
 		/// </summary>
@@ -424,7 +407,10 @@ namespace Tools.CrashReporter.CrashReportProcess
 			{
 				LoadedConfig.PS4LandingZone = Path.Combine(LoadedConfig.DebugTestingFolder, "PS4LandingZone");
 			}
-			LoadedConfig.InvalidReportsDirectory = Path.Combine(LoadedConfig.DebugTestingFolder, "InvalidReportsDirectory");
+			if (!string.IsNullOrWhiteSpace(LoadedConfig.MDDExecutablePath))
+			{
+				LoadedConfig.MDDExecutablePath = Path.Combine(LoadedConfig.DebugTestingFolder, "MinidumpDiagnostics", "Engine", "Binaries", "Win64", "MinidumpDiagnostics.exe");
+			}
 			LoadedConfig.VersionString += " debugbuild";
 			LoadedConfig.AWSCredentialsFilepath = Path.Combine(LoadedConfig.DebugTestingFolder, "AWS", "credentials.ini");
 			if (!string.IsNullOrWhiteSpace(LoadedConfig.ProcessedReportsIndexPath))
@@ -433,6 +419,10 @@ namespace Tools.CrashReporter.CrashReportProcess
 			}
 			LoadedConfig.CrashReportWebSite = string.Empty;
 			LoadedConfig.AWSS3OutputKeyPrefix = LoadedConfig.AWSS3OutputKeyPrefix.Replace("prod", "test");
+			LoadedConfig.AWSS3InvalidKeyPrefix = LoadedConfig.AWSS3InvalidKeyPrefix.Replace("prod", "test");
+
+			LoadedConfig.MinDesiredMemoryQueueSize = 5;
+			LoadedConfig.MaxMemoryQueueSize = 15;
 
 #if SLACKTESTING
 			LoadedConfig.SlackUsername = "CrashReportProcess_TESTING_IgnoreMe";
diff --git a/Engine/Source/Programs/CrashReporter/CrashReportProcess/Config/CrashReportProcess.config b/Engine/Source/Programs/CrashReporter/CrashReportProcess/Config/CrashReportProcess.config
--- a/Engine/Source/Programs/CrashReporter/CrashReportProcess/Config/CrashReportProcess.config
+++ b/Engine/Source/Programs/CrashReporter/CrashReportProcess/Config/CrashReportProcess.config
@@ -12,15 +12,10 @@
 	<ExternalLandingZone>MyExternallCrashLandingZoneFolderPath</ExternalLandingZone>
 	<DataRouterLandingZone>MyDataRouterCrashLandingZoneFolderPath</DataRouterLandingZone>
 	<PS4LandingZone>MyPS4CrashLandingZoneFolderPath</PS4LandingZone>
-	<InvalidReportsDirectory>MyInvalidCrashReportsFolderPath</InvalidReportsDirectory>
 	<QueueLowerLimitForDiscard>5000</QueueLowerLimitForDiscard>
 	<QueueUpperLimitForDiscard>20000</QueueUpperLimitForDiscard>
 	<DepotIndex>MyDepotCustomizationhPassedToMDD</DepotIndex>
-	<bSyncMinidumpDiagnostics>true</bSyncMinidumpDiagnostics>
-	<SyncBinariesFromDepot>MyMDDBinariesPathInP4Depot</SyncBinariesFromDepot>
-	<SyncConfigFromDepot>MyMDDConfigPathInP4Depot</SyncConfigFromDepot>
-	<SyncThirdPartyFromDepot>MyThirdPartyBinariesPathInP4Depot</SyncThirdPartyFromDepot>
-	<MDDBinariesFolderInDepot>MyMDDBinariesFolderPathInP4Depot</MDDBinariesFolderInDepot>
+	<MDDExecutablePath>MyMDDBinaryPath</MDDExecutablePath>
 	<MDDPDBCachePath>MyMDDPDBCacheFolderPath</MDDPDBCachePath>
 	<MDDPDBCacheSizeGB>500</MDDPDBCacheSizeGB>
 	<MDDPDBCacheMinFreeSpaceGB>32</MDDPDBCacheMinFreeSpaceGB>
@@ -49,12 +44,15 @@
 	<AWSS3ServiceOutputURL>https://s3.amazonaws.com</AWSS3ServiceOutputURL>
 	<CrashFilesToDisk>true</CrashFilesToDisk>
 	<CrashFilesToAWS>false</CrashFilesToAWS>
+	<InvalidReportsToAWS>false</InvalidReportsToAWS>
 	<AWSS3OutputBucket>crashreporterfiles</AWSS3OutputBucket>
 	<AWSS3OutputKeyPrefix>websitefiles</AWSS3OutputKeyPrefix>
+	<AWSS3InvalidKeyPrefix>invalidcrashreports</AWSS3InvalidKeyPrefix>
 	<MaxUncompressedS3RecordSize>33554432</MaxUncompressedS3RecordSize>
 	<ProcessedReportsIndexPath>MyProcessedReportsIndexFilePath</ProcessedReportsIndexPath>
 	<ReportsIndexRetentionDays>120</ReportsIndexRetentionDays>
 	<AddCrashRequestTimeoutMillisec>300000</AddCrashRequestTimeoutMillisec>
 	<AddCrashRetryDelayMillisec>20000</AddCrashRetryDelayMillisec>
 	<DiskSpaceAlertPercent>10.0</DiskSpaceAlertPercent>
+	<MonitorPerformance>false</MonitorPerformance>
 </Config>
\ No newline at end of file
diff --git a/Engine/Source/Programs/CrashReporter/CrashReportProcess/CrashReportProcessService.cs b/Engine/Source/Programs/CrashReporter/CrashReportProcess/CrashReportProcessService.cs
--- a/Engine/Source/Programs/CrashReporter/CrashReportProcess/CrashReportProcessService.cs
+++ b/Engine/Source/Programs/CrashReporter/CrashReportProcess/CrashReportProcessService.cs
@@ -2,20 +2,14 @@
 
 using System;
 using System.Collections.Generic;
-using System.ComponentModel;
-using System.Data;
-using System.Diagnostics;
 using System.IO;
 using System.Linq;
 using System.Reflection;
 using System.ServiceProcess;
-using System.Text;
-using System.Threading.Tasks;
 using Amazon.Runtime;
 using Amazon.S3;
 using Amazon.SQS;
 using Tools.CrashReporter.CrashReportCommon;
-using Tools.CrashReporter.CrashReportProcess.Properties;
 
 namespace Tools.CrashReporter.CrashReportProcess
 {
@@ -222,6 +216,10 @@ namespace Tools.CrashReporter.CrashReportProcess
 			FieldInfo[] EventNameFields = typeof(StatusReportingEventNames).GetFields(BindingFlags.Static | BindingFlags.Public);
 			StatusReporter.InitCounters(EventNameFields.Select(EventNameField => (string)EventNameField.GetValue(null)));
 
+			WriteEvent("Initializing Performance Mean Counters");
+			FieldInfo[] MeanNameFields = typeof(StatusReportingPerfMeanNames).GetFields(BindingFlags.Static | BindingFlags.Public);
+			StatusReporter.InitMeanCounters(MeanNameFields.Select(MeanNameField => (string)MeanNameField.GetValue(null)));
+
 			WriteEvent("Initializing Folder Monitors");
 			Dictionary<string, string> FoldersToMonitor = new Dictionary<string, string>();
 			FoldersToMonitor.Add(Config.Default.ProcessedReports, "Processed Reports");
@@ -230,7 +228,6 @@ namespace Tools.CrashReporter.CrashReportProcess
 			FoldersToMonitor.Add(Config.Default.InternalLandingZone, "CRR Landing Zone");
 			FoldersToMonitor.Add(Config.Default.DataRouterLandingZone, "Data Router Landing Zone");
 			FoldersToMonitor.Add(Config.Default.PS4LandingZone, "PS4 Landing Zone");
-			FoldersToMonitor.Add(Config.Default.InvalidReportsDirectory, "Invalid Reports");
 			FoldersToMonitor.Add(Assembly.GetExecutingAssembly().Location, "CRP Binaries and Logs");
 			FoldersToMonitor.Add(Config.Default.MDDPDBCachePath, "MDD PDB Cache");
 			StatusReporter.InitFolderMonitors(FoldersToMonitor);
diff --git a/Engine/Source/Programs/CrashReporter/CrashReportProcess/DataRouterConsumer.cs b/Engine/Source/Programs/CrashReporter/CrashReportProcess/DataRouterConsumer.cs
--- a/Engine/Source/Programs/CrashReporter/CrashReportProcess/DataRouterConsumer.cs
+++ b/Engine/Source/Programs/CrashReporter/CrashReportProcess/DataRouterConsumer.cs
@@ -1,11 +1,6 @@
 ﻿using System;
-using System.Collections.Generic;
 using System.IO;
-using System.Linq;
-using System.Text;
-using System.Threading.Tasks;
 using ProtoBuf;
-using Tools.CrashReporter.CrashReportCommon;
 
 namespace Tools.CrashReporter.CrashReportProcess
 {
diff --git a/Engine/Source/Programs/CrashReporter/CrashReportProcess/DataRouterReportQueue.cs b/Engine/Source/Programs/CrashReporter/CrashReportProcess/DataRouterReportQueue.cs
--- a/Engine/Source/Programs/CrashReporter/CrashReportProcess/DataRouterReportQueue.cs
+++ b/Engine/Source/Programs/CrashReporter/CrashReportProcess/DataRouterReportQueue.cs
@@ -1,14 +1,8 @@
 ﻿using System;
-using System.Collections.Generic;
 using System.IO;
 using System.IO.Compression;
 using System.Linq;
-using System.Net;
-using System.Text;
-using Amazon.Runtime;
-using Amazon.S3;
 using Amazon.S3.Model;
-using Amazon.SQS;
 using Amazon.SQS.Model;
 using Tools.CrashReporter.CrashReportCommon;
 
diff --git a/Engine/Source/Programs/CrashReporter/CrashReportProcess/ReceiverReportQueue.cs b/Engine/Source/Programs/CrashReporter/CrashReportProcess/ReceiverReportQueue.cs
--- a/Engine/Source/Programs/CrashReporter/CrashReportProcess/ReceiverReportQueue.cs
+++ b/Engine/Source/Programs/CrashReporter/CrashReportProcess/ReceiverReportQueue.cs
@@ -1,6 +1,4 @@
-﻿using System.IO;
-using System.Linq;
-
+﻿
 namespace Tools.CrashReporter.CrashReportProcess
 {
 	/// <summary>
diff --git a/Engine/Source/Programs/CrashReporter/CrashReportProcess/ReportIndex.cs b/Engine/Source/Programs/CrashReporter/CrashReportProcess/ReportIndex.cs
--- a/Engine/Source/Programs/CrashReporter/CrashReportProcess/ReportIndex.cs
+++ b/Engine/Source/Programs/CrashReporter/CrashReportProcess/ReportIndex.cs
@@ -1,9 +1,6 @@
 ﻿using System;
 using System.Collections.Generic;
 using System.IO;
-using System.Linq;
-using System.Text;
-using System.Threading.Tasks;
 
 namespace Tools.CrashReporter.CrashReportProcess
 {
diff --git a/Engine/Source/Programs/CrashReporter/CrashReportProcess/ReportProcessor.cs b/Engine/Source/Programs/CrashReporter/CrashReportProcess/ReportProcessor.cs
--- a/Engine/Source/Programs/CrashReporter/CrashReportProcess/ReportProcessor.cs
+++ b/Engine/Source/Programs/CrashReporter/CrashReportProcess/ReportProcessor.cs
@@ -3,11 +3,9 @@
 using System;
 using System.Collections.Generic;
 using System.Diagnostics;
-using System.Globalization;
 using System.IO;
 using System.Linq;
 using System.Net;
-using System.Text;
 using System.Threading;
 using System.Threading.Tasks;
 using Amazon.S3.Model;
@@ -28,7 +26,7 @@ namespace Tools.CrashReporter.CrashReportProcess
 		const string PerforceExePath = "C:\\Program Files\\Perforce\\p4.exe";
 
 		/// <summary>Give up on Perforce syncs after 2 minutes</summary>
-		const int SyncTimeoutSeconds = 2 * 60;
+		const int SyncTimeoutSeconds = 2*60;
 
 		/// <summary> Number of all processed reports. </summary>
 		public static int ProcessedReportCount = 1;
@@ -63,28 +61,12 @@ namespace Tools.CrashReporter.CrashReportProcess
 		/// <summary> Tasks to handle adding a new record to the database. </summary>
 		readonly List<Task> AddReportTasks = new List<Task>();
 
-		/// <summary>
-		/// Static init for one-time init jobs
-		/// </summary>
-		static ReportProcessor()
-		{
-#if !DEBUG
-			if (Config.Default.bSyncMinidumpDiagnostics)
-			{
-				if (!SyncMinidumpDiagnostics())
-				{
-					CrashReporterProcessServicer.WriteFailure("ReportProcessor: failed to sync files from Perforce");
-				}
-			}
-#endif			
-		}
-
 		/// <summary>
 		/// Global initialization of the processor.
 		/// </summary>
 		/// <param name="InWatcher">The object that watches for new crash reports coming in.</param>
 		/// <param name="InProcessorIndex">The index of this processor in the application's thread list</param>
-		public ReportProcessor( ReportWatcher InWatcher, int InProcessorIndex )
+		public ReportProcessor(ReportWatcher InWatcher, int InProcessorIndex)
 		{
 			// Create the thread to handle processing
 			Watcher = InWatcher;
@@ -116,40 +98,40 @@ namespace Tools.CrashReporter.CrashReportProcess
 		{
 			var Cancel = CancelSource.Token;
 			ProcessorTask = new Task(() =>
-			{
-				while (!Cancel.IsCancellationRequested)
-				{
-					try
-					{
-						bool bIdle = true;
-
-						foreach( var Queue in Watcher.ReportQueues )
-						{
-							FGenericCrashContext NewContext = null;
-							if (Queue.TryDequeueReport( out NewContext ))
-							{
-								ProcessReport( NewContext );
-
-								// The effect of this break is to prioritize ReportQueues by their order in the list, from highest to lowest
-								bIdle = false;
-								break;
-							}
-						}
-
-						if (bIdle)
-						{
-							// Don't use the CPU if we don't need.
-							Thread.Sleep(1000);							
-						}
-					}
-					catch (Exception Ex)
-					{
-						CrashReporterProcessServicer.WriteException(string.Format("PROC-{0} ", ProcessorIndex) + "ProcessNewReports: " + Ex, Ex);
-					}
-
-					TickStatic(Watcher);
-				}
-			});
+			                         {
+				                         while (!Cancel.IsCancellationRequested)
+				                         {
+					                         try
+					                         {
+						                         bool bIdle = true;
+
+						                         foreach (var Queue in Watcher.ReportQueues)
+						                         {
+							                         FGenericCrashContext NewContext = null;
+							                         if (Queue.TryDequeueReport(out NewContext))
+							                         {
+								                         ProcessReport(NewContext);
+
+								                         // The effect of this break is to prioritize ReportQueues by their order in the list, from highest to lowest
+								                         bIdle = false;
+								                         break;
+							                         }
+						                         }
+
+						                         if (bIdle)
+						                         {
+							                         // Don't use the CPU if we don't need.
+							                         Thread.Sleep(1000);
+						                         }
+					                         }
+					                         catch (Exception Ex)
+					                         {
+						                         CrashReporterProcessServicer.WriteException(string.Format("PROC-{0} ", ProcessorIndex) + "ProcessNewReports: " + Ex, Ex);
+					                         }
+
+					                         TickStatic(Watcher);
+				                         }
+			                         });
 		}
 
 		/// <summary>
@@ -166,26 +148,23 @@ namespace Tools.CrashReporter.CrashReportProcess
 		/// <summary>
 		/// Requests stopping the processor thread but doesn't block.
 		/// </summary>
-		public void RequestStop()
-		{
-			CancelSource.Cancel();
-		}
+		public void RequestStop() { CancelSource.Cancel(); }
 
 		/// <summary> 
 		/// Finalizes report files. 
 		/// Thread-safe.
 		/// </summary>
-		private void FinalizeReport( bool bAdded, DirectoryInfo DirInfo, FGenericCrashContext NewContext )
+		private void FinalizeReport(bool bAdded, DirectoryInfo DirInfo, FGenericCrashContext NewContext)
 		{
 			// Only remove if we added the report, otherwise leave all files to further investigation.
-			if( bAdded )
+			if (bAdded)
 			{
 				// Remove the report files as we're done with them.
-				CleanReport( DirInfo );
+				CleanReport(DirInfo);
 			}
 			else
 			{
-				MoveReportToInvalid( NewContext.CrashDirectory, ReportQueueBase.GetSafeFilename(NewContext.GetAsFilename()) );
+				MoveReportToInvalid(NewContext.CrashDirectory, ReportQueueBase.GetSafeFilename(NewContext.GetAsFilename()));
 			}
 		}
 
@@ -199,7 +178,7 @@ namespace Tools.CrashReporter.CrashReportProcess
 			const int MaxRetries = 3;
 
 			bool bWriteException = true;
-			for( int Retry = 0; Retry < MaxRetries; ++Retry )
+			for (int Retry = 0; Retry < MaxRetries; ++Retry)
 			{
 				try
 				{
@@ -208,34 +187,34 @@ namespace Tools.CrashReporter.CrashReportProcess
 						break;
 					}
 
-					foreach ( FileInfo Info in DirInfo.GetFiles() )
+					foreach (FileInfo Info in DirInfo.GetFiles())
 					{
 						Info.IsReadOnly = false;
 						Info.Attributes = FileAttributes.Normal;
 					}
-					DirInfo.Delete( true /* delete contents */);
+					DirInfo.Delete(true /* delete contents */);
 
 					// Random failures to delete with no exception seen regularly - retry
 					DirInfo.Refresh();
 					if (DirInfo.Exists)
-                    {
+					{
 						CrashReporterProcessServicer.WriteEvent("CleanReport: Failed to delete folder without an Exception " + DirInfo);
 						Thread.Sleep(500);
-                        continue;
-                    }
+						continue;
+					}
 
 					break;
 				}
-				catch( Exception Ex )
+				catch (Exception Ex)
 				{
-					if( bWriteException )
+					if (bWriteException)
 					{
-						CrashReporterProcessServicer.WriteException( "CleanReport: " + Ex, Ex);
+						CrashReporterProcessServicer.WriteException("CleanReport: " + Ex, Ex);
 						bWriteException = false;
 					}
 
 				}
-				System.Threading.Thread.Sleep( 100 );
+				System.Threading.Thread.Sleep(100);
 			}
 
 			DirInfo.Refresh();
@@ -250,38 +229,46 @@ namespace Tools.CrashReporter.CrashReportProcess
 		/// Thread-safe.
 		/// </summary>
 		/// <param name="DirPath">The directory to delete</param>
-		public static void CleanReport(string DirPath)
-		{
-			CleanReport( new DirectoryInfo( DirPath ) );
-		}
+		public static void CleanReport(string DirPath) { CleanReport(new DirectoryInfo(DirPath)); }
 
 		/// <summary> 
 		/// Moves specified report to the directory where are stored invalid reports. 
 		/// Thread-safe.
 		/// </summary>
-		public static void MoveReportToInvalid( string ReportName, string ReportNameAsFilename )
+		public static void MoveReportToInvalid(string ReportName, string ReportNameAsFilename)
 		{
 			try
 			{
 				CrashReporterProcessServicer.StatusReporter.IncrementCount(StatusReportingEventNames.ProcessingFailedEvent);
 
-				DirectoryInfo DirInfo = new DirectoryInfo( ReportName );
-				// Rename the report directory, so we should be able to quickly find the invalid reports.
-				CrashReporterProcessServicer.StatusReporter.AlertOnLowDisk(Config.Default.InvalidReportsDirectory, Config.Default.DiskSpaceAlertPercent);
-				Directory.CreateDirectory( Config.Default.InvalidReportsDirectory );
-				string DestinationDirectory = Path.Combine( Config.Default.InvalidReportsDirectory, ReportNameAsFilename);
-
-				Directory.CreateDirectory( DestinationDirectory );
+				DirectoryInfo DirInfo = new DirectoryInfo(ReportName);
 
-				// Copy all files from the source directory. We can't use MoveTo due to different disc location.
-				foreach( FileInfo File in DirInfo.GetFiles() )
+				if (Config.Default.InvalidReportsToAWS)
 				{
-					string DestinationFilepath = Path.Combine( DestinationDirectory, File.Name );
-					File.CopyTo( DestinationFilepath, true );
+					string LeafReportName = Path.GetFileName(ReportName);
+					DateTime S3KeyTime = DateTime.UtcNow;
+
+					// Copy all files from the source directory. We can't use MoveTo due to different disc location.
+					int FilesMoved = 0;
+					foreach (FileInfo InvalidFile in DirInfo.GetFiles())
+					{
+						try
+						{
+							string S3IDPrefix = string.Format("/{0}_{1}_{2}/{3}/{4}/", S3KeyTime.Year, S3KeyTime.Month, S3KeyTime.Day, S3KeyTime.Hour, LeafReportName);
+							UploadFileToS3(InvalidFile, Config.Default.AWSS3InvalidKeyPrefix + S3IDPrefix + InvalidFile.Name);
+							FilesMoved++;
+						}
+						catch (Exception Ex)
+						{
+							CrashReporterProcessServicer.WriteEvent("MoveReportToInvalid: Failed to write report file " + LeafReportName + ": " + InvalidFile.Name);
+							CrashReporterProcessServicer.WriteException("MoveReportToInvalid: " + Ex, Ex);
+						}
+					}
+
+					CrashReporterProcessServicer.WriteEvent(string.Format("MoveReportToInvalid moved {0} file(s) from {1} to S3", FilesMoved, LeafReportName));
 				}
-				DirInfo.Delete( true );
 
-				CrashReporterProcessServicer.WriteEvent( string.Format( "Moved to {0}", DestinationDirectory ) );
+				DirInfo.Delete( true );
 				UpdateProcessedReports();
 			}
 			catch( Exception Ex )
@@ -558,7 +545,7 @@ namespace Tools.CrashReporter.CrashReportProcess
 						if (bToS3)
 						{
 							WriteToS3Timer.Start();
-							UploadFileToS3(LogInfo, S3IDPrefix + "Launch.log");
+							UploadFileToS3(LogInfo, Config.Default.AWSS3OutputKeyPrefix + S3IDPrefix + "Launch.log");
 							WriteToS3Timer.Stop();
 						}
 						if (bToDisk)
@@ -578,7 +565,7 @@ namespace Tools.CrashReporter.CrashReportProcess
 					if (bToS3)
 					{
 						WriteToS3Timer.Start();
-						UploadFileToS3(CrashContextInfo, S3IDPrefix + FGenericCrashContext.CrashContextRuntimeXMLName);
+						UploadFileToS3(CrashContextInfo, Config.Default.AWSS3OutputKeyPrefix + S3IDPrefix + FGenericCrashContext.CrashContextRuntimeXMLName);
 						WriteToS3Timer.Stop();
 					}
 					if (bToDisk)
@@ -598,7 +585,7 @@ namespace Tools.CrashReporter.CrashReportProcess
 						if (bToS3)
 						{
 							WriteToS3Timer.Start();
-							UploadFileToS3(DumpInfo, S3IDPrefix + "MiniDump.dmp");
+							UploadFileToS3(DumpInfo, Config.Default.AWSS3OutputKeyPrefix + S3IDPrefix + "MiniDump.dmp");
 							WriteToS3Timer.Stop();
 						}
 						if (bToDisk)
@@ -622,7 +609,7 @@ namespace Tools.CrashReporter.CrashReportProcess
 						if (bToS3)
 						{
 							WriteToS3Timer.Start();
-							UploadFileToS3(VideoInfo, S3IDPrefix + CrashReporterConstants.VideoFileName);
+							UploadFileToS3(VideoInfo, Config.Default.AWSS3OutputKeyPrefix + S3IDPrefix + CrashReporterConstants.VideoFileName);
 							WriteToS3Timer.Stop();
 						}
 						if (bToDisk)
@@ -669,12 +656,12 @@ namespace Tools.CrashReporter.CrashReportProcess
 			return false;
 		}
 
-		private void UploadFileToS3(FileInfo FileInfo, string DestFilename)
+		private static void UploadFileToS3(FileInfo FileInfo, string DestFilename)
 		{
 			try
 			{
 				PutObjectResponse Response = CrashReporterProcessServicer.OutputAWS.PutS3ObjectFromFile(Config.Default.AWSS3OutputBucket,
-																										Config.Default.AWSS3OutputKeyPrefix + DestFilename,
+																										DestFilename,
 																										FileInfo.FullName);
 
 				if (Response == null || Response.HttpStatusCode != HttpStatusCode.OK)
@@ -688,50 +675,6 @@ namespace Tools.CrashReporter.CrashReportProcess
 			}
 		}
 
-		/// <summary>
-		/// Sync the MinidumpDiagnostics binary and the engine config files to #head.
-		/// </summary>
-		/// <returns>true if all the files synced without issue, false otherwise.</returns>
-		/// <remarks>As MinidumpDiagnostics is synced to #head, it requires the engine config files that match to run properly.</remarks>
-		private static bool SyncMinidumpDiagnostics()
-		{
-			// Use the latest MinidumpDiagnostics from the main branch.
-			CrashReporterProcessServicer.StatusReporter.AlertOnLowDisk(Config.Default.DepotRoot, Config.Default.DiskSpaceAlertPercent);
-			string UserString = "-u " + Config.Default.P4User;
-            string ClientString = "-c " + Config.Default.P4Client;
-            string SyncBinariesString = Path.Combine(Config.Default.DepotRoot, Config.Default.SyncBinariesFromDepot);
-            using (var MDDSyncProc = new LaunchProcess(PerforceExePath, null, CrashReporterProcessServicer.WriteP4, UserString, ClientString, "sync", SyncBinariesString))
-			{
-				if( MDDSyncProc.WaitForExit( SyncTimeoutSeconds * 1000 ) == EWaitResult.TimedOut )
-				{
-					CrashReporterProcessServicer.WriteFailure("Failed to sync MinidumpDiagnostics " + SyncBinariesString);
-					return false;
-				}
-			}
-
-			string SyncConfigString = Path.Combine(Config.Default.DepotRoot, Config.Default.SyncConfigFromDepot);
-			using (var ConfigSyncProc = new LaunchProcess(PerforceExePath, null, CrashReporterProcessServicer.WriteP4, UserString, ClientString, "sync", SyncConfigString))
-			{
-				if (ConfigSyncProc.WaitForExit(SyncTimeoutSeconds * 1000) == EWaitResult.TimedOut)
-				{
-					CrashReporterProcessServicer.WriteFailure("Failed to sync config files " + SyncConfigString);
-					return false;
-				}
-			}
-
-			string SyncThirdPartyString = Path.Combine(Config.Default.DepotRoot, Config.Default.SyncThirdPartyFromDepot); // Required by Perforce
-			using (var MiscSyncProc = new LaunchProcess(PerforceExePath, null, CrashReporterProcessServicer.WriteP4, UserString, ClientString, "sync", SyncThirdPartyString))
-			{
-				if (MiscSyncProc.WaitForExit(SyncTimeoutSeconds * 1000) == EWaitResult.TimedOut)
-				{
-					CrashReporterProcessServicer.WriteFailure("Failed to sync OpenSSL files " + SyncThirdPartyString);
-					return false;
-				}
-			}
-
-			return true;
-		}
-
 		/// <summary>
 		/// A function to process a newly landed crash report.
 		/// </summary>
@@ -901,19 +844,48 @@ namespace Tools.CrashReporter.CrashReportProcess
 			return AddReportTasks.Sum(AddReportTask => AddReportTask == null || AddReportTask.IsCompleted ? 0 : 1);
 		}
 
-		void ProcessDumpFile( string DiagnosticsPath, FGenericCrashContext NewContext )
+		void ProcessDumpFile( string DumpPath, FGenericCrashContext NewContext )
 		{
-			CrashReporterProcessServicer.WriteEvent(string.Format("PROC-{0} ", ProcessorIndex) + "ProcessDumpFile: Waiting to run MDD on " + DiagnosticsPath);
+			CrashReporterProcessServicer.WriteEvent(string.Format("PROC-{0} ", ProcessorIndex) + "ProcessDumpFile: Waiting to run MDD on " + DumpPath);
 
-			if (CrashReporterProcessServicer.Symbolicator.Run(DiagnosticsPath, NewContext, ProcessorIndex))
+			if (TrySimulateSymbolication(NewContext))
+			{
+				CrashReporterProcessServicer.WriteEvent(string.Format("PROC-{0} ", ProcessorIndex) + "ProcessDumpFile: Simulated symbolication (skipped MDD) " + NewContext.CrashDirectory);
+			}
+			else if(CrashReporterProcessServicer.Symbolicator.Run(DumpPath, NewContext, ProcessorIndex))
 			{
-				CrashReporterProcessServicer.WriteEvent(string.Format("PROC-{0} ", ProcessorIndex) + "ProcessDumpFile: MDD finished running on " + DiagnosticsPath);
+				CrashReporterProcessServicer.WriteEvent(string.Format("PROC-{0} ", ProcessorIndex) + "ProcessDumpFile: MDD finished running on " + DumpPath);
 				ReadDiagnosticsFile(NewContext);
 			}
 			else
 			{
-				CrashReporterProcessServicer.WriteFailure(string.Format("PROC-{0} ", ProcessorIndex) + "ProcessDumpFile: MDD didn't run on " + DiagnosticsPath);
+				CrashReporterProcessServicer.WriteFailure(string.Format("PROC-{0} ", ProcessorIndex) + "ProcessDumpFile: MDD didn't run on " + DumpPath);
 			}
+
+			if (string.IsNullOrWhiteSpace(NewContext.PrimaryCrashProperties.CallStack))
+			{
+				CrashReporterProcessServicer.StatusReporter.IncrementCount(StatusReportingEventNames.SymbolicationFailedEvent);
+			}
+			else
+			{
+				CrashReporterProcessServicer.StatusReporter.IncrementCount(StatusReportingEventNames.SymbolicationSucceededEvent);
+			}
+		}
+
+		bool TrySimulateSymbolication(FGenericCrashContext NewContext)
+		{
+			// Check for crashes we don't want to attempt to symbolicate
+
+			if (!string.IsNullOrWhiteSpace(NewContext.PrimaryCrashProperties.ErrorMessage) &&
+			    NewContext.PrimaryCrashProperties.ErrorMessage.Contains("Hang detected on"))
+			{
+				// Hangs from FThreadHeartBeat::Run() start wth "Hang detected on" and have a useful callstack already in the log
+				// Callstack and source from MDD will be useless so simulate a callstack that points the user to the log.
+				NewContext.PrimaryCrashProperties.CallStack = "See log for callstack from hang detection code()\nSee log for callstack from hang detection code()\nSee log for callstack from hang detection code()";
+				return true;
+			}
+
+			return false;
 		}
 
 		private static void TickStatic(ReportWatcher InWatcher)
diff --git a/Engine/Source/Programs/CrashReporter/CrashReportProcess/ReportQueue.cs b/Engine/Source/Programs/CrashReporter/CrashReportProcess/ReportQueue.cs
--- a/Engine/Source/Programs/CrashReporter/CrashReportProcess/ReportQueue.cs
+++ b/Engine/Source/Programs/CrashReporter/CrashReportProcess/ReportQueue.cs
@@ -6,7 +6,6 @@ using System.Globalization;
 using System.IO;
 using System.Linq;
 using System.Runtime.InteropServices;
-using System.Text;
 using Tools.CrashReporter.CrashReportCommon;
 using Tools.DotNETCommon.XmlHandler;
 
diff --git a/Engine/Source/Programs/CrashReporter/CrashReportProcess/StatusReportLoop.cs b/Engine/Source/Programs/CrashReporter/CrashReportProcess/StatusReportLoop.cs
--- a/Engine/Source/Programs/CrashReporter/CrashReportProcess/StatusReportLoop.cs
+++ b/Engine/Source/Programs/CrashReporter/CrashReportProcess/StatusReportLoop.cs
@@ -2,13 +2,14 @@
 using System.Collections.Generic;
 using System.Data.Linq.Mapping;
 using System.Linq;
+using System.Security.Policy;
 using System.Text;
 using System.Threading;
 using System.Threading.Tasks;
 
 namespace Tools.CrashReporter.CrashReportProcess
 {
-	class StatusReportLoop : IDisposable
+	abstract class StatusReportLoop : IDisposable
 	{
 		public TimeSpan InitialPause { get; set; }
 
@@ -22,14 +23,15 @@ namespace Tools.CrashReporter.CrashReportProcess
 				Monitor.PulseAll(TaskMonitor);
 			}
 			LoopTask.Wait();
+			LoopTask.Dispose();
 			CancelSource.Dispose();
 		}
 
-		public StatusReportLoop(TimeSpan InLoopPeriod, Func<TimeSpan, string> LoopFunction)
+		public StatusReportLoop(TimeSpan InLoopPeriod, Func<StatusReportLoop, TimeSpan, string> LoopFunction)
 		{
 			InitialPause = TimeSpan.FromSeconds(30);
 			LoopPeriod = InLoopPeriod;
-			LoopTask = Task.Factory.StartNew(() =>
+			LoopTask = new Task(() =>
 			{
 				var Cancel = CancelSource.Token;
 
@@ -42,7 +44,7 @@ namespace Tools.CrashReporter.CrashReportProcess
 					while (!Cancel.IsCancellationRequested)
 					{
 						DateTime PeriodEndTime = DateTime.UtcNow;
-						string ReportMessage = LoopFunction.Invoke(PeriodEndTime - PeriodStartTime);
+						string ReportMessage = LoopFunction.Invoke(this, PeriodEndTime - PeriodStartTime);
 						PeriodStartTime = PeriodEndTime;
 
 						if (!string.IsNullOrWhiteSpace(ReportMessage))
@@ -56,9 +58,116 @@ namespace Tools.CrashReporter.CrashReportProcess
 			});
 		}
 
+		protected void Start() { LoopTask.Start(); }
+
 		private Task LoopTask;
 		private readonly CancellationTokenSource CancelSource = new CancellationTokenSource();
 		private readonly DateTime CreationTimeUtc = DateTime.UtcNow;
-		private readonly Object TaskMonitor = new Object();
+		private readonly object TaskMonitor = new object();
+	}
+
+	class RegularStatusReport : StatusReportLoop
+	{
+		public RegularStatusReport(TimeSpan InLoopPeriod, Func<StatusReportLoop, TimeSpan, string> LoopFunction)
+			: base(InLoopPeriod, LoopFunction)
+		{
+			Start();
+		}
+
+		public Dictionary<string, int> GetCountsInPeriod(IDictionary<string, int> InCounters)
+		{
+			Dictionary<string, int> CountsInPeriod = new Dictionary<string, int>();
+
+			foreach (var Counter in InCounters)
+			{
+				int LastCount = 0;
+				if (!CountersAtLastReport.TryGetValue(Counter.Key, out LastCount))
+				{
+					CountersAtLastReport.Add(Counter.Key, 0);
+				}
+				CountsInPeriod.Add(Counter.Key, Counter.Value - LastCount);
+				CountersAtLastReport[Counter.Key] = Counter.Value;
+			}
+
+			return CountsInPeriod;
+		}
+
+		public int QueueSizeSumAtLastReport { get; set; }
+
+		private readonly Dictionary<string, int> CountersAtLastReport = new Dictionary<string, int>();
+	}
+
+	class DiskStatusReport : StatusReportLoop
+	{
+		public DiskStatusReport(TimeSpan InLoopPeriod, Func<StatusReportLoop, TimeSpan, string> LoopFunction)
+			: base(InLoopPeriod, LoopFunction)
+		{
+			Start();
+		}
+	}
+
+	class PerfStatusReport : StatusReportLoop
+	{
+		public PerfStatusReport(TimeSpan InLoopPeriod, Func<StatusReportLoop, TimeSpan, string> LoopFunction)
+			: base(InLoopPeriod, LoopFunction)
+		{
+			Start();
+		}
+
+		public Dictionary<string, StatusReportingMeanCounter> GetMeanCountsInPeriod(IDictionary<string, StatusReportingMeanCounter> InCounters)
+		{
+			Dictionary<string, StatusReportingMeanCounter> CountsInPeriod = new Dictionary<string, StatusReportingMeanCounter>();
+
+			foreach (var Counter in InCounters)
+			{
+				StatusReportingMeanCounter LastCount;
+				if (!MeanCountersAtLastReport.TryGetValue(Counter.Key, out LastCount))
+				{
+					LastCount = new StatusReportingMeanCounter();
+					MeanCountersAtLastReport.Add(Counter.Key, LastCount);
+				}
+				CountsInPeriod.Add(Counter.Key, new StatusReportingMeanCounter(LastCount, Counter.Value));
+
+				LastCount.SampleCount = Counter.Value.SampleCount;
+				LastCount.TotalMillisec = Counter.Value.TotalMillisec;
+			}
+
+			return CountsInPeriod;
+		}
+
+		public static string GetPerfTimeString(float Millisec)
+		{
+			if (Millisec >= 600000.0f) // <== over ten minutes
+			{
+				// Express as simple number of integer minutes
+				int Minutes = (int)Millisec/60000;
+				return Minutes + " min";
+			}
+
+			if (Millisec >= 20000.0f) // <== over twenty seconds
+			{
+				// Express as a number of integer seconds
+				int Seconds = (int)Millisec / 1000;
+				return Seconds + " s";
+			}
+
+			if (Millisec >= 1000.0f) // <== over one second
+			{
+				// Express as a number of seconds with one decimal place
+				float Seconds = Millisec / 1000.0f;
+				return string.Format("{0:F1} s", Seconds);
+			}
+
+			if (Millisec >= 20.0f) // <== over twenty milliseconds
+			{
+				// Express as a number of integer milliseconds
+				return string.Format("{0:F0} ms", Millisec);
+			}
+
+			// Express as a number of milliseconds with one decimal place
+			return string.Format("{0:F1} ms", Millisec);
+		}
+
+		private readonly Dictionary<string, StatusReportingMeanCounter> MeanCountersAtLastReport = new Dictionary<string, StatusReportingMeanCounter>();
 	}
 }
diff --git a/Engine/Source/Programs/CrashReporter/CrashReportProcess/StatusReporting.cs b/Engine/Source/Programs/CrashReporter/CrashReportProcess/StatusReporting.cs
--- a/Engine/Source/Programs/CrashReporter/CrashReportProcess/StatusReporting.cs
+++ b/Engine/Source/Programs/CrashReporter/CrashReportProcess/StatusReporting.cs
@@ -1,15 +1,25 @@
 ﻿using System;
 using System.Collections.Generic;
-using System.IO;
 using System.Linq;
-using System.Reflection;
 using System.Text;
-using System.Threading;
-using System.Threading.Tasks;
 using Tools.CrashReporter.CrashReportCommon;
 
 namespace Tools.CrashReporter.CrashReportProcess
 {
+	class StatusReportingMeanCounter
+	{
+		public int TotalMillisec { get; set; }
+		public int SampleCount { get; set; }
+
+		public StatusReportingMeanCounter() { }
+
+		public StatusReportingMeanCounter(StatusReportingMeanCounter LastPeriod, StatusReportingMeanCounter Total)
+		{
+			TotalMillisec = Total.TotalMillisec - LastPeriod.TotalMillisec;
+			SampleCount = Total.SampleCount - LastPeriod.SampleCount;
+		}
+	}
+
 	class StatusReporting : IDisposable
 	{
 		public void InitCounters(IEnumerable<string> CounterNames)
@@ -23,6 +33,17 @@ namespace Tools.CrashReporter.CrashReportProcess
 			}
 		}
 
+		public void InitMeanCounters(IEnumerable<string> CounterNames)
+		{
+			foreach (string CounterName in CounterNames)
+			{
+				if (!MeanCounters.ContainsKey(CounterName))
+				{
+					MeanCounters.Add(CounterName, new StatusReportingMeanCounter());
+				}
+			}
+		}
+
 		public void InitFolderMonitors(IDictionary<string, string> InFolderMonitors)
 		{
 			// Collapse list of folders and labels into a set of drive letters and one or more labels
@@ -77,6 +98,20 @@ namespace Tools.CrashReporter.CrashReportProcess
 			}
 		}
 
+		public void AddToMeanCounter(string CountName, int Millisec)
+		{
+			lock (DataLock)
+			{
+				if (!MeanCounters.ContainsKey(CountName))
+				{
+					MeanCounters.Add(CountName, new StatusReportingMeanCounter());
+				}
+
+				MeanCounters[CountName].TotalMillisec += Millisec;
+				MeanCounters[CountName].SampleCount++;
+			}
+		}
+
 		public void Alert(string AlertKey, string AlertText, int RepeatMinimumMinutes)
 		{
 			// Do "no repeat" check here of recent alerts and block duplicates based on timer (use AlertKey)
@@ -130,108 +165,142 @@ namespace Tools.CrashReporter.CrashReportProcess
 			}
 			CrashReporterProcessServicer.WriteSlack(StartupMessage.ToString());
 
-			ReporterTasks = new[]
+			List<StatusReportLoop> StatusReportLoops = new List<StatusReportLoop>();
+
+			StatusReportLoops.Add(
+				new RegularStatusReport(TimeSpan.FromMinutes(Config.Default.MinutesBetweenQueueSizeReports), (InLoop, InPeriod) =>
+				                        {
+					                        RegularStatusReport ThisLoop = (RegularStatusReport)InLoop;
+					                        StringBuilder StatusReportMessage = new StringBuilder();
+					                        lock (DataLock)
+					                        {
+						                        Dictionary<string, int> CountsInPeriod = ThisLoop.GetCountsInPeriod(Counters);
+
+						                        if (SetQueueSizesCallCount >= QueueSizes.Count)
+						                        {
+							                        int ProcessingStartedInPeriodReceiver = 0;
+							                        CountsInPeriod.TryGetValue(StatusReportingEventNames.ProcessingStartedReceiverEvent, out ProcessingStartedInPeriodReceiver);
+							                        int ProcessingStartedInPeriodDataRouter = 0;
+							                        CountsInPeriod.TryGetValue(StatusReportingEventNames.ProcessingStartedDataRouterEvent, out ProcessingStartedInPeriodDataRouter);
+							                        int ProcessingStartedInPeriod = ProcessingStartedInPeriodReceiver + ProcessingStartedInPeriodDataRouter;
+
+							                        if (ProcessingStartedInPeriod > 0)
+							                        {
+								                        int QueueSizeSum = QueueSizes.Values.Sum();
+								                        TimeSpan MeanWaitTime =
+									                        TimeSpan.FromTicks((long)(0.5*InPeriod.Ticks*(QueueSizeSum + ThisLoop.QueueSizeSumAtLastReport)/ProcessingStartedInPeriod));
+								                        ThisLoop.QueueSizeSumAtLastReport = QueueSizeSum;
+
+								                        int WaitMinutes = Convert.ToInt32(MeanWaitTime.TotalMinutes);
+								                        string WaitTimeString;
+								                        if (MeanWaitTime == TimeSpan.Zero)
+								                        {
+									                        WaitTimeString = "nil";
+								                        }
+								                        else if (MeanWaitTime < TimeSpan.FromMinutes(1))
+								                        {
+									                        WaitTimeString = "< 1 minute";
+								                        }
+								                        else if (WaitMinutes == 1)
+								                        {
+									                        WaitTimeString = "1 minute";
+								                        }
+								                        else
+								                        {
+									                        WaitTimeString = string.Format("{0} minutes", WaitMinutes);
+								                        }
+								                        StatusReportMessage.AppendLine("Queue waiting time " + WaitTimeString);
+							                        }
+
+							                        StatusReportMessage.AppendLine("Queue sizes...");
+							                        StatusReportMessage.Append("> ");
+							                        foreach (var Queue in QueueSizes)
+							                        {
+								                        StatusReportMessage.Append(Queue.Key + " " + Queue.Value + "                ");
+							                        }
+							                        StatusReportMessage.AppendLine();
+						                        }
+						                        if (CountsInPeriod.Values.Sum() > 0)
+						                        {
+													string PeriodTimeString = GetTimePeriodString(InPeriod);
+
+													StatusReportMessage.AppendLine(string.Format("Events in the last {0}...", PeriodTimeString));
+							                        foreach (var CountInPeriod in CountsInPeriod)
+							                        {
+								                        if (CountInPeriod.Value > 0)
+								                        {
+									                        StatusReportMessage.AppendLine("> " + CountInPeriod.Key + " " + CountInPeriod.Value);
+								                        }
+							                        }
+						                        }
+					                        }
+					                        return StatusReportMessage.ToString();
+				                        }));
+
+			StatusReportLoops.Add(
+				new DiskStatusReport(TimeSpan.FromDays(1.0), (InLoop, InPeriod) =>
+				                     {
+					                     StringBuilder DailyReportMessage = new StringBuilder();
+					                     lock (DataLock)
+					                     {
+						                     DailyReportMessage.AppendLine("Disk space available...");
+
+						                     foreach (var FolderMonitor in FolderMonitors)
+						                     {
+							                     string FreeSpaceText = "#Error";
+							                     string Drive = FolderMonitor.Key;
+							                     Int64 FreeSpace;
+							                     float FreePercent;
+							                     if (CrashReportCommon.StorageSpaceHelper.TryGetSpaceAvailable(Drive, out FreeSpace, out FreePercent))
+							                     {
+								                     FreeSpaceText = GetDiskSpaceString(FreeSpace, FreePercent);
+							                     }
+
+							                     DailyReportMessage.AppendLine("> " + FolderMonitor.Value + " =>> " + FreeSpaceText);
+						                     }
+					                     }
+					                     return DailyReportMessage.ToString();
+				                     }));
+
+			if (Config.Default.MonitorPerformance)
 			{
-				new StatusReportLoop(TimeSpan.FromMinutes(Config.Default.MinutesBetweenQueueSizeReports), InPeriod =>
-				{
-					StringBuilder StatusReportMessage = new StringBuilder();
-					lock (DataLock)
-					{
-						Dictionary<string, int> CountsInPeriod = GetCountsInPeriod(Counters, CountersAtLastReport);
-
-						if (SetQueueSizesCallCount >= QueueSizes.Count)
-						{
-							int ProcessingStartedInPeriodReceiver = 0;
-							CountsInPeriod.TryGetValue(StatusReportingEventNames.ProcessingStartedReceiverEvent, out ProcessingStartedInPeriodReceiver);
-							int ProcessingStartedInPeriodDataRouter = 0;
-							CountsInPeriod.TryGetValue(StatusReportingEventNames.ProcessingStartedDataRouterEvent, out ProcessingStartedInPeriodDataRouter);
-							int ProcessingStartedInPeriod = ProcessingStartedInPeriodReceiver + ProcessingStartedInPeriodDataRouter;
-
-							if (ProcessingStartedInPeriod > 0)
-							{
-								int QueueSizeSum = QueueSizes.Values.Sum();
-								TimeSpan MeanWaitTime =
-									TimeSpan.FromTicks((long)(0.5*InPeriod.Ticks*(QueueSizeSum + QueueSizeSumAtLastReport)/ProcessingStartedInPeriod));
-								QueueSizeSumAtLastReport = QueueSizeSum;
-
-								int WaitMinutes = Convert.ToInt32(MeanWaitTime.TotalMinutes);
-								string WaitTimeString;
-								if (MeanWaitTime == TimeSpan.Zero)
-								{
-									WaitTimeString = "nil";
-								}
-								else if (MeanWaitTime < TimeSpan.FromMinutes(1))
-								{
-									WaitTimeString = "< 1 minute";
-								}
-								else if (WaitMinutes == 1)
-								{
-									WaitTimeString = "1 minute";
-								}
-								else
-								{
-									WaitTimeString = string.Format("{0} minutes", WaitMinutes);
-								}
-								StatusReportMessage.AppendLine("Queue waiting time " + WaitTimeString);
-							}
-
-							StatusReportMessage.AppendLine("Queue sizes...");
-							StatusReportMessage.Append("> ");
-							foreach (var Queue in QueueSizes)
-							{
-								StatusReportMessage.Append(Queue.Key + " " + Queue.Value + "                ");
-							}
-							StatusReportMessage.AppendLine();
-						}
-						if (CountsInPeriod.Count > 0)
-						{
-							string PeriodTimeString;
-							if (InPeriod.TotalMinutes < 1.0)
-							{
-								PeriodTimeString = string.Format("{0:N0} seconds", InPeriod.TotalSeconds);
-							}
-							else
-							{
-								PeriodTimeString = string.Format("{0:N0} minutes", InPeriod.TotalMinutes);
-							}
-
-							StatusReportMessage.AppendLine(string.Format("Events in the last {0}...", PeriodTimeString));
-							foreach (var CountInPeriod in CountsInPeriod)
-							{
-								if (CountInPeriod.Value > 0)
-								{
-									StatusReportMessage.AppendLine("> " + CountInPeriod.Key + " " + CountInPeriod.Value);
-								}
-							}
-						}
-					}
-					return StatusReportMessage.ToString();
-				}),
-
-				new StatusReportLoop(TimeSpan.FromDays(1.0), InPeriod =>
-				{
-					StringBuilder DailyReportMessage = new StringBuilder();
-					lock (DataLock)
-					{
-						DailyReportMessage.AppendLine("Disk space available...");
+				StatusReportLoops.Add(
+					new PerfStatusReport(TimeSpan.FromMinutes(15), (InLoop, InPeriod) =>
+					                     {
+											 PerfStatusReport ThisLoop = (PerfStatusReport)InLoop;
+											 StringBuilder PerfReportMessage = new StringBuilder();
+						                     lock (DataLock)
+						                     {
+							                     Dictionary<string, StatusReportingMeanCounter> MeanCountsInPeriod = ThisLoop.GetMeanCountsInPeriod(MeanCounters);
+
+												 if (MeanCountsInPeriod.Count > 0)
+												 {
+													 string PeriodTimeString = GetTimePeriodString(InPeriod);
+													 PerfReportMessage.AppendLine(string.Format("Performance in the last {0}...", PeriodTimeString));
+
+													 foreach (var MeanCountInPeriod in MeanCountsInPeriod)
+													 {
+														 if (MeanCountInPeriod.Value.SampleCount == 0)
+														 {
+															 PerfReportMessage.AppendLine("> " + MeanCountInPeriod.Key + ": nil");
+														 }
+														 else
+														 {
+															 PerfReportMessage.AppendLine("> " + MeanCountInPeriod.Key + ": Total " +
+															                              PerfStatusReport.GetPerfTimeString((float)MeanCountInPeriod.Value.TotalMillisec) +
+															                              " / Count " + MeanCountInPeriod.Value.SampleCount +
+															                              " = Mean " +
+															                              PerfStatusReport.GetPerfTimeString(MeanCountInPeriod.Value.TotalMillisec/(float)MeanCountInPeriod.Value.SampleCount));
+														 }
+													 }
+												 }
+											 }
+
+											 return PerfReportMessage.ToString();
+										 }));
+			}
 
-						foreach (var FolderMonitor in FolderMonitors)
-						{
-							string FreeSpaceText = "#Error";
-							string Drive = FolderMonitor.Key;
-							Int64 FreeSpace;
-							float FreePercent;
-							if (CrashReportCommon.StorageSpaceHelper.TryGetSpaceAvailable(Drive, out FreeSpace, out FreePercent))
-							{
-								FreeSpaceText = GetDiskSpaceString(FreeSpace, FreePercent);
-							}
-
-							DailyReportMessage.AppendLine("> " + FolderMonitor.Value + " =>> " + FreeSpaceText);
-						}
-					}
-					return DailyReportMessage.ToString();
-				}), 
-			};
+			ReporterTasks = StatusReportLoops;
 		}
 
 		public void OnPreStopping()
@@ -252,22 +321,14 @@ namespace Tools.CrashReporter.CrashReportProcess
 			CrashReporterProcessServicer.WriteSlack("CRP stopped");
 		}
 
-		private static Dictionary<string, int> GetCountsInPeriod(Dictionary<string, int> InCounters, Dictionary<string, int> InCountersAtLastReport)
+		private static string GetTimePeriodString(TimeSpan InPeriod)
 		{
-			Dictionary<string, int> CountsInPeriod = new Dictionary<string, int>();
-
-			foreach (var Counter in InCounters)
+			if (InPeriod.TotalMinutes < 1.0)
 			{
-				int LastCount = 0;
-				if (!InCountersAtLastReport.TryGetValue(Counter.Key, out LastCount))
-				{
-					InCountersAtLastReport.Add(Counter.Key, 0);
-				}
-				CountsInPeriod.Add(Counter.Key, Counter.Value - LastCount);
-				InCountersAtLastReport[Counter.Key] = Counter.Value;
+				return string.Format("{0:N0} seconds", InPeriod.TotalSeconds);
 			}
 
-			return CountsInPeriod;
+			return string.Format("{0:N0} minutes", InPeriod.TotalMinutes);
 		}
 
 		private static string GetDiskSpaceString(long FreeSpace, float FreePercent)
@@ -323,9 +384,8 @@ namespace Tools.CrashReporter.CrashReportProcess
 		private readonly Dictionary<string, string> QueueLocations = new Dictionary<string, string>();
 		private readonly Dictionary<string, int> QueueSizes = new Dictionary<string, int>();
 		private readonly Dictionary<string, int> Counters = new Dictionary<string, int>();
+		private readonly Dictionary<string, StatusReportingMeanCounter> MeanCounters = new Dictionary<string, StatusReportingMeanCounter>();
 		private readonly Dictionary<string, string> FolderMonitors = new Dictionary<string, string>();
-		private int QueueSizeSumAtLastReport = 0;
-		private readonly Dictionary<string, int> CountersAtLastReport = new Dictionary<string, int>();
 		private readonly Object DataLock = new Object();
 		private int SetQueueSizesCallCount = 0;
 		private readonly Dictionary<string, DateTime> RecentAlerts = new Dictionary<string, DateTime>();
diff --git a/Engine/Source/Programs/CrashReporter/CrashReportProcess/StatusReportingConstants.cs b/Engine/Source/Programs/CrashReporter/CrashReportProcess/StatusReportingConstants.cs
--- a/Engine/Source/Programs/CrashReporter/CrashReportProcess/StatusReportingConstants.cs
+++ b/Engine/Source/Programs/CrashReporter/CrashReportProcess/StatusReportingConstants.cs
@@ -18,6 +18,13 @@ namespace Tools.CrashReporter.CrashReportProcess
 		public const string ProcessingStartedPS4Event = "Started processing (from PS4Services)";
 		public const string ProcessingSucceededEvent = "Processing succeeded";
 		public const string ProcessingFailedEvent = "Processing failed";
+		public const string SymbolicationSucceededEvent = "Symbolication succeeded";
+		public const string SymbolicationFailedEvent = "Symbolication skipped/failed";
 		public const string ExceptionEvent = "Handled exceptions";	
 	}
+
+	static class StatusReportingPerfMeanNames
+	{
+		public const string MinidumpDiagnostics = "MDD runs";
+	}
 }
diff --git a/Engine/Source/Programs/CrashReporter/CrashReportProcess/Symbolicator.cs b/Engine/Source/Programs/CrashReporter/CrashReportProcess/Symbolicator.cs
--- a/Engine/Source/Programs/CrashReporter/CrashReportProcess/Symbolicator.cs
+++ b/Engine/Source/Programs/CrashReporter/CrashReportProcess/Symbolicator.cs
@@ -3,8 +3,6 @@ using System.Collections.Generic;
 using System.Diagnostics;
 using System.IO;
 using System.Linq;
-using System.Text;
-using System.Threading;
 using System.Threading.Tasks;
 using Tools.CrashReporter.CrashReportCommon;
 using Tools.DotNETCommon.LaunchProcess;
@@ -22,15 +20,15 @@ namespace Tools.CrashReporter.CrashReportProcess
 		/// <returns>True, if successful</returns>
 		public bool Run(string DiagnosticsPath, FGenericCrashContext Context, int ProcessorIndex)
 		{
-			// Use the latest MinidumpDiagnostics from the main branch.
-			CrashReporterProcessServicer.StatusReporter.AlertOnLowDisk(Config.Default.DepotRoot, Config.Default.DiskSpaceAlertPercent);
-			string Win64BinariesDirectory = Path.Combine(Config.Default.DepotRoot, Config.Default.MDDBinariesFolderInDepot);
-#if DEBUG
-			// Note: the debug executable must be built locally or synced from Perforce manually
-			string MinidumpDiagnosticsName = Path.Combine(Win64BinariesDirectory, "MinidumpDiagnostics-Win64-Debug.exe");
-#else
-			string MinidumpDiagnosticsName = Path.Combine(Win64BinariesDirectory, "MinidumpDiagnostics.exe");
-#endif
+			if (!File.Exists(Config.Default.MDDExecutablePath))
+			{
+				CrashReporterProcessServicer.WriteEvent("Symbolicator.Run() file not found " + Config.Default.MDDExecutablePath);
+				return false;
+			}
+
+            // Use MinidumpDiagnostics from MDDExecutablePath.
+            CrashReporterProcessServicer.StatusReporter.AlertOnLowDisk(Config.Default.MDDExecutablePath, Config.Default.DiskSpaceAlertPercent);
+
 			// Don't purge logs
 			// TODO: make this clear to logs once a day or something (without letting MDD check on every run!)
 			string PurgeLogsDays = "-1";
@@ -55,8 +53,10 @@ namespace Tools.CrashReporter.CrashReportProcess
 
 			// Build the absolute log file path for MinidumpDiagnostics
 			string BaseFolder = CrashReporterProcessServicer.SymbolicatorLogFolder;
-			string SubFolder = DateTime.UtcNow.ToString("yyyy_MM_dd");
-			string Folder = Path.Combine(BaseFolder, SubFolder);
+		    DateTime WriteTime = DateTime.UtcNow;
+            string DateFolder = WriteTime.ToString("yyyy_MM_dd");
+            string HourFolder = WriteTime.ToString("HH");
+            string Folder = Path.Combine(BaseFolder, DateFolder, HourFolder);
 			string AbsLogPath = Path.Combine(Folder, Context.GetAsFilename() + ".log");
 			Directory.CreateDirectory(Folder);
 
@@ -110,18 +110,17 @@ namespace Tools.CrashReporter.CrashReportProcess
 				(
 					new[]
 					{
-						"-buildmachine",
-						"-forcelogflush"
+						"-buildmachine"
 					}
 				);
 
 				// Write some debugging message.
-				CrashReporterProcessServicer.WriteMDD(MinidumpDiagnosticsName + " Params: " + String.Join(" ", MinidumpDiagnosticsParams));
+				CrashReporterProcessServicer.WriteMDD("MinidumpDiagnostics Params: " + String.Join(" ", MinidumpDiagnosticsParams));
 			}
 
 			Task<bool> NewSymbolicatorTask = Task.FromResult(false);
 			Stopwatch WaitSW = Stopwatch.StartNew();
-			Double WaitForLockTime = 0.0;
+			double WaitForLockTime = 0.0;
 
 			lock (Tasks)
 			{
@@ -129,7 +128,7 @@ namespace Tools.CrashReporter.CrashReportProcess
 
 				Tasks[TaskIdx] = NewSymbolicatorTask = Task<bool>.Factory.StartNew(() =>
 				{
-					LaunchProcess ReportParser = new LaunchProcess(MinidumpDiagnosticsName, Path.GetDirectoryName(MinidumpDiagnosticsName), CaptureMessageDelegate,
+					LaunchProcess ReportParser = new LaunchProcess(Config.Default.MDDExecutablePath, Path.GetDirectoryName(Config.Default.MDDExecutablePath), CaptureMessageDelegate,
 																	MinidumpDiagnosticsParams.ToArray());
 
 					return ReportParser.WaitForExit(Config.Default.MDDTimeoutMinutes*1000*60) == EWaitResult.Ok;
@@ -140,8 +139,10 @@ namespace Tools.CrashReporter.CrashReportProcess
 
 			NewSymbolicatorTask.Wait();
 
-			Double TotalMDDTime = WaitSW.Elapsed.TotalSeconds;
-			CrashReporterProcessServicer.WriteEvent(string.Format("PROC-{0} ", ProcessorIndex) + string.Format("Symbolicator.Run: Thread blocked for {0:N1}s then MDD ran for {1:N1}s", WaitForLockTime, TotalMDDTime - WaitForLockTime));
+			double TotalMDDTime = WaitSW.Elapsed.TotalSeconds;
+			double MDDRunTime = TotalMDDTime - WaitForLockTime;
+			CrashReporterProcessServicer.StatusReporter.AddToMeanCounter(StatusReportingPerfMeanNames.MinidumpDiagnostics, (int)(MDDRunTime * 1000));
+			CrashReporterProcessServicer.WriteEvent(string.Format("PROC-{0} ", ProcessorIndex) + string.Format("Symbolicator.Run: Thread blocked for {0:N1}s then MDD ran for {1:N1}s", WaitForLockTime, MDDRunTime));
 
 			return NewSymbolicatorTask.Result;
 		}
