Copying //UE4/Dev-Core to //UE4/Dev-Main (Source: //UE4/Dev-Core @ 3228803)

==========================
MAJOR FEATURES + CHANGES
==========================

Change 3209807 on 2016/11/24 by [REDACTED_PERSON]

	CRP now has improved handling of lost comms with website/DB (CRP v.1.2.9)
	[UE-38397] - CrashReportProcess should not lose crashes when the website goes down.

Change 3209936 on 2016/11/24 by [REDACTED_PERSON]

	Reconciled offline work. Fixed a number of display bugs with the reports page.

Change 3209938 on 2016/11/24 by [REDACTED_PERSON]

	Reconciled offline work

	Adding JQuery UI packages

Change 3210736 on 2016/11/28 by [REDACTED_PERSON]

	Unset() made protected, which is an implementation details of TFunction and not the way to unbind one (may leak memory).
	Fixes to existing calls.

	https://answers.unrealengine.com/questions/494947/proper-way-to-deal-with-destroying-a-tfunction.html

Change 3211181 on 2016/11/28 by [REDACTED_PERSON]

	Improved error message when binding a delegate to a pending kill object.

Change 3211214 on 2016/11/28 by [REDACTED_PERSON]

	PR #2978: fixed typo in FMallocLeakDetection (Contributed by [REDACTED_PERSON])

Change 3211301 on 2016/11/28 by [REDACTED_PERSON]

	PR #2892: Copy bCustomVersionsAreReset when creating an Archive from another Ar. (Contributed by [REDACTED_PERSON])

Change 3213387 on 2016/11/29 by [REDACTED_PERSON]

	Test of a release note parsed from the changelist

Change 3213553 on 2016/11/29 by [REDACTED_PERSON]

	UE4 - Rework dependencies and UStruct creation so that we can create UBlueprintGeneratedClass's without needing the parent class serialized.

Change 3214800 on 2016/11/30 by [REDACTED_PERSON]

	ModuleManager will now use a critical section instead of FMultiReaderSingleWriterGT to make it really thread safe.

	- removed FMultiReaderSingleWriterGT because it was not fully thread safe

Change 3214926 on 2016/11/30 by [REDACTED_PERSON]

	Merging using Dev-Core_To_Dev-LoadTimes (reversed)

Change 3214981 on 2016/11/30 by [REDACTED_PERSON]

	UE4 - Make sure that subobjects of CDOs are exported even if they don't have any direct references; they might be archetypes.

Change 3215392 on 2016/11/30 by [REDACTED_PERSON]

	Error out on editor builds when reading a non-boolean value in an archive.

Change 3215674 on 2016/11/30 by [REDACTED_PERSON]

	Replacement of a custom scope-exit setup with our standard macro.

Change 3215720 on 2016/11/30 by [REDACTED_PERSON]

	Default constructor for TTuple which value-initializes the elements.

Change 3216777 on 2016/12/01 by [REDACTED_PERSON]

	Add NoRedist and NotForLicensees folders in the game folder to the C# version of config file hierarchy.

Change 3216858 on 2016/12/01 by [REDACTED_PERSON]

	Improvements to pak signing and encryption
	 - Remove compile time defines for encryption key and signing keys
	 - Embed keys in executable by including in the UBT generated UELinkerFixups.cpp file
	 - Add key access core delegate for passing these keys through to pak platform file

Change 3216860 on 2016/12/01 by [REDACTED_PERSON]

	Re-enable pak signing and encryption in ShooterGame

Change 3216861 on 2016/12/01 by [REDACTED_PERSON]

	Re-enable pak signing and encryption in Paragon

Change 3217076 on 2016/12/01 by [REDACTED_PERSON]

	UE4 - replaced !GIsIntialLoad with EVENT_DRIVEN_ASYNC_LOAD_ACTIVE_AT_RUNTIME as a first step toward boot time EDL

Change 3221139 on 2016/12/05 by [REDACTED_PERSON]

	Dummy integrate of CL 3221131 (versioning workarounds) from Dev-LoadTimes4.14 with accept target to prevent those workarounds being merged out of dev-loadtimes4.14.

	Command:
	p4 integrate //UE4/Dev-LoadTimes4.14/...[REDACTED_HANDLE],3221131

	w/ resolve/accept-target

Change 3221410 on 2016/12/05 by [REDACTED_PERSON]

	Allow Algo::IsSorted to work on C arrays and initializer_lists.

Change 3221673 on 2016/12/05 by [REDACTED_PERSON]

	set up an intentional merge conflict so we can do the right thing when dev-BP and dev-core come together

Change 3223182 on 2016/12/06 by [REDACTED_PERSON]

	Add EngineModeEx to CrashReportCommon
	[UE-39258] - Update CrashReporter to include whether Engine crashes are from vanilla build or not

Change 3224646 on 2016/12/07 by [REDACTED_PERSON]

	It's no longer necessary to rebuild the cooked exe to enable the event driven loader. EDL and the new async IO are now controlled with ini settings only.

Change 3224647 on 2016/12/07 by [REDACTED_PERSON]

	Event Driven Loader is now enabled by default.

Change 3224669 on 2016/12/07 by [REDACTED_PERSON]

	Compressing some crash files output from CRP to S3 (CRP v1.2.11)
	[UE-37564] - Add compression support for CRP crash files in S3

	Add EngineModeEx to support to CRP (CRP v1.2.10)
	[UE-39258] - Update CrashReporter to include whether Engine crashes are from vanilla build or not

Change 3224873 on 2016/12/07 by [REDACTED_PERSON]

	A few fixes to the EDL macro refactor.

Change 3224933 on 2016/12/07 by [REDACTED_PERSON]

	Async Loading settings will now be logged on startup in cooked games

Change 3224984 on 2016/12/07 by [REDACTED_PERSON]

	Small fix to EDL and new async IO global var init.

Change 3225027 on 2016/12/07 by [REDACTED_PERSON]

	Re-enabling EDL in ShooterGame since it no longer requires the cooked exe to be re-compiled (and EDL is on by default anyway).

Change 3226702 on 2016/12/08 by [REDACTED_PERSON]

	UHT can now skip over macro-like identifiers while parsing.
	Correctly handles FTypefaceEntry::Name parsing, which was broken by the workaround in CL# 3219332.
	Syntax workaround reverted in FTypefaceEntry.

Change 3226756 on 2016/12/08 by [REDACTED_PERSON]

	LexicalConversion renamed to Lex after #core discussion.
	Lex::ToString overloads added for FString itself.

Change 3226766 on 2016/12/08 by [REDACTED_PERSON]

	FCookStatsManager::ToString replaced with Lex::ToString, as it now supports the required functionality.

Change 3226949 on 2016/12/08 by [REDACTED_PERSON]

	Fixing static analysis warnings

Change 3228566 on 2016/12/09 by [REDACTED_PERSON]

	Making UBT not think it needs to use separate temporary target for non source code projects when EDL settings don't match default.

Change 3228601 on 2016/12/09 by [REDACTED_PERSON]

	Fix for TSparseArray visualization of types smaller than 8 bytes.

Change 3228654 on 2016/12/09 by [REDACTED_PERSON]

	UE4 - Fix very old bug with lock free lists, probably not relevant to anything we currently use.

Change 3228697 on 2016/12/09 by [REDACTED_PERSON]

	Rebuilt UnrealPak

diff --git a/Engine/Config/BaseEngine.ini b/Engine/Config/BaseEngine.ini
--- a/Engine/Config/BaseEngine.ini
+++ b/Engine/Config/BaseEngine.ini
@@ -1680,6 +1680,7 @@ HangDuration=0.0
 s.AsyncIOBandwidthLimit=0
 s.MinBulkDataSizeForAsyncLoading=131072
 s.AsyncLoadingThreadEnabled=False
+s.EventDrivenLoaderEnabled=True
 s.WarnIfTimeLimitExceeded=False
 s.TimeLimitExceededMultiplier=1.5
 s.TimeLimitExceededMinTime=0.005
diff --git a/Engine/Plugins/Developer/VisualStudioSourceCodeAccess/Source/VisualStudioSourceCodeAccess/Private/VisualStudioSourceCodeAccessor.cpp b/Engine/Plugins/Developer/VisualStudioSourceCodeAccess/Source/VisualStudioSourceCodeAccess/Private/VisualStudioSourceCodeAccessor.cpp
--- a/Engine/Plugins/Developer/VisualStudioSourceCodeAccess/Source/VisualStudioSourceCodeAccess/Private/VisualStudioSourceCodeAccessor.cpp
+++ b/Engine/Plugins/Developer/VisualStudioSourceCodeAccess/Source/VisualStudioSourceCodeAccess/Private/VisualStudioSourceCodeAccessor.cpp
@@ -120,7 +120,7 @@ int32 GetVisualStudioVersionForSolution(const FString& InSolutionFile)
 			}
 
 			int32 Version = 0;
-			LexicalConversion::FromString(Version, *VersionString);
+			Lex::FromString(Version, *VersionString);
 			return Version;
 		}
 	}
diff --git a/Engine/Plugins/Editor/EpicSurvey/Source/EpicSurvey/Private/SurveyTitleLocalStorage.cpp b/Engine/Plugins/Editor/EpicSurvey/Source/EpicSurvey/Private/SurveyTitleLocalStorage.cpp
--- a/Engine/Plugins/Editor/EpicSurvey/Source/EpicSurvey/Private/SurveyTitleLocalStorage.cpp
+++ b/Engine/Plugins/Editor/EpicSurvey/Source/EpicSurvey/Private/SurveyTitleLocalStorage.cpp
@@ -78,7 +78,7 @@ bool FSurveyTitleLocalStorage::EnumerateFiles(const FPagedQuery& Page)
 
 		FCloudFileHeader NewHeader;
 		NewHeader.FileName = Filename;
-		NewHeader.DLName = Filename + LexicalConversion::ToString(FileIdx);
+		NewHeader.DLName = Filename + Lex::ToString(FileIdx);
 		NewHeader.FileSize = 0;
 		NewHeader.Hash.Empty();
 
diff --git a/Engine/Plugins/MovieScene/LevelSequenceEditor/Source/LevelSequenceEditor/Private/CinematicViewport/SCinematicLevelViewport.cpp b/Engine/Plugins/MovieScene/LevelSequenceEditor/Source/LevelSequenceEditor/Private/CinematicViewport/SCinematicLevelViewport.cpp
--- a/Engine/Plugins/MovieScene/LevelSequenceEditor/Source/LevelSequenceEditor/Private/CinematicViewport/SCinematicLevelViewport.cpp
+++ b/Engine/Plugins/MovieScene/LevelSequenceEditor/Source/LevelSequenceEditor/Private/CinematicViewport/SCinematicLevelViewport.cpp
@@ -692,7 +692,7 @@ void SCinematicLevelViewport::Tick(const FGeometry& AllottedGeometry, const doub
 		else
 		{
 			// Just use the FoV
-			UIData.Filmback = FText::FromString(LexicalConversion::ToString(FNumericUnit<float>(CameraComponent->FieldOfView, EUnit::Degrees)));
+			UIData.Filmback = FText::FromString(Lex::ToString(FNumericUnit<float>(CameraComponent->FieldOfView, EUnit::Degrees)));
 		}
 	}
 	else
diff --git a/Engine/Plugins/Online/Android/OnlineSubsystemGameCircle/Source/OnlineSubsystemGameCircle/Private/OnlineStoreInterfaceGameCircle.cpp b/Engine/Plugins/Online/Android/OnlineSubsystemGameCircle/Source/OnlineSubsystemGameCircle/Private/OnlineStoreInterfaceGameCircle.cpp
--- a/Engine/Plugins/Online/Android/OnlineSubsystemGameCircle/Source/OnlineSubsystemGameCircle/Private/OnlineStoreInterfaceGameCircle.cpp
+++ b/Engine/Plugins/Online/Android/OnlineSubsystemGameCircle/Source/OnlineSubsystemGameCircle/Private/OnlineStoreInterfaceGameCircle.cpp
@@ -109,7 +109,7 @@ extern "C" void Java_com_epicgames_ue4_AmazonStoreHelper_nativeQueryComplete(JNI
 			jenv->ReleaseStringUTFChars(NextPrice, charsPrice);
 			jenv->DeleteLocalRef(NextPrice);
 
-			LexicalConversion::FromString(NewProductInfo.RawPrice, *NewProductInfo.DisplayPrice);
+			Lex::FromString(NewProductInfo.RawPrice, *NewProductInfo.DisplayPrice);
 
 			ProvidedProductInformation.Add(NewProductInfo);
 
diff --git a/Engine/Plugins/Runtime/GearVR/Source/GearVR/Private/GearVRSplash.cpp b/Engine/Plugins/Runtime/GearVR/Source/GearVR/Private/GearVRSplash.cpp
--- a/Engine/Plugins/Runtime/GearVR/Source/GearVR/Private/GearVRSplash.cpp
+++ b/Engine/Plugins/Runtime/GearVR/Source/GearVR/Private/GearVRSplash.cpp
@@ -74,7 +74,7 @@ FGearVRSplash::FGearVRSplash(FGearVR* InPlugin) :
 		{
 			break;
 		}
-		num = LexicalConversion::ToString(i);
+		num = Lex::ToString(i);
 	}
 
 	// set loading icon mode in the case if:
diff --git a/Engine/Plugins/Runtime/OculusRift/Source/OculusRift/Private/OculusRiftSplash.cpp b/Engine/Plugins/Runtime/OculusRift/Source/OculusRift/Private/OculusRiftSplash.cpp
--- a/Engine/Plugins/Runtime/OculusRift/Source/OculusRift/Private/OculusRiftSplash.cpp
+++ b/Engine/Plugins/Runtime/OculusRift/Source/OculusRift/Private/OculusRiftSplash.cpp
@@ -64,7 +64,7 @@ FOculusRiftSplash::FOculusRiftSplash(FOculusRiftHMD* InPlugin) :
 		{
 			AddSplash(SplashDesc);
 		}
-		num = LexicalConversion::ToString(i);
+		num = Lex::ToString(i);
 	}
 }
 
diff --git a/Engine/Source/Editor/KismetCompiler/Private/KismetCompilerVMBackend.cpp b/Engine/Source/Editor/KismetCompiler/Private/KismetCompilerVMBackend.cpp
--- a/Engine/Source/Editor/KismetCompiler/Private/KismetCompilerVMBackend.cpp
+++ b/Engine/Source/Editor/KismetCompiler/Private/KismetCompilerVMBackend.cpp
@@ -680,14 +680,14 @@ public:
 			else if (FLiteralTypeHelper::IsInt64(&Term->Type, CoerceProperty))
 			{
 				int64 Value = 0;
-				LexicalConversion::FromString(Value, *(Term->Name));
+				Lex::FromString(Value, *(Term->Name));
 				Writer << EX_Int64Const;
 				Writer << Value;
 			}
 			else if (FLiteralTypeHelper::IsUInt64(&Term->Type, CoerceProperty))
 			{
 				uint64 Value = 0;
-				LexicalConversion::FromString(Value, *(Term->Name));
+				Lex::FromString(Value, *(Term->Name));
 				Writer << EX_UInt64Const;
 				Writer << Value;
 			}
diff --git a/Engine/Source/Editor/PropertyEditor/Private/PropertyHandleImpl.cpp b/Engine/Source/Editor/PropertyEditor/Private/PropertyHandleImpl.cpp
--- a/Engine/Source/Editor/PropertyEditor/Private/PropertyHandleImpl.cpp
+++ b/Engine/Source/Editor/PropertyEditor/Private/PropertyHandleImpl.cpp
@@ -2713,7 +2713,7 @@ FPropertyAccess::Result FPropertyHandleInt::SetValue(const int8& NewValue, EProp
 	// Clamp the value from any meta data ranges stored on the property value
 	int8 FinalValue = ClampIntegerValueFromMetaData<int8>( NewValue, *Implementation->GetPropertyNode() );
 
-	const FString ValueStr = LexicalConversion::ToString(FinalValue);
+	const FString ValueStr = Lex::ToString(FinalValue);
 	Res = Implementation->ImportText(ValueStr, Flags);
 
 	return Res;
@@ -2726,7 +2726,7 @@ FPropertyAccess::Result FPropertyHandleInt::SetValue(const int16& NewValue, EPro
 	// Clamp the value from any meta data ranges stored on the property value
 	int16 FinalValue = ClampIntegerValueFromMetaData<int16>(NewValue, *Implementation->GetPropertyNode());
 
-	const FString ValueStr = LexicalConversion::ToString(FinalValue);
+	const FString ValueStr = Lex::ToString(FinalValue);
 	Res = Implementation->ImportText(ValueStr, Flags);
 
 	return Res;
@@ -2739,7 +2739,7 @@ FPropertyAccess::Result FPropertyHandleInt::SetValue( const int32& NewValue, EPr
 	// Clamp the value from any meta data ranges stored on the property value
 	int32 FinalValue = ClampIntegerValueFromMetaData<int32>( NewValue, *Implementation->GetPropertyNode() );
 
-	const FString ValueStr = LexicalConversion::ToString(FinalValue);
+	const FString ValueStr = Lex::ToString(FinalValue);
 	Res = Implementation->ImportText( ValueStr, Flags );
 
 	return Res;
@@ -2752,7 +2752,7 @@ FPropertyAccess::Result FPropertyHandleInt::SetValue(const int64& NewValue, EPro
 	// Clamp the value from any meta data ranges stored on the property value
 	int64 FinalValue = ClampIntegerValueFromMetaData<int64>(NewValue, *Implementation->GetPropertyNode());
 
-	const FString ValueStr = LexicalConversion::ToString(FinalValue);
+	const FString ValueStr = Lex::ToString(FinalValue);
 	Res = Implementation->ImportText(ValueStr, Flags);
 	return Res;
 }
@@ -2763,7 +2763,7 @@ FPropertyAccess::Result FPropertyHandleInt::SetValue(const uint16& NewValue, EPr
 	// Clamp the value from any meta data ranges stored on the property value
 	uint16 FinalValue = ClampIntegerValueFromMetaData<uint16>(NewValue, *Implementation->GetPropertyNode());
 
-	const FString ValueStr = LexicalConversion::ToString(FinalValue);
+	const FString ValueStr = Lex::ToString(FinalValue);
 	Res = Implementation->ImportText(ValueStr, Flags);
 
 	return Res;
@@ -2776,7 +2776,7 @@ FPropertyAccess::Result FPropertyHandleInt::SetValue(const uint32& NewValue, EPr
 	// Clamp the value from any meta data ranges stored on the property value
 	uint32 FinalValue = ClampIntegerValueFromMetaData<uint32>(NewValue, *Implementation->GetPropertyNode());
 
-	const FString ValueStr = LexicalConversion::ToString(FinalValue);
+	const FString ValueStr = Lex::ToString(FinalValue);
 	Res = Implementation->ImportText(ValueStr, Flags);
 
 	return Res;
@@ -2788,7 +2788,7 @@ FPropertyAccess::Result FPropertyHandleInt::SetValue(const uint64& NewValue, EPr
 	// Clamp the value from any meta data ranges stored on the property value
 	uint64 FinalValue = ClampIntegerValueFromMetaData<uint64>(NewValue, *Implementation->GetPropertyNode());
 
-	const FString ValueStr = LexicalConversion::ToString(FinalValue);
+	const FString ValueStr = Lex::ToString(FinalValue);
 	Res = Implementation->ImportText(ValueStr, Flags);
 	return Res;
 }
diff --git a/Engine/Source/Editor/SequenceRecorder/Private/SequenceRecorderModule.cpp b/Engine/Source/Editor/SequenceRecorder/Private/SequenceRecorderModule.cpp
--- a/Engine/Source/Editor/SequenceRecorder/Private/SequenceRecorderModule.cpp
+++ b/Engine/Source/Editor/SequenceRecorder/Private/SequenceRecorderModule.cpp
@@ -579,7 +579,7 @@ class FSequenceRecorderModule : public ISequenceRecorder, private FSelfRegisteri
 	{
 		if (Handle == AudioFactoryHandle)
 		{
-			AudioFactory.Unset();
+			AudioFactory = nullptr;
 			AudioFactoryHandle = FDelegateHandle();
 		}
 	}
diff --git a/Engine/Source/Editor/UnrealEd/Private/Commandlets/CookCommandlet.cpp b/Engine/Source/Editor/UnrealEd/Private/Commandlets/CookCommandlet.cpp
--- a/Engine/Source/Editor/UnrealEd/Private/Commandlets/CookCommandlet.cpp
+++ b/Engine/Source/Editor/UnrealEd/Private/Commandlets/CookCommandlet.cpp
@@ -206,13 +206,13 @@ namespace DetailedCookStats
 				double AssetTimeSec = 0.0;
 				if (AssetTimeSecAttr)
 				{
-					LexicalConversion::FromString(AssetTimeSec, *AssetTimeSecAttr->Value);
+					Lex::FromString(AssetTimeSec, *AssetTimeSecAttr->Value);
 				}
 				const FCookStatsManager::StringKeyValue* AssetSizeMBAttr = StatAttributes.FindByPredicate([](const FCookStatsManager::StringKeyValue& Item) { return Item.Key == TEXT("MB"); });
 				double AssetSizeMB = 0.0;
 				if (AssetSizeMBAttr)
 				{
-					LexicalConversion::FromString(AssetSizeMB, *AssetSizeMBAttr->Value);
+					Lex::FromString(AssetSizeMB, *AssetSizeMBAttr->Value);
 				}
 				const FCookStatsManager::StringKeyValue* ThreadNameAttr = StatAttributes.FindByPredicate([](const FCookStatsManager::StringKeyValue& Item) { return Item.Key == TEXT("ThreadName"); });
 				bool bIsGameThreadTime = ThreadNameAttr != nullptr && ThreadNameAttr->Value == TEXT("GameThread");
@@ -225,7 +225,7 @@ namespace DetailedCookStats
 					const FCookStatsManager::StringKeyValue* CountAttr = StatAttributes.FindByPredicate([](const FCookStatsManager::StringKeyValue& Item) { return Item.Key == TEXT("Count"); });
 					if (CountAttr)
 					{
-						LexicalConversion::FromString(AssetsBuilt, *CountAttr->Value);
+						Lex::FromString(AssetsBuilt, *CountAttr->Value);
 					}
 				}
 
@@ -1550,6 +1550,7 @@ bool UCookCommandlet::NewCook( const TArray<ITargetPlatform*>& Platforms, TArray
 			CookOnTheFlyServer->StartCookByTheBook(StartupOptions);
 		}
 
+
 		// Garbage collection should happen when either
 		//	1. We have cooked a map (configurable asset type)
 		//	2. We have cooked non-map packages and...
@@ -1729,6 +1730,8 @@ bool UCookCommandlet::NewCook( const TArray<ITargetPlatform*>& Platforms, TArray
 		}
 	} while (bTestCook);
 
+	VerifyEDLCookInfo();
+
 	return true;
 }
 
diff --git a/Engine/Source/Editor/UnrealEd/Private/CookOnTheFlyServer.cpp b/Engine/Source/Editor/UnrealEd/Private/CookOnTheFlyServer.cpp
--- a/Engine/Source/Editor/UnrealEd/Private/CookOnTheFlyServer.cpp
+++ b/Engine/Source/Editor/UnrealEd/Private/CookOnTheFlyServer.cpp
@@ -3809,25 +3809,25 @@ bool ReadCookedAssetRegistry(const FString& Filename, TArray<FCookedAssetPackage
 					FString DependentPackageHash;
 					if (JsonPackage->TryGetStringField(TEXT("DependentPackageHash"), DependentPackageHash))
 					{
-						LexicalConversion::FromString(PackageInfo.DependentPackageHash, *DependentPackageHash);
+						Lex::FromString(PackageInfo.DependentPackageHash, *DependentPackageHash);
 					}
 
 					FString FullPackageHash;
 					if (JsonPackage->TryGetStringField(TEXT("FullPackageHash"), FullPackageHash))
 					{
-						LexicalConversion::FromString(PackageInfo.FullPackageHash, *FullPackageHash);
+						Lex::FromString(PackageInfo.FullPackageHash, *FullPackageHash);
 					}
 
 					FString IsEditorPackage;
 					if (JsonPackage->TryGetStringField(TEXT("IsEditorOnly"), IsEditorPackage))
 					{
-						LexicalConversion::FromString(PackageInfo.IsEditorPackage, *IsEditorPackage);
+						Lex::FromString(PackageInfo.IsEditorPackage, *IsEditorPackage);
 					}
 
 					FString SucceededSave;
 					if (JsonPackage->TryGetStringField(TEXT("SucceededSave"), SucceededSave))
 					{
-						LexicalConversion::FromString(PackageInfo.SucceededSave, *SucceededSave);
+						Lex::FromString(PackageInfo.SucceededSave, *SucceededSave);
 					}
 
 
@@ -3998,13 +3998,13 @@ bool UCookOnTheFlyServer::SaveCookedAssetRegistry(const FString& InCookedAssetRe
 		{
 			FMD5Hash Hash;
 			PDInfoModule.DeterminePackageDependentHash(*RelativePath, Hash);
-			FString DependentPackageHashString = LexicalConversion::ToString(Hash);
+			FString DependentPackageHashString = Lex::ToString(Hash);
 			Json->WriteValue(TEXT("DependentPackageHash"), *DependentPackageHashString);
 		}
 		{
 			FMD5Hash Hash;
 			PDInfoModule.DetermineFullPackageHash(*RelativePath, Hash);
-			FString FullPackageHashString = LexicalConversion::ToString(Hash);
+			FString FullPackageHashString = Lex::ToString(Hash);
 			Json->WriteValue(TEXT("FullPackageHash"), *FullPackageHashString);
 		}
 
@@ -5903,6 +5903,12 @@ void UCookOnTheFlyServer::StartCookByTheBook( const FCookByTheBookStartupOptions
 
 	InitializeSandbox();
 
+	if (CurrentCookMode == ECookMode::CookByTheBook && !IsCookFlagSet(ECookInitializationFlags::Iterative))
+	{
+		StartSavingEDLCookInfoForVerification();
+	}
+
+
 	bool bRunConversion = FParse::Param(FCommandLine::Get(), TEXT("NativizeAssets"));
 	if (bRunConversion)
 	{
diff --git a/Engine/Source/Editor/UnrealEd/Private/UnrealEdPrivatePCH.h b/Engine/Source/Editor/UnrealEd/Private/UnrealEdPrivatePCH.h
--- a/Engine/Source/Editor/UnrealEd/Private/UnrealEdPrivatePCH.h
+++ b/Engine/Source/Editor/UnrealEd/Private/UnrealEdPrivatePCH.h
@@ -227,7 +227,6 @@
 #include "Containers/ResourceArray.h"
 #include "Containers/DynamicRHIResourceArray.h"
 #include "Modules/ModuleManager.h"
-#include "Misc/MultiReaderSingleWriterGT.h"
 #include "Modules/Boilerplate/ModuleBoilerplate.h"
 #include "Misc/ScopeLock.h"
 #include "Misc/BufferedOutputDevice.h"
diff --git a/Engine/Source/Editor/UnrealEd/Public/UnrealEdSharedPCH.h b/Engine/Source/Editor/UnrealEd/Public/UnrealEdSharedPCH.h
--- a/Engine/Source/Editor/UnrealEd/Public/UnrealEdSharedPCH.h
+++ b/Engine/Source/Editor/UnrealEd/Public/UnrealEdSharedPCH.h
@@ -223,7 +223,6 @@
 #include "GenericPlatform/GenericPlatformCompression.h"
 #include "Misc/SecureHash.h"
 #include "Modules/ModuleManager.h"
-#include "Misc/MultiReaderSingleWriterGT.h"
 #include "Modules/Boilerplate/ModuleBoilerplate.h"
 #include "Serialization/ArchiveProxy.h"
 #include "UObject/DebugSerializationFlags.h"
diff --git a/Engine/Source/Programs/CrashReporter/CompressionHelper/Compression.cpp b/Engine/Source/Programs/CrashReporter/CompressionHelper/Compression.cpp
--- a/Engine/Source/Programs/CrashReporter/CompressionHelper/Compression.cpp
+++ b/Engine/Source/Programs/CrashReporter/CompressionHelper/Compression.cpp
@@ -1,7 +1,7 @@
 // Copyright 1998-2017 Epic Games, Inc. All Rights Reserved.
 
-#include "../../../ThirdParty/zlib/zlib-1.2.5/Inc/zlib.h"
-#pragma comment( lib, "../../../ThirdParty/zlib/zlib-1.2.5/Lib/Win64/zlib_64.lib" )
+#include "../../../ThirdParty/zlib/v1.2.8/include/Win64/VS2013/zlib.h"
+#pragma comment( lib, "../../../ThirdParty/zlib/v1.2.8/lib/Win64/VS2013/zlibstatic.lib" )
 
 // Taken from GenericPlatform.h
 
@@ -19,19 +19,18 @@ typedef signed long long	int64;		// 64-bit signed.
 
 // Taken from Compression.cpp
 
-/**
- * Thread-safe abstract compression routine. Compresses memory from uncompressed buffer and writes it to compressed
- * buffer. Updates CompressedSize with size of compressed data.
- *
- * @param	UncompressedBuffer			Buffer containing uncompressed data
- * @param	UncompressedSize			Size of uncompressed data in bytes
- * @param	CompressedBuffer			Buffer compressed data is going to be read from
- * @param	CompressedSize				Size of CompressedBuffer data in bytes
- * @return  Less than zero values are error codes if the uncompress fails, greater than zero is the uncompressed size in bytes
- */
 extern "C"
 {
-
+	/**
+	 * Thread-safe abstract decompression routine. Decompresses memory from compressed buffer and writes it to uncompressed
+	 * buffer. Returns actual uncompressed data size or a negative error code.
+	 *
+	 * @param	UncompressedBuffer			Buffer containing uncompressed data
+	 * @param	UncompressedSize			Size of uncompressed data in bytes
+	 * @param	CompressedBuffer			Buffer compressed data is going to be read from
+	 * @param	CompressedSize				Size of CompressedBuffer data in bytes
+	 * @return  Less than zero values are error codes if the uncompress fails, greater than zero is the uncompressed size in bytes
+	 */
 	__declspec(dllexport)
 	int32 __cdecl UE4UncompressMemoryZLIB(void* UncompressedBuffer, int32 UncompressedSize, const void* CompressedBuffer, int32 CompressedSize)
 	{
@@ -42,8 +41,59 @@ extern "C"
 		// Uncompress data.
 		int32 Result = uncompress((uint8*)UncompressedBuffer, &ZUncompressedSize, (const uint8*)CompressedBuffer, ZCompressedSize);
 
-		// Sanity check to make sure we uncompressed as much data as we expected to.
 		return (Result == Z_OK) ? (int32)ZUncompressedSize : Result;
 	}
 
+	/**
+	 * Thread-safe abstract compression routine. Compresses memory from uncompressed buffer and writes it to compressed
+	 * buffer. Returns actual compressed data size or a negative error code.
+	 *
+	 * @param	CompressedBuffer			Buffer compressed data
+	 * @param	CompressedSize				Size of CompressedBuffer data in bytes
+	 * @param	UncompressedBuffer			Buffer containing uncompressed data is going to be read from
+	 * @param	UncompressedSize			Size of uncompressed data in bytes
+	 * @return  Less than zero values are error codes if the compress fails, greater than zero is the compressed size in bytes
+	 */
+	__declspec(dllexport)
+	int32 __cdecl UE4CompressMemoryZLIB(void* CompressedBuffer, int32 CompressedSize, const void* UncompressedBuffer, int32 UncompressedSize)
+	{
+		// Zlib wants to use unsigned long.
+		unsigned long ZCompressedSize = CompressedSize;
+		unsigned long ZUncompressedSize = UncompressedSize;
+
+		// Uncompress data.
+		int32 Result = compress2((uint8*)CompressedBuffer, &ZCompressedSize, (const uint8*)UncompressedBuffer, ZUncompressedSize, Z_DEFAULT_COMPRESSION);
+
+		return (Result == Z_OK) ? (int32)ZCompressedSize : Result;
+	}
+
+	/**
+	 * Thread-safe file compression routine. Compresses memory from uncompressed buffer and writes it to compressed
+	 * gzip file.
+	 *
+	 * @param	Path						File path to write a new compressed gzip file
+	 * @param	UncompressedBuffer			Buffer containing uncompressed data is going to be read from
+	 * @param	UncompressedSize			Size of uncompressed data in bytes
+	 * @return  Less than zero values are error codes if the compress fails, Z_OK if succeeded.
+	 */
+	__declspec(dllexport)
+	int32 __cdecl UE4CompressFileGZIP(const char* Path, const void* UncompressedBuffer, int32 UncompressedSize)
+	{
+		// Zlib wants to use unsigned long.
+		unsigned ZUncompressedSize = UncompressedSize;
+
+		gzFile FilePtr = gzopen(Path, "wb");
+		if (FilePtr == nullptr)
+		{
+			return Z_ERRNO;
+		}
+
+		int ZCompressedSize = gzwrite(FilePtr, UncompressedBuffer, ZUncompressedSize);
+		if (ZCompressedSize == 0)
+		{
+			return Z_ERRNO;
+		}
+
+		return gzclose(FilePtr);
+	}
 }
\ No newline at end of file
diff --git a/Engine/Source/Programs/CrashReporter/CrashReportClient/Private/CrashDescription.cpp b/Engine/Source/Programs/CrashReporter/CrashReportClient/Private/CrashDescription.cpp
--- a/Engine/Source/Programs/CrashReporter/CrashReportClient/Private/CrashDescription.cpp
+++ b/Engine/Source/Programs/CrashReporter/CrashReportClient/Private/CrashDescription.cpp
@@ -66,7 +66,7 @@ FCrashProperty& FCrashProperty::operator=(const bool NewValue)
 FCrashProperty& FCrashProperty::operator=(const int64 NewValue)
 {
 	bSet = true;
-	CachedValue = LexicalConversion::ToString( NewValue );
+	CachedValue = Lex::ToString( NewValue );
 	Owner->SetCrashProperty( MainCategory, SecondCategory, CachedValue );
 	return *this;
 }
diff --git a/Engine/Source/Programs/CrashReporter/CrashReportClient/Private/CrashDescription.h b/Engine/Source/Programs/CrashReporter/CrashReportClient/Private/CrashDescription.h
--- a/Engine/Source/Programs/CrashReporter/CrashReportClient/Private/CrashDescription.h
+++ b/Engine/Source/Programs/CrashReporter/CrashReportClient/Private/CrashDescription.h
@@ -78,7 +78,7 @@ struct FAnalyticsEventAttribute;
 	"bAllowToBeContacted"
  */
 
-namespace LexicalConversion
+namespace Lex
 {
 	inline void FromString( ECrashDescVersions& OutValue, const TCHAR* Buffer )
 	{
@@ -455,7 +455,7 @@ protected:
 			const FXmlNode* CategoryNode = MainNode->FindChildNode( SecondCategory );
 			if (CategoryNode)
 			{
-				LexicalConversion::FromString( out_ReadValue, *FGenericCrashContext::UnescapeXMLString( CategoryNode->GetContent() ) );
+				Lex::FromString( out_ReadValue, *FGenericCrashContext::UnescapeXMLString( CategoryNode->GetContent() ) );
 			}
 		}
 	}
diff --git a/Engine/Source/Programs/CrashReporter/CrashReportCommon/AmazonClient.cs b/Engine/Source/Programs/CrashReporter/CrashReportCommon/AmazonClient.cs
--- a/Engine/Source/Programs/CrashReporter/CrashReportCommon/AmazonClient.cs
+++ b/Engine/Source/Programs/CrashReporter/CrashReportCommon/AmazonClient.cs
@@ -2,6 +2,7 @@
 
 using System;
 using System.Collections.Generic;
+using System.IO;
 using System.Net;
 using Amazon.Runtime;
 using Amazon.S3;
@@ -184,6 +185,27 @@ namespace Tools.CrashReporter.CrashReportCommon
 			return S3Client.PutObject(ObjectRequest);
 		}
 
+		/// <summary>
+		/// Uploads a stream to S3 synchronously (blocks until complete).
+		/// </summary>
+		/// <param name="BucketName">S3 bucket name</param>
+		/// <param name="Key">S3 object key</param>
+		/// <param name="InStream">Source stream containing data to upload</param>
+		/// <returns>S3 object response for the uploaded object or null if it fails.</returns>
+		public PutObjectResponse PutS3ObjectFromStream(string BucketName, string Key, Stream InStream)
+		{
+			if (!IsS3Valid) return null;
+
+			PutObjectRequest ObjectRequest = new PutObjectRequest
+			{
+				BucketName = BucketName,
+				Key = Key,
+				InputStream = InStream
+			};
+
+			return S3Client.PutObject(ObjectRequest);
+		}
+
 		private bool DeleteRecordSQS(string QueueUrl, Message InMessage)
 		{
 			DeleteMessageRequest DeleteRequest = new DeleteMessageRequest
diff --git a/Engine/Source/Programs/CrashReporter/CrashReportCommon/Common.cs b/Engine/Source/Programs/CrashReporter/CrashReportCommon/Common.cs
--- a/Engine/Source/Programs/CrashReporter/CrashReportCommon/Common.cs
+++ b/Engine/Source/Programs/CrashReporter/CrashReportCommon/Common.cs
@@ -69,6 +69,10 @@ namespace Tools.CrashReporter.CrashReportCommon
 		[XmlElement]
 		public string EngineMode = "";
 
+		/// <summary>Extended mode info (Vanilla, Dirty, Unset).</summary>
+		[XmlElement]
+		public string EngineModeEx = "";
+
 		/// <summary>The four component version of the app e.g. 4.4.1.0</summary>
 		[XmlElement]
 		public string EngineVersion = "";
diff --git a/Engine/Source/Programs/CrashReporter/CrashReportCommon/CrashContext.cs b/Engine/Source/Programs/CrashReporter/CrashReportCommon/CrashContext.cs
--- a/Engine/Source/Programs/CrashReporter/CrashReportCommon/CrashContext.cs
+++ b/Engine/Source/Programs/CrashReporter/CrashReportCommon/CrashContext.cs
@@ -37,20 +37,27 @@ namespace Tools.CrashReporter.CrashReportCommon
 	/* PrimaryCrashProperties. Extracted from: FGenericCrashContext::SerializeContentToBuffer */
 	/*
 		"CrashVersion"
-		"ProcessId"
 		"CrashGUID"
+		"ProcessId"
 		"IsInternalBuild"
 		"IsPerforceBuild"
 		"IsSourceDistribution"
+		"IsEnsure"
+		"IsAssert"
+		"CrashType"
 		"SecondsSinceStart"
 		"GameName"
 		"ExecutableName"
 		"BuildConfiguration"
+		"GameSessionID"
 		"PlatformName"
 		"PlatformNameIni"
 		"PlatformFullName"
 		"EngineMode"
+		"EngineModeEx"
+		"DeploymentName"
 		"EngineVersion"
+		"BuildVersion"
 		"CommandLine"
 		"LanguageLCID"
 		"AppDefaultLocale"
@@ -67,6 +74,7 @@ namespace Tools.CrashReporter.CrashReportCommon
 		"UserActivityHint" 
 		"ErrorMessage"
 		"CrashDumpMode"
+		"FullCrashDumpLocation"
 		"Misc.NumberOfCores"
 		"Misc.NumberOfCoresIncludingHyperthreads"
 		"Misc.Is64bitOperatingSystem"
@@ -92,6 +100,7 @@ namespace Tools.CrashReporter.CrashReportCommon
 		"MemoryStats.OOMAllocationAlignment"
 		"TimeofCrash"
 		"bAllowToBeContacted"
+		"PlatformCallbackResult"
 	*/
 
 	/// <summary> Enumerates crash description versions. </summary>
@@ -425,6 +434,15 @@ namespace Tools.CrashReporter.CrashReportCommon
 		[XmlElement]
 		public string EngineMode;
 
+		/// <summary></summary>
+		[XmlElement]
+		public string EngineModeEx;
+
+		/// <summary>
+		/// Get the engine mode ex. Can be empty but we don't want null.
+		/// </summary>
+		public string GetEngineModeEx() { return EngineModeEx ?? string.Empty; }
+
 		/// <summary></summary>
 		[XmlElement]
 		public string DeploymentName;
diff --git a/Engine/Source/Programs/CrashReporter/CrashReportProcess/CompressionNativeMethods.cs b/Engine/Source/Programs/CrashReporter/CrashReportProcess/CompressionNativeMethods.cs
new file mode 100644
--- /dev/null
+++ b/Engine/Source/Programs/CrashReporter/CrashReportProcess/CompressionNativeMethods.cs
@@ -0,0 +1,58 @@
+﻿using System;
+using System.Runtime.InteropServices;
+
+namespace Tools.CrashReporter.CrashReportProcess
+{
+	static class NativeMethods
+	{
+		public static unsafe int UncompressMemoryZlib(byte[] OutUncompressedBuffer, byte[] InCompressedBuffer)
+		{
+			fixed (byte* UncompressedBufferPtr = OutUncompressedBuffer, CompressedBufferPtr = InCompressedBuffer)
+			{
+				return UE4UncompressMemoryZLIB((IntPtr)UncompressedBufferPtr, OutUncompressedBuffer.Length, (IntPtr)CompressedBufferPtr, InCompressedBuffer.Length);
+			}
+		}
+
+		public static unsafe int CompressMemoryZlib(byte[] OutCompressedBuffer, byte[] InUncompressedBuffer)
+		{
+			fixed (byte* CompressedBufferPtr = OutCompressedBuffer, UncompressedBufferPtr = InUncompressedBuffer)
+			{
+				return UE4CompressMemoryZLIB((IntPtr)CompressedBufferPtr, OutCompressedBuffer.Length, (IntPtr)UncompressedBufferPtr, InUncompressedBuffer.Length);
+			}
+		}
+
+		public static unsafe int CompressFileGZIP(string InFilepath, byte[] InUncompressedBuffer)
+		{
+			fixed (byte* UncompressedBufferPtr = InUncompressedBuffer)
+			{
+				//fixed (char* FilepathPtr = InFilepath)
+				{
+					return UE4CompressFileGZIP(InFilepath, (IntPtr)UncompressedBufferPtr, InUncompressedBuffer.Length);
+				}
+			}
+		}
+
+		[DllImport("CompressionHelper.dll", CallingConvention = CallingConvention.Cdecl)]
+		private static extern Int32 UE4UncompressMemoryZLIB( /*void**/ IntPtr UncompressedBuffer, Int32 UncompressedSize, /*const void**/ IntPtr CompressedBuffer, Int32 CompressedSize);
+
+		[DllImport("CompressionHelper.dll", CallingConvention = CallingConvention.Cdecl)]
+		private static extern Int32 UE4CompressMemoryZLIB( /*void**/ IntPtr CompressedBuffer, Int32 CompressedSize, /*const void**/ IntPtr UncompressedBuffer, Int32 UncompressedSize);
+
+		[DllImport("CompressionHelper.dll", CallingConvention = CallingConvention.Cdecl, CharSet = CharSet.Ansi)]
+		private static extern Int32 UE4CompressFileGZIP( /*const char**/ string Path, /*const void**/ IntPtr UncompressedBuffer, Int32 UncompressedSize);
+
+		public static string GetZlibError(Int32 ErrorCode)
+		{
+			if (ErrorCode == 0) return "Z_OK";
+			if (ErrorCode == 1) return "Z_STREAM_END";
+			if (ErrorCode == 2) return "Z_NEED_DICT";
+			if (ErrorCode == -1) return "Z_ERRNO";
+			if (ErrorCode == -2) return "Z_STREAM_ERROR";
+			if (ErrorCode == -3) return "Z_DATA_ERROR";
+			if (ErrorCode == -4) return "Z_MEM_ERROR";
+			if (ErrorCode == -5) return "Z_BUF_ERROR";
+			if (ErrorCode == -6) return "Z_VERSION_ERROR";
+			return "UnknownZLibError";
+		}
+	}
+}
diff --git a/Engine/Source/Programs/CrashReporter/CrashReportProcess/Config.cs b/Engine/Source/Programs/CrashReporter/CrashReportProcess/Config.cs
--- a/Engine/Source/Programs/CrashReporter/CrashReportProcess/Config.cs
+++ b/Engine/Source/Programs/CrashReporter/CrashReportProcess/Config.cs
@@ -205,12 +205,6 @@ namespace Tools.CrashReporter.CrashReportProcess
 		[XmlElement]
 		public int MinutesBetweenQueueSizeReports { get; set; }
 
-		/// <summary>
-		/// The limit for AddCrash() calls to the website failing in a row. An alert will be created if this is exceeded.
-		/// </summary>
-		[XmlElement]
-		public int ConsecutiveFailedWebAddLimit { get; set; }
-
 		/// <summary>
 		/// Size limit below which queues will attempt to enqueue more crashes into memory
 		/// Above this size, a queue will skip enqueueing until the next update.
@@ -278,6 +272,12 @@ namespace Tools.CrashReporter.CrashReportProcess
 		[XmlElement]
 		public bool CrashFilesToAWS { get; set; }
 
+		/// <summary>
+		/// Should large crash files saved to S3 be compressed? (crash context will never be compressed)
+		/// </summary>
+		[XmlElement]
+		public bool CompressCrashFilesOnAWS { get; set; }
+
 		/// <summary>
 		/// Should we save invalid reports that fail to process to S3?
 		/// </summary>
@@ -302,6 +302,12 @@ namespace Tools.CrashReporter.CrashReportProcess
 		[XmlElement]
 		public string AWSS3InvalidKeyPrefix { get; set; }
 
+		/// <summary>
+		/// AWSSDK AWS S3 path/key suffix used in place of the existing extension for writing compressed reports (optional)
+		/// </summary>
+		[XmlElement]
+		public string AWSS3CompressedSuffix { get; set; }
+
 		/// <summary>
 		/// Buffer size used to decompress zlib archives taken from S3
 		/// </summary>
@@ -324,13 +330,19 @@ namespace Tools.CrashReporter.CrashReportProcess
 		/// Timeout when calling AddCrash to submit crashes to the website/database.
 		/// </summary>
 		[XmlElement]
-		public int AddCrashRequestTimeoutMillisec { get; set; }
+		public int AddCrashRequestTimeoutSeconds { get; set; }
 
 		/// <summary>
 		/// Time that we wait between a failed AddCrash call and a retry.
 		/// </summary>
 		[XmlElement]
-		public int AddCrashRetryDelayMillisec { get; set; }
+		public int AddCrashRetryDelaySeconds { get; set; }
+
+		/// <summary>
+		/// The time limit for consecutive failed AddCrash() calls after which an alert will be created.
+		/// </summary>
+		[XmlElement]
+		public int FailedWebAddAlertTimeSeconds { get; set; }
 
 		/// <summary>
 		/// Disk space available threshold that generates alerts. If a disk has less space than this, it will generate alerts.
@@ -422,7 +434,6 @@ namespace Tools.CrashReporter.CrashReportProcess
 			LoadedConfig.CrashReportWebSite = string.Empty;
 			LoadedConfig.AWSS3OutputKeyPrefix = LoadedConfig.AWSS3OutputKeyPrefix.Replace("prod", "test");
 			LoadedConfig.AWSS3InvalidKeyPrefix = LoadedConfig.AWSS3InvalidKeyPrefix.Replace("prod", "test");
-
 			LoadedConfig.MinDesiredMemoryQueueSize = 5;
 			LoadedConfig.MaxMemoryQueueSize = 15;
 
diff --git a/Engine/Source/Programs/CrashReporter/CrashReportProcess/Config/CrashReportProcess.config b/Engine/Source/Programs/CrashReporter/CrashReportProcess/Config/CrashReportProcess.config
--- a/Engine/Source/Programs/CrashReporter/CrashReportProcess/Config/CrashReportProcess.config
+++ b/Engine/Source/Programs/CrashReporter/CrashReportProcess/Config/CrashReportProcess.config
@@ -32,7 +32,6 @@
 	<SlackDecimateAlertRepeatMinimumMinutes>45</SlackDecimateAlertRepeatMinimumMinutes>
 	<DebugTestingFolder>MyDebugTestingFolderPath</DebugTestingFolder>
 	<MinutesBetweenQueueSizeReports>60</MinutesBetweenQueueSizeReports>
-	<ConsecutiveFailedWebAddLimit>10</ConsecutiveFailedWebAddLimit>
 	<MinDesiredMemoryQueueSize>50</MinDesiredMemoryQueueSize>
 	<MaxMemoryQueueSize>250</MaxMemoryQueueSize>
 	<AWSCredentialsFilepath>MyAWSCredsFilepath</AWSCredentialsFilepath>
@@ -44,15 +43,18 @@
 	<AWSS3ServiceOutputURL>https://s3.amazonaws.com</AWSS3ServiceOutputURL>
 	<CrashFilesToDisk>true</CrashFilesToDisk>
 	<CrashFilesToAWS>false</CrashFilesToAWS>
+	<CompressCrashFilesOnAWS>false</CompressCrashFilesOnAWS>
 	<InvalidReportsToAWS>false</InvalidReportsToAWS>
 	<AWSS3OutputBucket>crashreporterfiles</AWSS3OutputBucket>
 	<AWSS3OutputKeyPrefix>websitefiles</AWSS3OutputKeyPrefix>
 	<AWSS3InvalidKeyPrefix>invalidcrashreports</AWSS3InvalidKeyPrefix>
+	<AWSS3CompressedSuffix>.gz</AWSS3CompressedSuffix>
 	<MaxUncompressedS3RecordSize>33554432</MaxUncompressedS3RecordSize>
 	<ProcessedReportsIndexPath>MyProcessedReportsIndexFilePath</ProcessedReportsIndexPath>
 	<ReportsIndexRetentionDays>120</ReportsIndexRetentionDays>
-	<AddCrashRequestTimeoutMillisec>300000</AddCrashRequestTimeoutMillisec>
-	<AddCrashRetryDelayMillisec>20000</AddCrashRetryDelayMillisec>
+	<AddCrashRequestTimeoutSeconds>30</AddCrashRequestTimeoutSeconds>
+	<AddCrashRetryDelaySeconds>15</AddCrashRetryDelaySeconds>
+	<FailedWebAddAlertTimeSeconds>300</FailedWebAddAlertTimeSeconds>
 	<DiskSpaceAlertPercent>10.0</DiskSpaceAlertPercent>
 	<MonitorPerformance>false</MonitorPerformance>
 </Config>
\ No newline at end of file
diff --git a/Engine/Source/Programs/CrashReporter/CrashReportProcess/CrashReportProcess.csproj b/Engine/Source/Programs/CrashReporter/CrashReportProcess/CrashReportProcess.csproj
--- a/Engine/Source/Programs/CrashReporter/CrashReportProcess/CrashReportProcess.csproj
+++ b/Engine/Source/Programs/CrashReporter/CrashReportProcess/CrashReportProcess.csproj
@@ -44,7 +44,7 @@
     <AllowUnsafeBlocks>true</AllowUnsafeBlocks>
   </PropertyGroup>
   <PropertyGroup>
-    <SignManifests>true</SignManifests>
+    <SignManifests>false</SignManifests>
   </PropertyGroup>
   <PropertyGroup>
     <ManifestCertificateThumbprint>7B0CF23A7AEEF0637C57BC4BD76F2F43119EFAD6</ManifestCertificateThumbprint>
@@ -93,6 +93,7 @@
     <Reference Include="System.Xml" />
   </ItemGroup>
   <ItemGroup>
+    <Compile Include="CompressionNativeMethods.cs" />
     <Compile Include="Config.cs" />
     <Compile Include="CrashReportProcessService.cs">
       <SubType>Component</SubType>
diff --git a/Engine/Source/Programs/CrashReporter/CrashReportProcess/DataRouterReportQueue.cs b/Engine/Source/Programs/CrashReporter/CrashReportProcess/DataRouterReportQueue.cs
--- a/Engine/Source/Programs/CrashReporter/CrashReportProcess/DataRouterReportQueue.cs
+++ b/Engine/Source/Programs/CrashReporter/CrashReportProcess/DataRouterReportQueue.cs
@@ -151,25 +151,19 @@ namespace Tools.CrashReporter.CrashReportProcess
 			return UnpackedRecordCount;
 		}
 
-		private static unsafe bool DecompressDataRouterContent(byte[] CompressedBufferArray, string InLandingZone)
+		private static bool DecompressDataRouterContent(byte[] CompressedBufferArray, string InLandingZone)
 		{
 			// Decompress to landing zone
 			byte[] UncompressedBufferArray = new byte[Config.Default.MaxUncompressedS3RecordSize];
-			int UncompressedSize = 0;
-			fixed (byte* UncompressedBufferPtr = UncompressedBufferArray, CompressedBufferPtr = CompressedBufferArray)
+
+			int UncompressedSize = NativeMethods.UncompressMemoryZlib(UncompressedBufferArray, CompressedBufferArray);
+			if (UncompressedSize < 0)
 			{
-				Int32 UncompressResult = NativeMethods.UE4UncompressMemoryZLIB((IntPtr) UncompressedBufferPtr,
-					UncompressedBufferArray.Length,
-					(IntPtr) CompressedBufferPtr, CompressedBufferArray.Length);
-				if (UncompressResult < 0)
-				{
-					string FailString = "! DecompressDataRouterContent() failed in UE4UncompressMemoryZLIB() with " +
-					                    NativeMethods.GetUncompressError(UncompressResult);
-					CrashReporterProcessServicer.WriteFailure(FailString);
-					CrashReporterProcessServicer.StatusReporter.IncrementCount(StatusReportingEventNames.ReadS3RecordFailedEvent);
-					return false;
-				}
-				UncompressedSize = UncompressResult;
+				string FailString = "! DecompressDataRouterContent() failed in UncompressMemoryZlib() with " +
+									NativeMethods.GetZlibError(UncompressedSize);
+				CrashReporterProcessServicer.WriteFailure(FailString);
+				CrashReporterProcessServicer.StatusReporter.IncrementCount(StatusReportingEventNames.ReadS3RecordFailedEvent);
+				return false;
 			}
 
 			using (BinaryReader BinaryData = new BinaryReader(new MemoryStream(UncompressedBufferArray, 0, UncompressedSize, false)))
diff --git a/Engine/Source/Programs/CrashReporter/CrashReportProcess/ReportProcessor.cs b/Engine/Source/Programs/CrashReporter/CrashReportProcess/ReportProcessor.cs
--- a/Engine/Source/Programs/CrashReporter/CrashReportProcess/ReportProcessor.cs
+++ b/Engine/Source/Programs/CrashReporter/CrashReportProcess/ReportProcessor.cs
@@ -40,12 +40,7 @@ namespace Tools.CrashReporter.CrashReportProcess
 
 		private static DateTime LoggingDate = DateTime.UtcNow.Date;
 
-		private static Object TickStaticLock = new Object();
-
-		private static Object MinidumpDiagnosticsLock = new Object();
-
-		private static Object FailedUploadLock = new Object();
-		private static int ConsecutiveFailedUploads = 0;
+		private static object TickStaticLock = new object();
 
 		/// <summary>The thread that handles detection of new crashes.</summary>
 		public ReportWatcher Watcher = null;
@@ -255,7 +250,7 @@ namespace Tools.CrashReporter.CrashReportProcess
 						try
 						{
 							string S3IDPrefix = string.Format("/{0}_{1}_{2}/{3}/{4}/", S3KeyTime.Year, S3KeyTime.Month, S3KeyTime.Day, S3KeyTime.Hour, LeafReportName);
-							UploadFileToS3(InvalidFile, Config.Default.AWSS3InvalidKeyPrefix + S3IDPrefix + InvalidFile.Name);
+							UploadFileToS3(InvalidFile, Config.Default.AWSS3InvalidKeyPrefix + S3IDPrefix + InvalidFile.Name, false);
 							FilesMoved++;
 						}
 						catch (Exception Ex)
@@ -331,6 +326,7 @@ namespace Tools.CrashReporter.CrashReportProcess
 				NewCrash.GameName = NewContext.PrimaryCrashProperties.GameName;
 				NewCrash.Platform = NewContext.PrimaryCrashProperties.PlatformFullName;
 				NewCrash.EngineMode = NewContext.PrimaryCrashProperties.EngineMode;
+				NewCrash.EngineModeEx = NewContext.PrimaryCrashProperties.GetEngineModeEx();
 				NewCrash.EngineVersion = EngineVersion.VersionNumber;
 			    NewCrash.BuildVersion = NewContext.PrimaryCrashProperties.BuildVersion;
 
@@ -429,12 +425,10 @@ namespace Tools.CrashReporter.CrashReportProcess
 						RequestString = "http://localhost:80/Crashes/AddCrash/-1"; 
 					}
 
-					string ErrorMessage = string.Empty;
-
-					const int MaxRetries = 1;
-					for (int AddCrashTry = 0; AddCrashTry < MaxRetries + 1; ++AddCrashTry)
+					int Attempt = 1;
+					Func<bool> TryUpload = () =>
 					{
-						string ResponseString = SimpleWebRequest.GetWebServiceResponse(RequestString, Payload, Config.Default.AddCrashRequestTimeoutMillisec);
+						string ResponseString = SimpleWebRequest.GetWebServiceResponse(RequestString, Payload, Config.Default.AddCrashRequestTimeoutSeconds * 1000);
 						if (ResponseString.Length > 0)
 						{
 							// Convert response into a string
@@ -442,17 +436,30 @@ namespace Tools.CrashReporter.CrashReportProcess
 							if (Result.ID > 0)
 							{
 								NewID = Result.ID;
-								break;
+								return true;
 							}
-							CrashReporterProcessServicer.WriteEvent(string.Format("PROC-{0} UploadCrash response timeout (attempt {1} of {2}): {3}", ProcessorIndex, AddCrashTry + 1, MaxRetries + 1, ErrorMessage));
-							ErrorMessage = Result.Message;
+
+							CrashReporterProcessServicer.WriteEvent(string.Format("PROC-{0} UploadCrash response invalid (upload attempt {1}): {2}", ProcessorIndex, Attempt, Result.Message));
 						}
-						Thread.Sleep(Config.Default.AddCrashRetryDelayMillisec);
-					}
+						return false;
+					};
+					Action<bool> OnRetry = bUploaded =>
+					{
+						if (!bUploaded)
+						{
+							Thread.Sleep(Config.Default.AddCrashRetryDelaySeconds*1000);
+							CrashReporterProcessServicer.WriteEvent(string.Format("PROC-{0} UploadCrash retrying upload attempt {1}", ProcessorIndex, Attempt));
+							Attempt++;
+						}
+					};
 
-					if (NewID == -1)
+					for (bool bUploaded = false; !bUploaded; OnRetry(bUploaded))
 					{
-						CrashReporterProcessServicer.WriteFailure(string.Format("PROC-{0} ", ProcessorIndex) + "UploadCrash failed: " + ErrorMessage);
+						bUploaded = TryUpload();
+						CrashReporterProcessServicer.StatusReporter.AlertOnConsecutiveFails("AddCrash",
+						                                                                    "Cannot contact CR website. CRP is paused!",
+						                                                                    "Contact with CR website restored.",
+						                                                                    TimeSpan.FromSeconds(Config.Default.FailedWebAddAlertTimeSeconds), bUploaded);
 					}
 				}
 #if DEBUG
@@ -498,25 +505,11 @@ namespace Tools.CrashReporter.CrashReportProcess
 				// Upload the crash to the database, and retrieve the new row id
 				int ReportID = UploadCrash(CrashDetails);
 
-				lock (FailedUploadLock)
+				if (ReportID <= 0)
 				{
-					if (ReportID <= 0)
-					{
-						// Upload failed
-						ConsecutiveFailedUploads++;
-						if (ConsecutiveFailedUploads == Config.Default.ConsecutiveFailedWebAddLimit)
-						{
-							CrashReporterProcessServicer.StatusReporter.Alert("ReportProcessor.AddReport.NoContact", "Cannot contact Crash Report website.", Config.Default.SlackAlertRepeatMinimumMinutes);
-							CrashReporterProcessServicer.WriteFailure("Cannot contact Crash Report website.");
-						}
-						CrashReporterProcessServicer.WriteFailure(string.Format("PROC-{0} ", ProcessorIndex) + "! NoUpload: Path=" + NewContext.CrashDirectory);
-						string PayloadFailedFileName = Path.Combine(DirInfo.FullName, "PayloadFailed.xml");
-						File.WriteAllText(PayloadFailedFileName, CrashDetails);
-						return false;
-					}
-
-					// Upload succeeded
-					ConsecutiveFailedUploads = 0;
+					// Upload failed
+					CrashReporterProcessServicer.WriteFailure(string.Format("PROC-{0} ", ProcessorIndex) + "! NoUpload: Path=" + NewContext.CrashDirectory);
+					return false;
 				}
 
 				bool bToS3 = Config.Default.CrashFilesToAWS && CrashReporterProcessServicer.OutputAWS.IsS3Valid;
@@ -545,7 +538,7 @@ namespace Tools.CrashReporter.CrashReportProcess
 						if (bToS3)
 						{
 							WriteToS3Timer.Start();
-							UploadFileToS3(LogInfo, Config.Default.AWSS3OutputKeyPrefix + S3IDPrefix + "Launch.log");
+							UploadFileToS3(LogInfo, Config.Default.AWSS3OutputKeyPrefix + S3IDPrefix + "Launch.log", Config.Default.CompressCrashFilesOnAWS, Config.Default.AWSS3CompressedSuffix);
 							WriteToS3Timer.Stop();
 						}
 						if (bToDisk)
@@ -565,7 +558,7 @@ namespace Tools.CrashReporter.CrashReportProcess
 					if (bToS3)
 					{
 						WriteToS3Timer.Start();
-						UploadFileToS3(CrashContextInfo, Config.Default.AWSS3OutputKeyPrefix + S3IDPrefix + FGenericCrashContext.CrashContextRuntimeXMLName);
+						UploadFileToS3(CrashContextInfo, Config.Default.AWSS3OutputKeyPrefix + S3IDPrefix + FGenericCrashContext.CrashContextRuntimeXMLName, false);
 						WriteToS3Timer.Stop();
 					}
 					if (bToDisk)
@@ -585,7 +578,7 @@ namespace Tools.CrashReporter.CrashReportProcess
 						if (bToS3)
 						{
 							WriteToS3Timer.Start();
-							UploadFileToS3(DumpInfo, Config.Default.AWSS3OutputKeyPrefix + S3IDPrefix + "MiniDump.dmp");
+							UploadFileToS3(DumpInfo, Config.Default.AWSS3OutputKeyPrefix + S3IDPrefix + "MiniDump.dmp", Config.Default.CompressCrashFilesOnAWS, Config.Default.AWSS3CompressedSuffix);
 							WriteToS3Timer.Stop();
 						}
 						if (bToDisk)
@@ -609,7 +602,7 @@ namespace Tools.CrashReporter.CrashReportProcess
 						if (bToS3)
 						{
 							WriteToS3Timer.Start();
-							UploadFileToS3(VideoInfo, Config.Default.AWSS3OutputKeyPrefix + S3IDPrefix + CrashReporterConstants.VideoFileName);
+							UploadFileToS3(VideoInfo, Config.Default.AWSS3OutputKeyPrefix + S3IDPrefix + CrashReporterConstants.VideoFileName, Config.Default.CompressCrashFilesOnAWS, Config.Default.AWSS3CompressedSuffix);
 							WriteToS3Timer.Stop();
 						}
 						if (bToDisk)
@@ -656,17 +649,55 @@ namespace Tools.CrashReporter.CrashReportProcess
 			return false;
 		}
 
-		private static void UploadFileToS3(FileInfo FileInfo, string DestFilename)
+		private static void UploadFileToS3(FileInfo FileInfo, string DestFilepath, bool bCompressed, string CompressedSuffix = null)
 		{
 			try
 			{
-				PutObjectResponse Response = CrashReporterProcessServicer.OutputAWS.PutS3ObjectFromFile(Config.Default.AWSS3OutputBucket,
-																										DestFilename,
-																										FileInfo.FullName);
+				if (bCompressed)
+				{
+					string LocalCompressedFilepath = FileInfo.FullName;
+					if (CompressedSuffix != null)
+					{
+						LocalCompressedFilepath += CompressedSuffix;
+					}
+					else
+					{
+						LocalCompressedFilepath += ".gz";
+					}
+
+					byte[] UncompressedBytes = File.ReadAllBytes(FileInfo.FullName);
+
+					int ReturnCode = NativeMethods.CompressFileGZIP(LocalCompressedFilepath, UncompressedBytes);
+					if (ReturnCode < 0)
+					{
+						throw new CrashReporterException(string.Format("Failed to compress {0} to gzip file with error {1}", FileInfo.FullName, NativeMethods.GetZlibError(ReturnCode)));
+					}
+
+					string CompressedDestFilepath = DestFilepath;
+					if (CompressedSuffix != null)
+					{
+						CompressedDestFilepath += CompressedSuffix;
+					}
+
+					PutObjectResponse Response = CrashReporterProcessServicer.OutputAWS.PutS3ObjectFromFile(Config.Default.AWSS3OutputBucket,
+																											CompressedDestFilepath,
+																											LocalCompressedFilepath);
 
-				if (Response == null || Response.HttpStatusCode != HttpStatusCode.OK)
+					if (Response == null || Response.HttpStatusCode != HttpStatusCode.OK)
+					{
+						throw new CrashReporterException(string.Format("Failed to upload compressed gzip {0} to {1}", LocalCompressedFilepath, CompressedDestFilepath));
+					}
+				}
+				else
 				{
-					throw new CrashReporterException(string.Format("Failed to upload {0} to {1}", FileInfo.FullName, DestFilename));
+					PutObjectResponse Response = CrashReporterProcessServicer.OutputAWS.PutS3ObjectFromFile(Config.Default.AWSS3OutputBucket,
+																											DestFilepath,
+																											FileInfo.FullName);
+
+					if (Response == null || Response.HttpStatusCode != HttpStatusCode.OK)
+					{
+						throw new CrashReporterException(string.Format("Failed to upload {0} to {1}", FileInfo.FullName, DestFilepath));
+					}
 				}
 			}
 			catch (Exception Ex)
diff --git a/Engine/Source/Programs/CrashReporter/CrashReportProcess/ReportQueue.cs b/Engine/Source/Programs/CrashReporter/CrashReportProcess/ReportQueue.cs
--- a/Engine/Source/Programs/CrashReporter/CrashReportProcess/ReportQueue.cs
+++ b/Engine/Source/Programs/CrashReporter/CrashReportProcess/ReportQueue.cs
@@ -5,7 +5,6 @@ using System.Collections.Generic;
 using System.Globalization;
 using System.IO;
 using System.Linq;
-using System.Runtime.InteropServices;
 using Tools.CrashReporter.CrashReportCommon;
 using Tools.DotNETCommon.XmlHandler;
 
@@ -548,7 +547,7 @@ namespace Tools.CrashReporter.CrashReportProcess
 		/// <summary> 
 		/// Decompresses a compressed crash report. 
 		/// </summary>
-		unsafe static protected bool DecompressReport(string CompressedReportPath, FCompressedCrashInformation Meta)
+		static protected bool DecompressReport(string CompressedReportPath, FCompressedCrashInformation Meta)
 		{
 			string ReportPath = Path.GetDirectoryName(CompressedReportPath);
 			using (FileStream FileStream = File.OpenRead(CompressedReportPath))
@@ -566,15 +565,12 @@ namespace Tools.CrashReporter.CrashReportProcess
 
 				FileStream.Read(CompressedBufferArray, 0, CompressedSize);
 
-				fixed (byte* UncompressedBufferPtr = UncompressedBufferArray, CompressedBufferPtr = CompressedBufferArray)
+				int UncompressResult = NativeMethods.UncompressMemoryZlib(UncompressedBufferArray, CompressedBufferArray);
+				if (UncompressResult < 0)
 				{
-					Int32 UncompressResult = NativeMethods.UE4UncompressMemoryZLIB((IntPtr)UncompressedBufferPtr, UncompressedSize, (IntPtr)CompressedBufferPtr, CompressedSize);
-					if (UncompressResult < 0)
-					{
-						CrashReporterProcessServicer.WriteFailure("! DecompressReport() failed in UE4UncompressMemoryZLIB() with " + NativeMethods.GetUncompressError(UncompressResult));
-						CrashReporterProcessServicer.WriteFailure("! ZLibFail in DecompressReport(): Path=" + ReportPath);
-						return false;
-					}
+					CrashReporterProcessServicer.WriteFailure("! DecompressReport() failed in UncompressMemoryZlib() with " + NativeMethods.GetZlibError(UncompressResult));
+					CrashReporterProcessServicer.WriteFailure("! ZLibFail in DecompressReport(): Path=" + ReportPath);
+					return false;
 				}
 
 				using (BinaryReader BinaryData = new BinaryReader(new MemoryStream(UncompressedBufferArray, false)))
@@ -708,23 +704,3 @@ namespace Tools.CrashReporter.CrashReportProcess
 		private float DecimationCounter;
 	}
 }
-
-static class NativeMethods
-{
-	[DllImport("CompressionHelper.dll", CallingConvention = CallingConvention.Cdecl)]
-	public extern static Int32 UE4UncompressMemoryZLIB( /*void**/IntPtr UncompressedBuffer, Int32 UncompressedSize, /*const void**/IntPtr CompressedBuffer, Int32 CompressedSize);
-
-	public static string GetUncompressError(Int32 ErrorCode)
-	{
-		if (ErrorCode == 0) return "Z_OK";
-		else if (ErrorCode == 1) return "Z_STREAM_END";
-		else if (ErrorCode == 2) return "Z_NEED_DICT";
-		else if (ErrorCode == -1) return "Z_ERRNO";
-		else if (ErrorCode == -2) return "Z_STREAM_ERROR";
-		else if (ErrorCode == -3) return "Z_DATA_ERROR";
-		else if (ErrorCode == -4) return "Z_MEM_ERROR";
-		else if (ErrorCode == -5) return "Z_BUF_ERROR";
-		else if (ErrorCode == -6) return "Z_VERSION_ERROR";
-		else return "UnknownZLibError";
-	}
-}
\ No newline at end of file
diff --git a/Engine/Source/Programs/CrashReporter/CrashReportProcess/StatusReporting.cs b/Engine/Source/Programs/CrashReporter/CrashReportProcess/StatusReporting.cs
--- a/Engine/Source/Programs/CrashReporter/CrashReportProcess/StatusReporting.cs
+++ b/Engine/Source/Programs/CrashReporter/CrashReportProcess/StatusReporting.cs
@@ -154,6 +154,48 @@ namespace Tools.CrashReporter.CrashReportProcess
 			}
         }
 
+		public void AlertOnConsecutiveFails(string AlertKey, string AlertText, string RecoveredText, TimeSpan TimeSinceSuccessThreshold, bool bSucceeded)
+		{
+			try
+			{
+				lock (DataLock)
+				{
+					if (bSucceeded)
+					{
+						// Success - clear from list of failing keys
+						DateTime FailAlertStartTime;
+						if (FailAlertStartTimes.TryGetValue(AlertKey, out FailAlertStartTime))
+						{
+							if (FailAlertStartTime == DateTime.MaxValue)
+							{
+								Alert(AlertKey, RecoveredText, 0);
+							}
+							FailAlertStartTimes.Remove(AlertKey);
+						}
+					}
+					else if (!FailAlertStartTimes.ContainsKey(AlertKey))
+					{
+						// Failed but no record yet - first failure so note the time
+						FailAlertStartTimes.Add(AlertKey, DateTime.UtcNow);
+					}
+					else
+					{
+						// Failed and not the first consecutive fail
+						TimeSpan TimeSinceStartedFailing = DateTime.UtcNow - FailAlertStartTimes[AlertKey];
+						if (TimeSinceStartedFailing > TimeSinceSuccessThreshold)
+						{
+							FailAlertStartTimes[AlertKey] = DateTime.MaxValue; // forces future fails to make TimeSinceStartedFailing negative until at least one success
+							Alert(AlertKey, AlertText, 0);
+						}
+					}
+				}
+			}
+			catch (Exception Ex)
+			{
+				CrashReporterProcessServicer.WriteException("AlertOnConsecutiveFails failed: " + Ex, Ex);
+			}
+		}
+
 		public void Start()
 		{
 			CrashReporterProcessServicer.WriteSlack(string.Format("CRP started (version {0})", Config.Default.VersionString));
@@ -267,7 +309,7 @@ namespace Tools.CrashReporter.CrashReportProcess
 			if (Config.Default.MonitorPerformance)
 			{
 				StatusReportLoops.Add(
-					new PerfStatusReport(TimeSpan.FromMinutes(15), (InLoop, InPeriod) =>
+					new PerfStatusReport(TimeSpan.FromMinutes(10), (InLoop, InPeriod) =>
 					                     {
 											 PerfStatusReport ThisLoop = (PerfStatusReport)InLoop;
 											 StringBuilder PerfReportMessage = new StringBuilder();
@@ -362,24 +404,27 @@ namespace Tools.CrashReporter.CrashReportProcess
 
 		private bool CheckRecentAlerts(string AlertKey, int RepeatMinimumMinutes)
 		{
-			// Check RecentAlerts for a matching entry without the time period
-			// that limits repeats Config.Default.SlackAlertRepeatMinimumMinutes
-			if (RecentAlerts.ContainsKey(AlertKey))
+			lock (DataLock)
 			{
-				DateTime LastMatchingAlert = RecentAlerts[AlertKey];
-				TimeSpan RepeatMinimum = TimeSpan.FromMinutes(RepeatMinimumMinutes);
-
-				if (LastMatchingAlert < DateTime.UtcNow - RepeatMinimum)
+				// Check RecentAlerts for a matching entry without the time period
+				// that limits repeats Config.Default.SlackAlertRepeatMinimumMinutes
+				if (RecentAlerts.ContainsKey(AlertKey))
 				{
-					RecentAlerts[AlertKey] = DateTime.UtcNow;
-					return true;
+					DateTime LastMatchingAlert = RecentAlerts[AlertKey];
+					TimeSpan RepeatMinimum = TimeSpan.FromMinutes(RepeatMinimumMinutes);
+
+					if (LastMatchingAlert < DateTime.UtcNow - RepeatMinimum)
+					{
+						RecentAlerts[AlertKey] = DateTime.UtcNow;
+						return true;
+					}
+
+					return false;
 				}
 
-				return false;
+				RecentAlerts.Add(AlertKey, DateTime.UtcNow);
+				return true;
 			}
-
-			RecentAlerts.Add(AlertKey, DateTime.UtcNow);
-			return true;
 		}
 
 		private IEnumerable<StatusReportLoop> ReporterTasks;
@@ -391,5 +436,6 @@ namespace Tools.CrashReporter.CrashReportProcess
 		private readonly Object DataLock = new Object();
 		private int SetQueueSizesCallCount = 0;
 		private readonly Dictionary<string, DateTime> RecentAlerts = new Dictionary<string, DateTime>();
+		private readonly Dictionary<string, DateTime> FailAlertStartTimes = new Dictionary<string, DateTime>();
 	}
 }
diff --git a/Engine/Source/Programs/CrashReporter/CrashReportReceiver/UploadsInProgress.cs b/Engine/Source/Programs/CrashReporter/CrashReportReceiver/UploadsInProgress.cs
--- a/Engine/Source/Programs/CrashReporter/CrashReportReceiver/UploadsInProgress.cs
+++ b/Engine/Source/Programs/CrashReporter/CrashReportReceiver/UploadsInProgress.cs
@@ -13,7 +13,7 @@ namespace Tools.CrashReporter.CrashReportReceiver
 	/// Mainly to cope with multiple clients fighting to upload the same report. This shouldn't happen, but the service
 	/// should be resilient against it even so.
 	/// </remarks>
-	public class FUploadsInProgress
+	public sealed class FUploadsInProgress : IDisposable
 	{
 		/// <summary> Background thread used to monitor CPU/network usage. </summary>
 		private Thread PerformancePooler;
@@ -66,6 +66,16 @@ namespace Tools.CrashReporter.CrashReportReceiver
 			This.PoolPerformanceData();
 		}
 
+		/// <summary>Dispose of any resources owned by this instance.</summary>
+		public void Dispose()
+		{
+			if (PerfCounterProcessorTime != null)
+			{
+				PerfCounterProcessorTime.Dispose();
+				PerfCounterProcessorTime = null;
+			}
+		}
+
 		private void PoolPerformanceData()
 		{
 			while( true )
diff --git a/Engine/Source/Programs/CrashReporter/CrashReportReceiver/WebHandler.cs b/Engine/Source/Programs/CrashReporter/CrashReportReceiver/WebHandler.cs
--- a/Engine/Source/Programs/CrashReporter/CrashReportReceiver/WebHandler.cs
+++ b/Engine/Source/Programs/CrashReporter/CrashReportReceiver/WebHandler.cs
@@ -67,6 +67,11 @@ namespace Tools.CrashReporter.CrashReportReceiver
 		/// <param name="Disposing">true if the Dispose call is from user code, and not system code.</param>
 		protected virtual void Dispose( bool Disposing )
 		{
+			if (UploadsInProgress != null)
+			{
+				UploadsInProgress.Dispose();
+				UploadsInProgress = null;
+			}
 			if (ServiceHttpListener != null)
 			{
 				ServiceHttpListener.Close();
diff --git a/Engine/Source/Programs/CrashReporter/CrashReportWebSite/App_Start/UnityMvcActivator.cs b/Engine/Source/Programs/CrashReporter/CrashReportWebSite/App_Start/UnityMvcActivator.cs
--- a/Engine/Source/Programs/CrashReporter/CrashReportWebSite/App_Start/UnityMvcActivator.cs
+++ b/Engine/Source/Programs/CrashReporter/CrashReportWebSite/App_Start/UnityMvcActivator.cs
@@ -23,7 +23,6 @@ namespace Tools.CrashReporter.CrashReportWebSite.App_Start
             DependencyResolver.SetResolver(new UnityDependencyResolver(container));
 
             // TODO: Uncomment if you want to use PerRequestLifetimeManager
-            // Microsoft.Web.Infrastructure.DynamicModuleHelper.DynamicModuleUtility.RegisterModule(typeof(UnityPerRequestHttpModule));
         }
 
         /// <summary>Disposes the Unity container when the application is shut down.</summary>
diff --git a/Engine/Source/Programs/CrashReporter/CrashReportWebSite/Controllers/BuggsController.cs b/Engine/Source/Programs/CrashReporter/CrashReportWebSite/Controllers/BuggsController.cs
--- a/Engine/Source/Programs/CrashReporter/CrashReportWebSite/Controllers/BuggsController.cs
+++ b/Engine/Source/Programs/CrashReporter/CrashReportWebSite/Controllers/BuggsController.cs
@@ -399,7 +399,7 @@ namespace Tools.CrashReporter.CrashReportWebSite.Controllers
             foreach (var bugg in sortedResultsList)
             {
                 var crash =
-                    _unitOfWork.CrashRepository.First(data => data.Buggs.Any(bg => bg.Id == bugg.Id));
+                    _unitOfWork.CrashRepository.First(data => data.BuggId == bugg.Id);
                 bugg.FunctionCalls = crash.GetCallStack().GetFunctionCalls();
             }
 
diff --git a/Engine/Source/Programs/CrashReporter/CrashReportWebSite/Controllers/CSVController.cs b/Engine/Source/Programs/CrashReporter/CrashReportWebSite/Controllers/CSVController.cs
--- a/Engine/Source/Programs/CrashReporter/CrashReportWebSite/Controllers/CSVController.cs
+++ b/Engine/Source/Programs/CrashReporter/CrashReportWebSite/Controllers/CSVController.cs
@@ -151,7 +151,7 @@ namespace Tools.CrashReporter.CrashReportWebSite.Controllers
                             Branch = c.Branch,
                             CrashType = c.CrashType,
                             EpicId = c.EpicAccountId,
-                            BuggId = c.Buggs.First().Id
+                            BuggId = c.Bugg.Id
                         }
                     );
 
diff --git a/Engine/Source/Programs/CrashReporter/CrashReportWebSite/Controllers/CrashesController.cs b/Engine/Source/Programs/CrashReporter/CrashReportWebSite/Controllers/CrashesController.cs
--- a/Engine/Source/Programs/CrashReporter/CrashReportWebSite/Controllers/CrashesController.cs
+++ b/Engine/Source/Programs/CrashReporter/CrashReportWebSite/Controllers/CrashesController.cs
@@ -143,6 +143,8 @@ namespace Tools.CrashReporter.CrashReportWebSite.Controllers
 				if( !string.IsNullOrEmpty( FormValue ) )
 				{
 					currentCrash.Jira = FormValue;
+                    if(currentCrash.Bugg != null)
+                        currentCrash.Bugg.TTPID = FormValue;
 				}
 
 				// Valid to set description to an empty string
@@ -152,6 +154,9 @@ namespace Tools.CrashReporter.CrashReportWebSite.Controllers
 					currentCrash.Description = FormValue;
 				}
 
+                _unitOfWork.CrashRepository.Update(currentCrash);
+                _unitOfWork.Save();
+
 				currentCallStack = new CallStackContainer( currentCrash );
 			    currentCrash.Module = currentCallStack.GetModuleName();
 
@@ -893,10 +898,11 @@ namespace Tools.CrashReporter.CrashReportWebSite.Controllers
 
 	                //if a bugg exists update this crash from the bugg
 	                //buggs are authoritative in this case
-	                newCrash.Buggs.Add(bugg);
+                    newCrash.BuggId = bugg.Id;
 	                newCrash.Jira = bugg.TTPID;
 	                newCrash.FixedChangeList = bugg.FixedChangeList;
 	                newCrash.Status = bugg.Status;
+                    newCrash.Buggs.Add(bugg);
 	                _unitOfWork.CrashRepository.Save(newCrash);
 	                _unitOfWork.Save();
 	            }
@@ -917,6 +923,7 @@ namespace Tools.CrashReporter.CrashReportWebSite.Controllers
 	                bugg.Status = newCrash.Status;
                     bugg.FixedChangeList = newCrash.FixedChangeList;
 	                bugg.ErrorMessageId = errorMessage.Id;
+                    newCrash.Bugg = bugg;
                     newCrash.Buggs.Add(bugg);
                     _unitOfWork.BuggRepository.Save(bugg);
                     _unitOfWork.CrashRepository.Save(newCrash);
@@ -1037,5 +1044,22 @@ namespace Tools.CrashReporter.CrashReportWebSite.Controllers
 
 	        base.Dispose(disposing);
 	    }
+
+	    public void TestAddCrash()
+	    {
+	        string crashContext;
+	        using (
+	            var stream =
+	                new FileStream(
+	                    "D:/CR/Engine/Source/Programs/CrashReporter/CrashReportWebSite/bin/Content/CrashContext.txt",
+	                    FileMode.Open))
+	        {
+	            crashContext = new StreamReader(stream).ReadToEnd();
+	        }
+
+            var crash = XmlHandler.FromXmlString<CrashDescription>(crashContext);
+
+	        CreateCrash(crash);
+	    }
 	}
 }
\ No newline at end of file
diff --git a/Engine/Source/Programs/CrashReporter/CrashReportWebSite/Controllers/ErrorController.cs b/Engine/Source/Programs/CrashReporter/CrashReportWebSite/Controllers/ErrorController.cs
--- a/Engine/Source/Programs/CrashReporter/CrashReportWebSite/Controllers/ErrorController.cs
+++ b/Engine/Source/Programs/CrashReporter/CrashReportWebSite/Controllers/ErrorController.cs
@@ -1,9 +1,5 @@
 ﻿// Copyright 1998-2017 Epic Games, Inc. All Rights Reserved.
 
-using System;
-using System.Collections.Generic;
-using System.Linq;
-using System.Web;
 using System.Web.Mvc;
 
 namespace Tools.CrashReporter.CrashReportWebSite.Controllers
diff --git a/Engine/Source/Programs/CrashReporter/CrashReportWebSite/Controllers/ReportsController.cs b/Engine/Source/Programs/CrashReporter/CrashReportWebSite/Controllers/ReportsController.cs
--- a/Engine/Source/Programs/CrashReporter/CrashReportWebSite/Controllers/ReportsController.cs
+++ b/Engine/Source/Programs/CrashReporter/CrashReportWebSite/Controllers/ReportsController.cs
@@ -237,7 +237,7 @@ namespace Tools.CrashReporter.CrashReportWebSite.Controllers
         /// Retrieve all Buggs matching the search criteria.
 	    /// </summary>
 	    /// <returns></returns>
-        private ReportsViewModel GetResults(string branchName, DateTime startDate, DateTime endDate, int BuggIDToBeAddedToJira)
+        private ReportsViewModel GetResults(string branchName, DateTime startDate, DateTime endDate, int BuggIDToBeAddedToJira, bool sigbool)
 	    {
 	        // It would be great to have a CSV export of this as well with buggs ID being the key I can then use to join them :)
 	        // 
@@ -247,37 +247,31 @@ namespace Tools.CrashReporter.CrashReportWebSite.Controllers
 	        var JiraComponents = JC.GetNameToComponents();
 	        var JiraVersions = JC.GetNameToVersions();
 
-	        using (FAutoScopedLogTimer LogTimer = new FAutoScopedLogTimer(this.GetType().ToString()))
+	        using (var logTimer = new FAutoScopedLogTimer(this.GetType().ToString()))
 	        {
-	            string AnonumousGroup = "Anonymous";
-	            //List<String> Users = FRepository.Get().GetUserNamesFromGroupName( AnonumousGroup );
-	            int AnonymousGroupID = _unitOfWork.UserGroupRepository.First(data => data.Name == "Anonymous").Id;
-	            List<int> AnonumousIDs =
-                    _unitOfWork.UserGroupRepository.First(data => data.Name == "Anonymous").Users.Select(data => data.Id).ToList();
-	            int AnonymousID = AnonumousIDs.First();
-
-	            var CrashesQuery = _unitOfWork.CrashRepository.ListAll().Where(data => data.TimeOfCrash > startDate && data.TimeOfCrash <= endDate)
+	            const string anonymousGroup = "Anonymous";
+                var anonymousGroupId = _unitOfWork.UserGroupRepository.First(data => data.Name == anonymousGroup).Id;
+	            var anonymousIDs =
+                    _unitOfWork.UserGroupRepository.First(data => data.Id == anonymousGroupId).Users.Select(data => data.Id).ToList();
+	            var anonymousId = anonymousIDs.First();
+
+	            var crashesQuery = _unitOfWork.CrashRepository.ListAll().Where(data => data.TimeOfCrash > startDate && data.TimeOfCrash <= endDate)
 	                // Only crashes and asserts
-	                .Where(Crash => Crash.CrashType == 1 || Crash.CrashType == 2)
+	                .Where(crash => crash.CrashType == 1 || crash.CrashType == 2)
 	                // Only anonymous user
-	                .Where(Crash => Crash.UserNameId.Value == AnonymousID);
+	                .Where(crash => crash.UserNameId.Value == anonymousId);
 
 	            // Filter by BranchName
 	            if (!string.IsNullOrEmpty(branchName))
 	            {
-	                CrashesQuery =
-	                    (
-	                        from CrashDetail in CrashesQuery
-	                        where CrashDetail.Branch.Equals(branchName)
-	                        select CrashDetail
-	                        );
+	                crashesQuery = crashesQuery.Where(data => data.Branch == branchName);
 	            }
 
-	            var crashes = CrashesQuery.Select(data => new
+	            var crashes = crashesQuery.Select(data => new
 	            {
                     ID = data.Id,
                     TimeOfCrash = data.TimeOfCrash.Value,
-	                //UserID = Crash.UserNameId.Value, 
+	                UserID = data.UserNameId,
                     BuildVersion = data.BuildVersion,
                     JIRA = data.Jira,
                     Platform = data.PlatformName,
@@ -289,120 +283,90 @@ namespace Tools.CrashReporter.CrashReportWebSite.Controllers
                     Description = data.Description,
                     RawCallStack = data.RawCallStack,
 	            }).ToList();
-                var numCrashes = CrashesQuery.Count();
-
-	            /*
-                // Build patterns for crashes where patters is null.
-                var CrashesWithoutPattern = FRepository.Get().Crashes
-                    .FilterByDate( FRepository.Get().Crashes.ListAll(), startDate, endDate.AddDays( 1 ) )
-                    // Only crashes and asserts
-                    .Where( Crash => Crash.Pattern == null || Crash.Pattern == "" )
-                    .Select( Crash => Crash )
-                    .ToList();
-
-                foreach( var Crash in CrashesWithoutPattern )
-                {
-                    Crash.BuildPattern();
-                }
-                */
+                var numCrashes = crashesQuery.Count();
 
 	            // Total # of ALL (Anonymous) crashes in timeframe
-	            int TotalAnonymousCrashes = numCrashes;
+	            var totalAnonymousCrashes = numCrashes;
 
 	            // Total # of UNIQUE (Anonymous) crashes in timeframe
-	            HashSet<string> UniquePatterns = new HashSet<string>();
-	            HashSet<string> UniqueMachines = new HashSet<string>();
-	            Dictionary<string, int> PatternToCount = new Dictionary<string, int>();
-
-	            //List<int> CrashesWithoutPattern = new List<int>();
-	            //List<DateTime> CrashesWithoutPatternDT = new List<DateTime>();
-
-	            foreach (var Crash in crashes)
+                var uniquePatterns = new HashSet<string>();
+                var uniqueMachines = new HashSet<string>();
+	            var patternToCount = new Dictionary<string, int>();
+                
+	            foreach (var crash in crashes)
 	            {
-	                if (string.IsNullOrEmpty(Crash.Pattern))
-	                {
-	                    //CrashesWithoutPattern.Add( Crash.ID );
-	                    //CrashesWithoutPatternDT.Add( Crash.TimeOfCrash );
-	                    continue;
-	                }
-
-	                UniquePatterns.Add(Crash.Pattern);
-	                UniqueMachines.Add(Crash.MachineID);
+	                uniquePatterns.Add(crash.Pattern);
+	                uniqueMachines.Add(crash.MachineID);
 
-	                bool bAdd = !PatternToCount.ContainsKey(Crash.Pattern);
+	                var bAdd = !patternToCount.ContainsKey(crash.Pattern);
 	                if (bAdd)
 	                {
-	                    PatternToCount.Add(Crash.Pattern, 1);
+	                    patternToCount.Add(crash.Pattern, 1);
 	                }
 	                else
 	                {
-	                    PatternToCount[Crash.Pattern]++;
+	                    patternToCount[crash.Pattern]++;
 	                }
 	            }
-	            var PatternToCountOrdered = PatternToCount.OrderByDescending(X => X.Value).ToList();
+	            var PatternToCountOrdered = patternToCount.OrderByDescending(X => X.Value).ToList();
 	            const int NumTopRecords = 200;
-	            var PatternAndCount = PatternToCountOrdered.Take(NumTopRecords).ToDictionary(X => X.Key, Y => Y.Value);
+	            var PatternAndCount = PatternToCountOrdered.Take(NumTopRecords).ToDictionary(x => x.Key, y => y.Value);
 
-	            int TotalUniqueAnonymousCrashes = UniquePatterns.Count;
+	            int TotalUniqueAnonymousCrashes = uniquePatterns.Count;
 
 	            // Total # of AFFECTED USERS (Anonymous) in timeframe
-	            int TotalAffectedUsers = UniqueMachines.Count;
+	            int TotalAffectedUsers = uniqueMachines.Count;
 
 	            var RealBuggs = _unitOfWork.BuggRepository.ListAll().Where(Bugg => PatternAndCount.Keys.Contains(Bugg.Pattern));
 
 	            // Build search string.
-	            HashSet<string> FoundJiras = new HashSet<string>();
-	            Dictionary<string, List<Bugg>> JiraIDtoBugg = new Dictionary<string, List<Bugg>>();
+	            var foundJiras = new HashSet<string>();
+	            var jiraIDtoBugg = new Dictionary<string, List<Bugg>>();
 
-	            List<Bugg> Buggs = new List<Bugg>(NumTopRecords);
+	            var buggs = new List<Bugg>(NumTopRecords);
 	            foreach (var Top in PatternAndCount)
 	            {
-	                Bugg NewBugg = RealBuggs.FirstOrDefault(X => X.Pattern == Top.Key);
-	                if (NewBugg != null)
+	                var newBugg = RealBuggs.FirstOrDefault(X => X.Pattern == Top.Key);
+	                if (newBugg != null)
 	                {
-	                    using (
-	                        FAutoScopedLogTimer TopTimer =
-	                            new FAutoScopedLogTimer(string.Format("{0}:{1}", Buggs.Count + 1, NewBugg.Id)))
-	                    {
-	                        var CrashesForBugg = crashes.Where(Crash => Crash.Pattern == Top.Key).ToList();
+                        var crashesForBugg = crashes.Where(Crash => Crash.Pattern == Top.Key).ToList();
 
-	                        // Convert anonymous to full type;
-	                        var FullCrashesForBugg = new List<Crash>(CrashesForBugg.Count);
-	                        foreach (var Anon in CrashesForBugg)
+	                    // Convert anonymous to full type;
+	                    var fullCrashesForBugg = new List<Crash>(crashesForBugg.Count);
+	                    foreach (var anon in crashesForBugg)
+	                    {
+	                        fullCrashesForBugg.Add(new Crash()
 	                        {
-	                            FullCrashesForBugg.Add(new Crash()
-	                            {
-	                                Id = Anon.ID,
-	                                TimeOfCrash = Anon.TimeOfCrash,
-	                                BuildVersion = Anon.BuildVersion,
-	                                Jira = Anon.JIRA,
-	                                PlatformName = Anon.Platform,
-	                                FixedChangeList = Anon.FixCL,
-	                                ChangeListVersion = Anon.BuiltFromCL,
-	                                Pattern = Anon.Pattern,
-	                                ComputerName = Anon.MachineID,
-	                                Branch = Anon.Branch,
-	                                Description = Anon.Description,
-	                                RawCallStack = Anon.RawCallStack,
-	                            });
-	                        }
+	                            Id = anon.ID,
+	                            TimeOfCrash = anon.TimeOfCrash,
+	                            BuildVersion = anon.BuildVersion,
+	                            Jira = anon.JIRA,
+	                            PlatformName = anon.Platform,
+	                            FixedChangeList = anon.FixCL,
+	                            ChangeListVersion = anon.BuiltFromCL,
+	                            Pattern = anon.Pattern,
+	                            ComputerName = anon.MachineID,
+	                            Branch = anon.Branch,
+	                            Description = anon.Description,
+	                            RawCallStack = anon.RawCallStack,
+	                        });
+	                    }
 
-	                        NewBugg.PrepareBuggForJira(FullCrashesForBugg);
+	                    newBugg.PrepareBuggForJira(fullCrashesForBugg);
 
-	                        // Verify valid JiraID, this may be still a TTP 
-	                        if (!string.IsNullOrEmpty(NewBugg.TTPID))
-	                        {
-	                            int TTPID = 0;
-	                            int.TryParse(NewBugg.TTPID, out TTPID);
+	                    // Verify valid JiraID, this may be still a TTP 
+	                    if (!string.IsNullOrEmpty(newBugg.TTPID))
+	                    {
+	                        int ttpid = 0;
+	                        int.TryParse(newBugg.TTPID, out ttpid);
 
-	                            if (TTPID == 0)
-	                            {
-	                                AddBuggJiraMapping(NewBugg, ref FoundJiras, ref JiraIDtoBugg);
-	                            }
+	                        if (ttpid == 0)
+	                        {
+	                            AddBuggJiraMapping(newBugg, ref foundJiras, ref jiraIDtoBugg);
 	                        }
-
-	                        Buggs.Add(NewBugg);
 	                    }
+
+	                    buggs.Add(newBugg);
 	                }
 	                else
 	                {
@@ -412,47 +376,71 @@ namespace Tools.CrashReporter.CrashReportWebSite.Controllers
 
 	            if (BuggIDToBeAddedToJira > 0)
 	            {
-                    var Bugg = Buggs.Where(X => X.Id == BuggIDToBeAddedToJira).FirstOrDefault();
-                    if (Bugg != null)
+                    var bugg = buggs.FirstOrDefault(x => x.Id == BuggIDToBeAddedToJira);
+                    if (bugg != null)
                     {
-                        Bugg.CopyToJira();
-                        AddBuggJiraMapping(Bugg, ref FoundJiras, ref JiraIDtoBugg);
+                        bugg.CopyToJira();
+                        AddBuggJiraMapping(bugg, ref foundJiras, ref jiraIDtoBugg);
                     }
 	            }
 
 	            if (JC.CanBeUsed())
 	            {
-	                var BuggsCopy = new List<Bugg>(Buggs);
+	                var buggsCopy = new List<Bugg>(buggs);
 
-	                HashSet<string> InvalidJiras = new HashSet<string>();
+	                var invalidJiras = new HashSet<string>();
 
 	                // Simple verification of JIRA
-	                foreach (var Value in FoundJiras)
+	                foreach (var Value in foundJiras)
 	                {
 	                    if (Value.Length < 3 || !Value.Contains('-'))
 	                    {
-	                        InvalidJiras.Add(Value);
+	                        invalidJiras.Add(Value);
 	                    }
 	                }
 
-	                foreach (var InvalidJira in InvalidJiras)
+	                foreach (var invalidJira in invalidJiras)
 	                {
-	                    FoundJiras.Remove(InvalidJira);
+	                    foundJiras.Remove(invalidJira);
 	                }
 
 	                // Grab the data form JIRA.
-	                string JiraSearchQuery = string.Join(" OR ", FoundJiras);
+	                string JiraSearchQuery = string.Join(" OR ", foundJiras);
+                    
+                    bool bInvalid = false;
+                    var jiraResults = new Dictionary<string, Dictionary<string, object>>();
+	                try
+	                {
+	                    if (!string.IsNullOrWhiteSpace(JiraSearchQuery))
+	                    {
+	                        jiraResults = JC.SearchJiraTickets(
+	                            JiraSearchQuery,
+	                            new string[]
+	                            {
+	                                "key", // string
+	                                "summary", // string
+	                                "components", // System.Collections.ArrayList, Dictionary<string,object>, name
+	                                "resolution",
+	                                //System.Collections.Generic.Dictionary`2[System.String,System.Object], name
+	                                "fixVersions", // System.Collections.ArrayList, Dictionary<string,object>, name
+	                                "customfield_11200" // string
+	                            });
+	                    }
+	                }
+	                catch (System.Exception)
+	                {
+	                    bInvalid = true;
+	                }
 
-	                using (FAutoScopedLogTimer JiraResultsTimer = new FAutoScopedLogTimer("JiraResults"))
+	                // Invalid records have been found, find the broken using the slow path.
+	                if (bInvalid)
 	                {
-	                    bool bInvalid = false;
-	                    var JiraResults = new Dictionary<string, Dictionary<string, object>>();
-	                    try
+	                    foreach (var Query in foundJiras)
 	                    {
-	                        if (!string.IsNullOrWhiteSpace(JiraSearchQuery))
+	                        try
 	                        {
-	                            JiraResults = JC.SearchJiraTickets(
-	                                JiraSearchQuery,
+	                            var TempResult = JC.SearchJiraTickets(
+	                                Query,
 	                                new string[]
 	                                {
 	                                    "key", // string
@@ -463,105 +451,78 @@ namespace Tools.CrashReporter.CrashReportWebSite.Controllers
 	                                    "fixVersions", // System.Collections.ArrayList, Dictionary<string,object>, name
 	                                    "customfield_11200" // string
 	                                });
-	                        }
-	                    }
-	                    catch (System.Exception)
-	                    {
-	                        bInvalid = true;
-	                    }
 
-	                    // Invalid records have been found, find the broken using the slow path.
-	                    if (bInvalid)
-	                    {
-	                        foreach (var Query in FoundJiras)
-	                        {
-	                            try
+	                            foreach (var Temp in TempResult)
 	                            {
-	                                var TempResult = JC.SearchJiraTickets(
-	                                    Query,
-	                                    new string[]
-	                                    {
-	                                        "key", // string
-	                                        "summary", // string
-	                                        "components", // System.Collections.ArrayList, Dictionary<string,object>, name
-	                                        "resolution",
-	                                        // System.Collections.Generic.Dictionary`2[System.String,System.Object], name
-	                                        "fixVersions", // System.Collections.ArrayList, Dictionary<string,object>, name
-	                                        "customfield_11200" // string
-	                                    });
-
-	                                foreach (var Temp in TempResult)
-	                                {
-	                                    JiraResults.Add(Temp.Key, Temp.Value);
-	                                }
+	                                jiraResults.Add(Temp.Key, Temp.Value);
 	                            }
-	                            catch (System.Exception)
-	                            {
+	                        }
+	                        catch (System.Exception)
+	                        {
 
-	                            }
 	                        }
 	                    }
+	                }
 
-	                    // Jira Key, Summary, Components, Resolution, Fix version, Fix changelist
-	                    foreach (var Jira in JiraResults)
-	                    {
-	                        string JiraID = Jira.Key;
+	                // Jira Key, Summary, Components, Resolution, Fix version, Fix changelist
+	                foreach (var Jira in jiraResults)
+	                {
+	                    string JiraID = Jira.Key;
 
-	                        string Summary = (string) Jira.Value["summary"];
+	                    string Summary = (string) Jira.Value["summary"];
 
-	                        string ComponentsText = "";
-	                        System.Collections.ArrayList Components =
-	                            (System.Collections.ArrayList) Jira.Value["components"];
-	                        foreach (Dictionary<string, object> Component in Components)
-	                        {
-	                            ComponentsText += (string) Component["name"];
-	                            ComponentsText += " ";
-	                        }
+	                    string ComponentsText = "";
+	                    System.Collections.ArrayList Components =
+	                        (System.Collections.ArrayList) Jira.Value["components"];
+	                    foreach (Dictionary<string, object> Component in Components)
+	                    {
+	                        ComponentsText += (string) Component["name"];
+	                        ComponentsText += " ";
+	                    }
 
-	                        Dictionary<string, object> ResolutionFields =
-	                            (Dictionary<string, object>) Jira.Value["resolution"];
-	                        string Resolution = ResolutionFields != null ? (string) ResolutionFields["name"] : "";
+	                    Dictionary<string, object> ResolutionFields =
+	                        (Dictionary<string, object>) Jira.Value["resolution"];
+	                    string Resolution = ResolutionFields != null ? (string) ResolutionFields["name"] : "";
 
-	                        string FixVersionsText = "";
-	                        System.Collections.ArrayList FixVersions =
-	                            (System.Collections.ArrayList) Jira.Value["fixVersions"];
-	                        foreach (Dictionary<string, object> FixVersion in FixVersions)
-	                        {
-	                            FixVersionsText += (string) FixVersion["name"];
-	                            FixVersionsText += " ";
-	                        }
+	                    string FixVersionsText = "";
+	                    System.Collections.ArrayList FixVersions =
+	                        (System.Collections.ArrayList) Jira.Value["fixVersions"];
+	                    foreach (Dictionary<string, object> FixVersion in FixVersions)
+	                    {
+	                        FixVersionsText += (string) FixVersion["name"];
+	                        FixVersionsText += " ";
+	                    }
 
-	                        int FixCL = Jira.Value["customfield_11200"] != null
-	                            ? (int) (decimal) Jira.Value["customfield_11200"]
-	                            : 0;
+	                    int FixCL = Jira.Value["customfield_11200"] != null
+	                        ? (int) (decimal) Jira.Value["customfield_11200"]
+	                        : 0;
 
-	                        List<Bugg> BuggsForJira;
-	                        JiraIDtoBugg.TryGetValue(JiraID, out BuggsForJira);
+	                    List<Bugg> BuggsForJira;
+	                    jiraIDtoBugg.TryGetValue(JiraID, out BuggsForJira);
 
-	                        //var BuggsForJira = JiraIDtoBugg[JiraID];
+	                    //var BuggsForJira = JiraIDtoBugg[JiraID];
 
-	                        if (BuggsForJira != null)
+	                    if (BuggsForJira != null)
+	                    {
+	                        foreach (Bugg Bugg in BuggsForJira)
 	                        {
-	                            foreach (Bugg Bugg in BuggsForJira)
+	                            Bugg.JiraSummary = Summary;
+	                            Bugg.JiraComponentsText = ComponentsText;
+	                            Bugg.JiraResolution = Resolution;
+	                            Bugg.JiraFixVersionsText = FixVersionsText;
+	                            if (FixCL != 0)
 	                            {
-	                                Bugg.JiraSummary = Summary;
-	                                Bugg.JiraComponentsText = ComponentsText;
-	                                Bugg.JiraResolution = Resolution;
-	                                Bugg.JiraFixVersionsText = FixVersionsText;
-	                                if (FixCL != 0)
-	                                {
-	                                    Bugg.JiraFixCL = FixCL.ToString();
-	                                }
-
-	                                BuggsCopy.Remove(Bugg);
+	                                Bugg.JiraFixCL = FixCL.ToString();
 	                            }
+
+	                            buggsCopy.Remove(Bugg);
 	                        }
 	                    }
 	                }
 
 	                // If there are buggs, we need to update the summary to indicate an error.
 	                // Usually caused when bugg's project has changed.
-	                foreach (var Bugg in BuggsCopy.Where(b => !string.IsNullOrEmpty(b.TTPID)))
+	                foreach (var Bugg in buggsCopy.Where(b => !string.IsNullOrEmpty(b.TTPID)))
 	                {
 	                    Bugg.JiraSummary = "JIRA MISMATCH";
 	                    Bugg.JiraComponentsText = "JIRA MISMATCH";
@@ -571,22 +532,186 @@ namespace Tools.CrashReporter.CrashReportWebSite.Controllers
 	                }
 	            }
 
-	            Buggs = Buggs.OrderByDescending(b => b.CrashesInTimeFrameGroup).ToList();
+	            buggs = buggs.OrderByDescending(b => b.CrashesInTimeFrameGroup).ToList();
 
 	            return new ReportsViewModel
 	            {
-	                Buggs = Buggs,
+	                Buggs = buggs,
                     BranchName = branchName,
 	                BranchNames = _unitOfWork.CrashRepository.GetBranchesAsListItems(),
 	                DateFrom = (long) (startDate - CrashesViewModel.Epoch).TotalMilliseconds,
 	                DateTo = (long) (endDate - CrashesViewModel.Epoch).TotalMilliseconds,
 	                TotalAffectedUsers = TotalAffectedUsers,
-	                TotalAnonymousCrashes = TotalAnonymousCrashes,
+	                TotalAnonymousCrashes = totalAnonymousCrashes,
 	                TotalUniqueAnonymousCrashes = TotalUniqueAnonymousCrashes
 	            };
 	        }
 	    }
 
+        /// <summary>
+        /// 
+        /// </summary>
+        /// <param name="branchName"></param>
+        /// <param name="startDate"></param>
+        /// <param name="endDate"></param>
+        /// <returns></returns>
+        private ReportsViewModel GetResults(string branchName, DateTime startDate, DateTime endDate, int BuggIDToBeAddedToJira)
+        {
+            const string anonymousGroup = "Anonymous";
+            var anonymousGroupId = _unitOfWork.UserGroupRepository.First(data => data.Name == anonymousGroup).Id;
+            var anonymousIDs =
+                _unitOfWork.UserGroupRepository.First(data => data.Id == anonymousGroupId).Users.Select(data => data.Id).ToList();
+            var anonymousId = anonymousIDs.First();
+
+            var crashesQuery = _unitOfWork.CrashRepository.ListAll()
+                .Where(data => data.TimeOfCrash > startDate && data.TimeOfCrash <= endDate)
+                // Only crashes and asserts
+                .Where(crash => crash.CrashType == 1 || crash.CrashType == 2);
+                // Only anonymous user
+                    //.Where(crash => crash.UserNameId.Value == anonymousId);
+
+            // Filter by BranchName
+            if (!string.IsNullOrEmpty(branchName))
+            {
+                crashesQuery = crashesQuery.Where(data => data.Branch == branchName);
+            }
+
+            var crashesList = crashesQuery.ToList();
+
+            var buggsList = crashesQuery.Where(data => data.Bugg != null).GroupBy(data => data.Bugg).Select(data => data.Key).ToList();
+
+            foreach (var bugg in buggsList)
+            {
+                if (bugg == null)
+                    continue;
+
+                var crashes = crashesList.Where(data => data.BuggId == bugg.Id).ToList();
+
+                bugg.PrepareBuggForJira(crashes);
+
+                bugg.CrashesInTimeFrameGroup = crashesList.Count(data => data.BuggId == bugg.Id);
+                bugg.AffectedVersions = new SortedSet<string>(crashes.Select(data => data.BuildVersion));
+                bugg.NumberOfUniqueMachines = crashes.GroupBy(data => data.ComputerName).Count();
+
+                //get latest changelist
+                int latestClAffected;
+                int.TryParse(crashes.Max(data => data.ChangeListVersion), out latestClAffected);
+                bugg.LatestCLAffected = latestClAffected;
+
+                if (string.IsNullOrWhiteSpace(bugg.TTPID))
+                {
+                    bugg.JiraSummary = "JIRA MISMATCH";
+                    bugg.JiraComponentsText = "JIRA MISMATCH";
+                    bugg.JiraResolution = "JIRA MISMATCH";
+                    bugg.JiraFixVersionsText = "JIRA MISMATCH";
+                    bugg.JiraFixCL = "JIRA MISMATCH";
+                }
+            }
+
+            //Get jira connection
+            var jiraConnection = JiraConnection.Get();
+
+            if (BuggIDToBeAddedToJira > 0)
+            {
+                var bugg = buggsList.FirstOrDefault(x => x.Id == BuggIDToBeAddedToJira);
+                if (bugg != null)
+                {
+                    bugg.CopyToJira();
+                }
+            }
+
+            var jiraResults = new Dictionary<string, Dictionary<string, object>>();
+
+            foreach (var Query in buggsList.Where(data =>
+                !string.IsNullOrWhiteSpace(data.TTPID)).Select(data => data.TTPID))
+            {
+                try
+                {
+                    var TempResult = jiraConnection.SearchJiraTickets(
+                        "key = " + Query,
+                        new string[]
+	                    {
+	                        "key", // string
+	                        "summary", // string
+	                        "components", // System.Collections.ArrayList, Dictionary<string,object>, name
+	                        "resolution",
+	                        // System.Collections.Generic.Dictionary`2[System.String,System.Object], name
+	                        "fixVersions", // System.Collections.ArrayList, Dictionary<string,object>, name
+	                        "customfield_11200" // string
+	                    });
+
+                    foreach (var Temp in TempResult)
+                    {
+                        jiraResults.Add(Temp.Key, Temp.Value);
+                    }
+                }
+                catch (System.Exception)
+                {
+
+                }
+            }
+
+            // Jira Key, Summary, Components, Resolution, Fix version, Fix changelist
+            foreach (var Jira in jiraResults)
+            {
+                var jiraId = Jira.Key;
+
+                var summary = (string)Jira.Value["summary"];
+
+                var ComponentsText = "";
+                System.Collections.ArrayList Components =
+                    (System.Collections.ArrayList)Jira.Value["components"];
+                foreach (Dictionary<string, object> Component in Components)
+                {
+                    ComponentsText += (string)Component["name"];
+                    ComponentsText += " ";
+                }
+
+                Dictionary<string, object> ResolutionFields =
+                    (Dictionary<string, object>)Jira.Value["resolution"];
+                string Resolution = ResolutionFields != null ? (string)ResolutionFields["name"] : "";
+
+                string FixVersionsText = "";
+                System.Collections.ArrayList FixVersions =
+                    (System.Collections.ArrayList)Jira.Value["fixVersions"];
+                foreach (Dictionary<string, object> FixVersion in FixVersions)
+                {
+                    FixVersionsText += (string)FixVersion["name"];
+                    FixVersionsText += " ";
+                }
+
+                int FixCL = Jira.Value["customfield_11200"] != null
+                    ? (int)(decimal)Jira.Value["customfield_11200"]
+                    : 0;
+
+
+                if (buggsList.Any(data => data.TTPID == jiraId))
+                {
+                    var bugg = buggsList.First(data => data.TTPID == jiraId);
+                    bugg.JiraSummary = summary;
+                    bugg.JiraComponentsText = ComponentsText;
+                    bugg.JiraResolution = Resolution;
+                    bugg.JiraFixVersionsText = FixVersionsText;
+                    if (FixCL != 0)
+                    {
+                        bugg.JiraFixCL = FixCL.ToString();
+                    }
+                }
+            }
+
+            return new ReportsViewModel
+            {
+                Buggs = buggsList.OrderByDescending(data => data.CrashesInTimeFrameGroup).ToList(),
+                BranchName = branchName,
+                BranchNames = _unitOfWork.CrashRepository.GetBranchesAsListItems(),
+                DateFrom = (long)(startDate - CrashesViewModel.Epoch).TotalMilliseconds,
+                DateTo = (long)(endDate - CrashesViewModel.Epoch).TotalMilliseconds,
+                TotalAffectedUsers = crashesQuery.GroupBy(data => data.ComputerName).Count(),
+                TotalAnonymousCrashes = crashesQuery.Count(),
+                TotalUniqueAnonymousCrashes = crashesQuery.GroupBy(data => data.BuggId).Count()
+            };
+        }
+
         /// <summary>
         /// Private method to get reports page for a branch and render to an HTML string
         /// </summary>
@@ -629,6 +754,23 @@ namespace Tools.CrashReporter.CrashReportWebSite.Controllers
             base.Dispose(disposing);
         }
 
+        /// <summary>
+        /// Get a list of bugs within a specified date range
+        /// </summary>
+        /// <param name="results"></param>
+        /// <param name="dateFrom"></param>
+        /// <param name="dateTo"></param>
+        /// <returns></returns>
+        public IQueryable<Bugg> FilterByDate(IQueryable<Bugg> results, DateTime dateFrom, DateTime dateTo)
+        {
+            dateTo = dateTo.AddDays(1);
+            var buggsInTimeFrame =
+                results.Where(bugg => (bugg.TimeOfFirstCrash >= dateFrom && bugg.TimeOfFirstCrash <= dateTo) ||
+                    (bugg.TimeOfLastCrash >= dateFrom && bugg.TimeOfLastCrash <= dateTo) || (bugg.TimeOfFirstCrash <= dateFrom && bugg.TimeOfLastCrash >= dateTo));
+
+            return buggsInTimeFrame;
+        }
+
         #endregion
     }
 
