Copying //UE4/Dev-Core to //UE4/Dev-Main (Source: //UE4/Dev-Core @ 3228803)

diff --git a/Engine/Source/Programs/CrashReporter/CrashReportWebSite/CrashReportWebSite.csproj b/Engine/Source/Programs/CrashReporter/CrashReportWebSite/CrashReportWebSite.csproj
--- a/Engine/Source/Programs/CrashReporter/CrashReportWebSite/CrashReportWebSite.csproj
+++ b/Engine/Source/Programs/CrashReporter/CrashReportWebSite/CrashReportWebSite.csproj
@@ -337,24 +337,17 @@
       <LastGenOutput>CrashReports.cs</LastGenOutput>
     </Content>
     <Content Include="Global.asax" />
-    <Content Include="Scripts\jquery-1.5.1-vsdoc.js" />
-    <Content Include="Scripts\jquery-1.5.1.js" />
-    <Content Include="Scripts\jquery-1.5.1.min.js" />
-    <Content Include="Scripts\jquery-ui-1.8.11.js" />
-    <Content Include="Scripts\jquery-ui-1.8.11.min.js" />
-    <Content Include="Scripts\jquery.blockUI.js" />
-    <Content Include="Scripts\jquery.shiftcheckbox.js" />
-    <Content Include="Scripts\jquery.unobtrusive-ajax.js" />
-    <Content Include="Scripts\jquery.unobtrusive-ajax.min.js" />
-    <Content Include="Scripts\jquery.validate-vsdoc.js" />
-    <Content Include="Scripts\jquery.validate.js" />
-    <Content Include="Scripts\jquery.validate.min.js" />
-    <Content Include="Scripts\jquery.validate.unobtrusive.js" />
-    <Content Include="Scripts\jquery.validate.unobtrusive.min.js" />
-    <Content Include="Scripts\modernizr-1.7.js" />
-    <Content Include="Scripts\modernizr-1.7.min.js" />
+    <None Include="Scripts\jquery-3.1.1.intellisense.js" />
+    <Content Include="Scripts\jquery-3.1.1.js" />
+    <Content Include="Scripts\jquery-3.1.1.min.js" />
+    <Content Include="Scripts\jquery-3.1.1.slim.js" />
+    <Content Include="Scripts\jquery-3.1.1.slim.min.js" />
+    <Content Include="Scripts\jquery-ui-1.12.1.js" />
+    <Content Include="Scripts\jquery-ui-1.12.1.min.js" />
     <Content Include="Scripts\searchCollapse.js" />
     <Content Include="Scripts\UserFilterScript.js" />
+    <Content Include="Scripts\jquery-3.1.1.slim.min.map" />
+    <Content Include="Scripts\jquery-3.1.1.min.map" />
     <None Include="Views\Buggs\Index.cshtml" />
     <None Include="Views\Buggs\Show.cshtml" />
     <None Include="Views\Buggs\ViewBuggs.cshtml" />
@@ -389,12 +382,6 @@
       <DependentUpon>Web.config</DependentUpon>
     </Content>
     <Content Include="Content\Site.css" />
-    <Content Include="Scripts\MicrosoftAjax.js" />
-    <Content Include="Scripts\MicrosoftAjax.debug.js" />
-    <Content Include="Scripts\MicrosoftMvcAjax.js" />
-    <Content Include="Scripts\MicrosoftMvcAjax.debug.js" />
-    <Content Include="Scripts\MicrosoftMvcValidation.js" />
-    <Content Include="Scripts\MicrosoftMvcValidation.debug.js" />
     <None Include="Views\Shared\Error.cshtml" />
     <Content Include="Views\Web.config">
       <SubType>Designer</SubType>
@@ -410,6 +397,7 @@
     <None Include="Properties\Settings.settings">
       <Generator>SettingsSingleFileGenerator</Generator>
       <LastGenOutput>Settings.Designer.cs</LastGenOutput>
+      <SubType>Designer</SubType>
     </None>
     <Content Include="Views\Shared\_CRWLayout.cshtml" />
     <None Include="Web.Development.config">
diff --git a/Engine/Source/Programs/CrashReporter/CrashReportWebSite/DataModels/BuggDerived.cs b/Engine/Source/Programs/CrashReporter/CrashReportWebSite/DataModels/BuggDerived.cs
--- a/Engine/Source/Programs/CrashReporter/CrashReportWebSite/DataModels/BuggDerived.cs
+++ b/Engine/Source/Programs/CrashReporter/CrashReportWebSite/DataModels/BuggDerived.cs
@@ -6,6 +6,7 @@ using System.Linq;
 using System.Web;
 using Newtonsoft.Json;
 using Tools.CrashReporter.CrashReportCommon;
+using Tools.CrashReporter.CrashReportWebSite.Properties;
 
 namespace Tools.CrashReporter.CrashReportWebSite.DataModels
 {
@@ -347,7 +348,7 @@ namespace Tools.CrashReporter.CrashReportWebSite.DataModels
             //fields.Add("customfield_11500", th );
 
             //Platforms - required = false
-            fields.Add("customfield_11203", new List<object>() { new Dictionary<string, object>() { { "self", "https://jira.it.epicgames.net/rest/api/2/customFieldOption/11425" }, { "value", "PC" } } });
+            fields.Add("customfield_11203", new List<object>() { new Dictionary<string, object>() { { "self", Settings.Default.JiraDeploymentAddress + "/customFieldOption/11425" }, { "value", "PC" } } });
 
             //repro rate required = false
             //fields.Add("customfield_11900", this.Buggs_Crashes.Count);
diff --git a/Engine/Source/Programs/CrashReporter/CrashReportWebSite/Global.asax.cs b/Engine/Source/Programs/CrashReporter/CrashReportWebSite/Global.asax.cs
--- a/Engine/Source/Programs/CrashReporter/CrashReportWebSite/Global.asax.cs
+++ b/Engine/Source/Programs/CrashReporter/CrashReportWebSite/Global.asax.cs
@@ -1,9 +1,5 @@
 ﻿// Copyright 1998-2017 Epic Games, Inc. All Rights Reserved.
 
-using System;
-using System.Collections.Generic;
-using System.Linq;
-using System.Web;
 using System.Web.Mvc;
 using System.Web.Routing;
 
diff --git a/Engine/Source/Programs/CrashReporter/CrashReportWebSite/Web.config b/Engine/Source/Programs/CrashReporter/CrashReportWebSite/Web.config
--- a/Engine/Source/Programs/CrashReporter/CrashReportWebSite/Web.config
+++ b/Engine/Source/Programs/CrashReporter/CrashReportWebSite/Web.config
@@ -1,4 +1,4 @@
-﻿<?xml version="1.0"?>
+﻿<?xml version="1.0" encoding="utf-8"?>
 <!-- 
     Note: As an alternative to hand editing this file you can use the 
     web admin tool to configure settings for your application. Use
@@ -10,28 +10,28 @@
 <configuration>
   <configSections>
     <sectionGroup name="applicationSettings" type="System.Configuration.ApplicationSettingsGroup, System, Version=4.0.0.0, Culture=neutral, PublicKeyToken=b77a5c561934e089">
-      <section name="Tools.CrashReporter.CrashReportWebSite.Properties.Settings" type="System.Configuration.ClientSettingsSection, System, Version=4.0.0.0, Culture=neutral, PublicKeyToken=b77a5c561934e089" requirePermission="false"/>
+      <section name="Tools.CrashReporter.CrashReportWebSite.Properties.Settings" type="System.Configuration.ClientSettingsSection, System, Version=4.0.0.0, Culture=neutral, PublicKeyToken=b77a5c561934e089" requirePermission="false" />
     </sectionGroup>
-    <section name="entityFramework" type="System.Data.Entity.Internal.ConfigFile.EntityFrameworkSection, EntityFramework, Version=6.0.0.0, Culture=neutral, PublicKeyToken=b77a5c561934e089" requirePermission="false"/>
+    <section name="entityFramework" type="System.Data.Entity.Internal.ConfigFile.EntityFrameworkSection, EntityFramework, Version=6.0.0.0, Culture=neutral, PublicKeyToken=b77a5c561934e089" requirePermission="false" />
   </configSections>
-  <appSettings/>
+  <appSettings />
   <connectionStrings>
-    <add name="Hangfire" connectionString="Data Source=db-09;Initial Catalog=CrashReporterSchduledTask;Server=db-09;Database=CrashReport;Integrated Security=true" providerName="System.Data.SqlClient"/>
-    <add name="CrashReportEntities" connectionString="metadata=res://*/DataModels.CrashReports.csdl|res://*/DataModels.CrashReports.ssdl|res://*/DataModels.CrashReports.msl;provider=System.Data.SqlClient;provider connection string=&quot;data source=db-09;initial catalog=CrashReport;integrated security=True;MultipleActiveResultSets=True;App=EntityFramework;Max Pool Size=500;Connection Timeout=60;&quot;" providerName="System.Data.EntityClient"/>
+    <add name="Hangfire" connectionString="Data Source=db-09;Initial Catalog=CrashReporterSchduledTask;Server=db-09;Database=CrashReport;Integrated Security=true" providerName="System.Data.SqlClient" />
+    <add name="CrashReportEntities" connectionString="metadata=res://*/DataModels.CrashReports.csdl|res://*/DataModels.CrashReports.ssdl|res://*/DataModels.CrashReports.msl;provider=System.Data.SqlClient;provider connection string=&quot;data source=db-09;initial catalog=CrashReport;integrated security=True;MultipleActiveResultSets=True;App=EntityFramework;Max Pool Size=500;Connection Timeout=60;&quot;" providerName="System.Data.EntityClient" />
   </connectionStrings>
   <system.web>
-    <compilation debug="true" targetFramework="4.5"/>
-    <customErrors mode="Off"/>
+    <compilation debug="true" targetFramework="4.5" />
+    <customErrors mode="Off" />
     <pages enableSessionState="false" clientIDMode="AutoID" controlRenderingCompatibilityVersion="4.0">
       <namespaces>
-        <add namespace="System.Web.Mvc"/>
-        <add namespace="System.Web.Mvc.Ajax"/>
-        <add namespace="System.Web.Mvc.Html"/>
-        <add namespace="System.Web.Routing"/>
-        <add namespace="Tools.CrashReporter.CrashReportWebSite.Views.Helpers"/>
+        <add namespace="System.Web.Mvc" />
+        <add namespace="System.Web.Mvc.Ajax" />
+        <add namespace="System.Web.Mvc.Html" />
+        <add namespace="System.Web.Routing" />
+        <add namespace="Tools.CrashReporter.CrashReportWebSite.Views.Helpers" />
       </namespaces>
     </pages>
-    <sessionState mode="Off"/>
+    <sessionState mode="Off" />
   </system.web>
   <!-- 
         The system.webServer section is required for running ASP.NET AJAX under Internet
@@ -39,16 +39,16 @@
   -->
   <system.webServer>
     <modules runAllManagedModulesForAllRequests="false">
-      <add name="UriCacheModule"/>
-      <add name="TokenCacheModule"/>
-      <add name="RequestMonitorModule"/>
+      <add name="UriCacheModule" />
+      <add name="TokenCacheModule" />
+      <add name="RequestMonitorModule" />
     </modules>
     <handlers>
-      <remove name="UrlRoutingHandler"/>
+      <remove name="UrlRoutingHandler" />
     </handlers>
     <caching>
       <profiles>
-        <add extension=".aspx" policy="CacheForTimePeriod" kernelCachePolicy="CacheForTimePeriod" duration="00:02:30" varyByHeaders="Post"/>
+        <add extension=".aspx" policy="CacheForTimePeriod" kernelCachePolicy="CacheForTimePeriod" duration="00:02:30" varyByHeaders="Post" />
       </profiles>
     </caching>
   </system.webServer>
@@ -86,34 +86,34 @@
   <runtime>
     <assemblyBinding xmlns="urn:schemas-microsoft-com:asm.v1">
       <dependentAssembly>
-        <assemblyIdentity name="System.Web.Helpers" publicKeyToken="31bf3856ad364e35"/>
-        <bindingRedirect oldVersion="1.0.0.0-3.0.0.0" newVersion="3.0.0.0"/>
+        <assemblyIdentity name="System.Web.Helpers" publicKeyToken="31bf3856ad364e35" />
+        <bindingRedirect oldVersion="1.0.0.0-3.0.0.0" newVersion="3.0.0.0" />
       </dependentAssembly>
       <dependentAssembly>
-        <assemblyIdentity name="System.Data.Linq" publicKeyToken="b77a5c561934e089"/>
-        <bindingRedirect oldVersion="1.0.0.0-4.0.0.0" newVersion="4.0.0.0"/>
+        <assemblyIdentity name="System.Data.Linq" publicKeyToken="b77a5c561934e089" />
+        <bindingRedirect oldVersion="1.0.0.0-4.0.0.0" newVersion="4.0.0.0" />
       </dependentAssembly>
       <dependentAssembly>
-        <assemblyIdentity name="System.Web.WebPages" publicKeyToken="31bf3856ad364e35"/>
-        <bindingRedirect oldVersion="1.0.0.0-3.0.0.0" newVersion="3.0.0.0"/>
+        <assemblyIdentity name="System.Web.WebPages" publicKeyToken="31bf3856ad364e35" />
+        <bindingRedirect oldVersion="1.0.0.0-3.0.0.0" newVersion="3.0.0.0" />
       </dependentAssembly>
       <dependentAssembly>
-        <assemblyIdentity name="System.Web.Mvc" publicKeyToken="31bf3856ad364e35"/>
-        <bindingRedirect oldVersion="0.0.0.0-5.2.3.0" newVersion="5.2.3.0"/>
+        <assemblyIdentity name="System.Web.Mvc" publicKeyToken="31bf3856ad364e35" />
+        <bindingRedirect oldVersion="0.0.0.0-5.2.3.0" newVersion="5.2.3.0" />
       </dependentAssembly>
     </assemblyBinding>
   </runtime>
   <system.net>
     <mailSettings>
       <smtp deliveryMethod="Network">
-        <network host="smtp.epicgames.net" defaultCredentials="false"/>
+        <network host="smtp.epicgames.net" defaultCredentials="false" />
       </smtp>
     </mailSettings>
   </system.net>
   <entityFramework>
-    <defaultConnectionFactory type="System.Data.Entity.Infrastructure.SqlConnectionFactory, EntityFramework"/>
+    <defaultConnectionFactory type="System.Data.Entity.Infrastructure.SqlConnectionFactory, EntityFramework" />
     <providers>
-      <provider invariantName="System.Data.SqlClient" type="System.Data.Entity.SqlServer.SqlProviderServices, EntityFramework.SqlServer"/>
+      <provider invariantName="System.Data.SqlClient" type="System.Data.Entity.SqlServer.SqlProviderServices, EntityFramework.SqlServer" />
     </providers>
   </entityFramework>
 </configuration>
\ No newline at end of file
diff --git a/Engine/Source/Programs/CrashReporter/CrashReportWebSite/packages.config b/Engine/Source/Programs/CrashReporter/CrashReportWebSite/packages.config
--- a/Engine/Source/Programs/CrashReporter/CrashReportWebSite/packages.config
+++ b/Engine/Source/Programs/CrashReporter/CrashReportWebSite/packages.config
@@ -8,6 +8,8 @@
   <package id="Hangfire.Core" version="1.5.3" targetFramework="net45" />
   <package id="Hangfire.SqlServer" version="1.5.3" targetFramework="net45" />
   <package id="Hangfire.Unity" version="1.1.0" targetFramework="net45" />
+  <package id="jQuery" version="3.1.1" targetFramework="net45" />
+  <package id="jQuery.UI.Combined" version="1.12.1" targetFramework="net45" />
   <package id="Microsoft.AspNet.Mvc" version="5.2.3" targetFramework="net45" />
   <package id="Microsoft.AspNet.Razor" version="3.2.3" targetFramework="net45" />
   <package id="Microsoft.AspNet.WebPages" version="3.2.3" targetFramework="net45" />
diff --git a/Engine/Source/Programs/Mac/ShaderCacheTool/Private/ShaderCacheTool.cpp b/Engine/Source/Programs/Mac/ShaderCacheTool/Private/ShaderCacheTool.cpp
--- a/Engine/Source/Programs/Mac/ShaderCacheTool/Private/ShaderCacheTool.cpp
+++ b/Engine/Source/Programs/Mac/ShaderCacheTool/Private/ShaderCacheTool.cpp
@@ -46,7 +46,7 @@ namespace SCT
 					FString Vers = Switch.Mid(8);
 					if(Vers.IsNumeric())
 					{
-						LexicalConversion::FromString(GameVersion, *Vers);
+						Lex::FromString(GameVersion, *Vers);
 					}
 				}
 				else
diff --git a/Engine/Source/Programs/UnrealBuildTool/Configuration/EngineConfiguration.cs b/Engine/Source/Programs/UnrealBuildTool/Configuration/EngineConfiguration.cs
--- a/Engine/Source/Programs/UnrealBuildTool/Configuration/EngineConfiguration.cs
+++ b/Engine/Source/Programs/UnrealBuildTool/Configuration/EngineConfiguration.cs
@@ -39,7 +39,7 @@ namespace UnrealBuildTool
 																	"/Script/BuildSettings.BuildSettings", "/Script/IOSRuntimeSettings.IOSRuntimeSettings", "/Script/WindowsTargetPlatform.WindowsTargetSettings",
 																	"/Script/UnrealEd.ProjectPackagingSettings", "/Script/PS4PlatformEditor.PS4TargetSettings", "/Script/XboxOnePlatformEditor.XboxOneTargetSettings",
 																	"/Script/HTML5PlatformEditor.HTML5TargetSettings","PS4SymbolServer","/Script/EngineSettings.GeneralProjectSettings",
-                                                                    "/Script/UnrealEd.ProjectPackagingSettings", "InstalledPlatforms", "OnlineSubsystemGooglePlay.Store","/Script/Engine.StreamingSettings", "Core.Encryption",
+                                                                    "/Script/UnrealEd.ProjectPackagingSettings", "InstalledPlatforms", "OnlineSubsystemGooglePlay.Store", "Core.Encryption",
                                                                     "/Script/SwitchRuntimeSettings.SwitchRuntimeSettings", "/Script/Engine.NetworkSettings", "/Script/SourceCodeAccess.SourceCodeAccessSettings",
 																	"/Script/LinuxTargetPlatform.LinuxTargetSettings", "BuildConfiguration", "UEBuildConfiguration", "/Script/WolfRuntimeSettings.WolfRuntimeSettings", 
 																	"/Script/LocationServicesIOSEditor.LocationServicesIOSSettings"};
@@ -809,10 +809,22 @@ namespace UnrealBuildTool
 				// Engine/Config/Platform/Platform* ini
 				yield return FileReference.Combine(EngineDirectory, "Config", PlatformName, PlatformName + BaseIniName + ".ini");
 
+				// Engine/Config/NotForLicensees/Platform/Platform* ini
+				yield return FileReference.Combine(EngineDirectory, "Config", "NotForLicensees", PlatformName, PlatformName + BaseIniName + ".ini");
+
+				// Engine/Config/NoRedist/Platform/Platform* ini
+				yield return FileReference.Combine(EngineDirectory, "Config", "NoRedist", PlatformName, PlatformName + BaseIniName + ".ini");
+
 				if (ProjectDirectory != null)
 				{
 					// Game/Config/Platform/Platform* ini
 					yield return FileReference.Combine(ProjectDirectory, "Config", PlatformName, PlatformName + BaseIniName + ".ini");
+
+					// Engine/Config/NotForLicensees/Platform/Platform* ini
+					yield return FileReference.Combine(ProjectDirectory, "Config", "NotForLicensees", PlatformName, PlatformName + BaseIniName + ".ini");
+
+					// Engine/Config/NoRedist/Platform/Platform* ini
+					yield return FileReference.Combine(ProjectDirectory, "Config", "NoRedist", PlatformName, PlatformName + BaseIniName + ".ini");
 				}
 			}
 
diff --git a/Engine/Source/Programs/UnrealBuildTool/Configuration/UEBuildConfiguration.cs b/Engine/Source/Programs/UnrealBuildTool/Configuration/UEBuildConfiguration.cs
--- a/Engine/Source/Programs/UnrealBuildTool/Configuration/UEBuildConfiguration.cs
+++ b/Engine/Source/Programs/UnrealBuildTool/Configuration/UEBuildConfiguration.cs
@@ -407,7 +407,6 @@ namespace UnrealBuildTool
 			bEditorDependsOnShaderCompileWorker = true;
             bForceCompileDevelopmentAutomationTests = false;
             bForceCompilePerformanceAutomationTests = false;
-			bEventDrivenLoader = false;
 		}
 
 		/// <summary>
diff --git a/Engine/Source/Programs/UnrealBuildTool/Configuration/UEBuildPlatform.cs b/Engine/Source/Programs/UnrealBuildTool/Configuration/UEBuildPlatform.cs
--- a/Engine/Source/Programs/UnrealBuildTool/Configuration/UEBuildPlatform.cs
+++ b/Engine/Source/Programs/UnrealBuildTool/Configuration/UEBuildPlatform.cs
@@ -276,13 +276,6 @@ namespace UnrealBuildTool
 					UEBuildConfiguration.bCompileCEF3 = bValue;
 				}
 
-				if (Target != null && Target.IsCooked)
-				{
-					if (Ini.GetBool("/Script/Engine.StreamingSettings", "s.EventDrivenLoaderEnabled", out bValue))
-					{
-						UEBuildConfiguration.bEventDrivenLoader = bValue;
-					}
-				}
 				bInitializedProject = true;
 			}
             else
@@ -820,11 +813,6 @@ namespace UnrealBuildTool
 		/// </summary>
 		public virtual bool HasDefaultBuildConfig(UnrealTargetPlatform Platform, DirectoryReference ProjectDirectoryName)
 		{
-			if(!DoProjectSettingsMatchDefault(Platform, ProjectDirectoryName, "/Script/Engine.StreamingSettings", new string[] { "s.EventDrivenLoaderEnabled" }, null, null))
-			{
-				return false;
-			}
-
 			string[] BoolKeys = new string[] {
 				"bCompileApex", "bCompileBox2D", "bCompileICU", "bCompileSimplygon", "bCompileSimplygonSSF",
 				"bCompileLeanAndMeanUE", "bIncludeADO", "bCompileRecast", "bCompileSpeedTree", 
diff --git a/Engine/Source/Programs/UnrealBuildTool/Configuration/UEBuildTarget.cs b/Engine/Source/Programs/UnrealBuildTool/Configuration/UEBuildTarget.cs
--- a/Engine/Source/Programs/UnrealBuildTool/Configuration/UEBuildTarget.cs
+++ b/Engine/Source/Programs/UnrealBuildTool/Configuration/UEBuildTarget.cs
@@ -3261,6 +3261,44 @@ namespace UnrealBuildTool
 				Result.Add("#endif");
 			}
 
+			// Write functions for accessing embedded pak signing keys
+			String EncryptionKey;
+			String[] PakSigningKeys;
+			GetEncryptionAndSigningKeys(out EncryptionKey, out PakSigningKeys);
+			bool bRegisterEncryptionKey = false;
+			bool bRegisterPakSigningKeys = false;
+
+			if (!string.IsNullOrEmpty(EncryptionKey))
+			{
+				Result.Add("extern void RegisterEncryptionKey(const char*);");
+				bRegisterEncryptionKey = true;
+			}
+
+			if (PakSigningKeys != null && PakSigningKeys.Length == 3 && !string.IsNullOrEmpty(PakSigningKeys[1]) && !string.IsNullOrEmpty(PakSigningKeys[2]))
+			{
+				Result.Add("extern void RegisterPakSigningKeys(const char*, const char*);");
+				bRegisterPakSigningKeys = true;
+			}
+
+			if (bRegisterEncryptionKey || bRegisterPakSigningKeys)
+			{
+				Result.Add("struct FEncryptionAndSigningKeyRegistration");
+				Result.Add("{");
+				Result.Add("\tFEncryptionAndSigningKeyRegistration()");
+				Result.Add("\t{");
+				if (bRegisterEncryptionKey)
+				{
+					Result.Add(string.Format("\t\tRegisterEncryptionKey(\"{0}\");", EncryptionKey));
+				}
+				if (bRegisterPakSigningKeys)
+				{
+					Result.Add(string.Format("\t\tRegisterPakSigningKeys(\"{0}\", \"{1}\");", PakSigningKeys[2], PakSigningKeys[1]));
+				}
+				Result.Add("\t}");
+				Result.Add("};");
+				Result.Add("FEncryptionAndSigningKeyRegistration GEncryptionAndSigningKeyRegistration;");
+			}
+
 			// Add a function that is not referenced by anything that invokes all the empty functions in the different static libraries
 			Result.Add("void " + LinkerFixupsName + "()");
 			Result.Add("{");
@@ -4026,12 +4064,7 @@ namespace UnrealBuildTool
 				}
 				GlobalCompileEnvironment.Config.Definitions.Add(String.Format("IS_PROGRAM={0}", TargetType == TargetRules.TargetType.Program ? "1" : "0"));
 			}
-
-			if (!bUseSharedBuildEnvironment)
-			{
-				SetupEncryptionAndSingningKeys();
-			}
-
+			
 			// Validate UE configuration - needs to happen before setting any environment mojo and after argument parsing.
 			PlatformContext.ValidateUEBuildConfiguration();
 			UEBuildConfiguration.ValidateConfiguration();
@@ -4270,15 +4303,6 @@ namespace UnrealBuildTool
 			SetUpConfigurationEnvironment();
 			SetUpProjectEnvironment(Configuration);
 
-			if (UEBuildConfiguration.bEventDrivenLoader && !bUseSharedBuildEnvironment)
-			{
-				GlobalCompileEnvironment.Config.Definitions.Add("USE_NEW_ASYNC_IO=1");
-			}
-			else
-			{
-				GlobalCompileEnvironment.Config.Definitions.Add("USE_NEW_ASYNC_IO=0");
-			}
-
 			// Validates current settings and updates if required.
 			BuildConfiguration.ValidateConfiguration(
 				GlobalCompileEnvironment.Config.Configuration,
@@ -4287,35 +4311,21 @@ namespace UnrealBuildTool
 				PlatformContext);
 		}
 
-		private void SetupEncryptionAndSingningKeys()
+		private void GetEncryptionAndSigningKeys(out String AESKey, out String[] PakSigningKeys)
 		{
-			String[] RSAKeys;
-			String AESKey;
-			ParseEncryptionIni(ProjectDirectory, Platform, out RSAKeys, out AESKey);
-			bool bSigningEnabledByIniFile = false;
-
-			if (RSAKeys != null)
-			{
-				GlobalCompileEnvironment.Config.Definitions.Add("DECRYPTION_KEY_MODULUS=\"" + RSAKeys[1] + "\"");
-				GlobalCompileEnvironment.Config.Definitions.Add("DECRYPTION_KEY_EXPONENT=\"" + RSAKeys[2] + "\"");
-
-				bSigningEnabledByIniFile = true;
-			}
+			ParseEncryptionIni(ProjectDirectory, Platform, out PakSigningKeys, out AESKey);
 
 			if (!String.IsNullOrEmpty(AESKey))
 			{
 				if (AESKey.Length < 32)
 				{
 					Log.TraceError("AES key specified in configs must be at least 32 characters long!");
-				}
-				else
-				{
-					GlobalCompileEnvironment.Config.Definitions.Add("AES_KEY=\"" + AESKey + "\"");
+					AESKey = String.Empty;
 				}
 			}
 
 			// If we didn't extract any keys from the new ini file setup, try looking for the old keys text file
-			if (!bSigningEnabledByIniFile && !string.IsNullOrEmpty(Rules.PakSigningKeysFile))
+			if (PakSigningKeys == null && !string.IsNullOrEmpty(Rules.PakSigningKeysFile))
 			{
 				string FullFilename = Path.Combine(ProjectDirectory.FullName, Rules.PakSigningKeysFile);
 
@@ -4338,8 +4348,9 @@ namespace UnrealBuildTool
 
 					if (Keys.Count == 3)
 					{
-						GlobalCompileEnvironment.Config.Definitions.Add("DECRYPTION_KEY_MODULUS=\"" + Keys[1] + "\"");
-						GlobalCompileEnvironment.Config.Definitions.Add("DECRYPTION_KEY_EXPONENT=\"" + Keys[2] + "\"");
+						PakSigningKeys = new String[2];
+						PakSigningKeys[0] = Keys[1];
+						PakSigningKeys[1] = Keys[2];
 					}
 					else
 					{
diff --git a/Engine/Source/Programs/UnrealHeaderTool/Private/BaseParser.cpp b/Engine/Source/Programs/UnrealHeaderTool/Private/BaseParser.cpp
--- a/Engine/Source/Programs/UnrealHeaderTool/Private/BaseParser.cpp
+++ b/Engine/Source/Programs/UnrealHeaderTool/Private/BaseParser.cpp
@@ -897,7 +897,7 @@ bool FBaseParser::PeekIdentifier( const TCHAR* Match )
 //
 // Unget the most recently gotten token.
 //
-void FBaseParser::UngetToken( FToken& Token )
+void FBaseParser::UngetToken( const FToken& Token )
 {
 	InputPos = Token.StartPos;
 	InputLine = Token.StartLine;
diff --git a/Engine/Source/Programs/UnrealHeaderTool/Private/BaseParser.h b/Engine/Source/Programs/UnrealHeaderTool/Private/BaseParser.h
--- a/Engine/Source/Programs/UnrealHeaderTool/Private/BaseParser.h
+++ b/Engine/Source/Programs/UnrealHeaderTool/Private/BaseParser.h
@@ -45,7 +45,8 @@ class FBaseParser
 {
 protected:
 	FBaseParser();
-protected:
+
+public:
 	// Input text.
 	const TCHAR* Input;
 
@@ -72,7 +73,7 @@ protected:
 
 	// Total number of lines parsed.
 	int32 LinesParsed;
-protected:
+
 	void ResetParser(const TCHAR* SourceBuffer, int32 StartingLineNumber = 1);
 
 	// Low-level parsing functions.
@@ -126,7 +127,7 @@ protected:
 	// Doesn't quit if StopChar is found inside a double-quoted string, but does not support quote escapes
 	bool GetRawTokenRespectingQuotes( FToken& Token, TCHAR StopChar = TCHAR('\n') );
 
-	void UngetToken( FToken& Token );
+	void UngetToken( const FToken& Token );
 	bool GetIdentifier( FToken& Token, bool bNoConsts = false );
 	bool GetSymbol( FToken& Token );
 
@@ -175,7 +176,6 @@ protected:
 	// Reads a set of specifiers (with optional values) inside the () of a new-style metadata macro like UPROPERTY or UFUNCTION
 	void ReadSpecifierSetInsideMacro(TArray<FPropertySpecifier>& SpecifiersFound, const FString& TypeOfSpecifier, TMap<FName, FString>& MetaData);
 
-public:
 	// Validates and inserts one key-value pair into the meta data map
 	static void InsertMetaDataPair(TMap<FName, FString>& MetaData, const FString& InKey, const FString& InValue);
 
diff --git a/Engine/Source/Programs/UnrealHeaderTool/Private/HeaderParser.cpp b/Engine/Source/Programs/UnrealHeaderTool/Private/HeaderParser.cpp
--- a/Engine/Source/Programs/UnrealHeaderTool/Private/HeaderParser.cpp
+++ b/Engine/Source/Programs/UnrealHeaderTool/Private/HeaderParser.cpp
@@ -83,6 +83,13 @@ namespace
 {
 	bool ProbablyAMacro(const TCHAR* Identifier)
 	{
+		// Macros must start with a capitalized alphanumeric character or underscore
+		TCHAR FirstChar = Identifier[0];
+		if (FirstChar != TEXT('_') && (FirstChar < TEXT('A') || FirstChar > TEXT('Z')))
+		{
+			return false;
+		}
+
 		// Test for known delegate and event macros.
 		TCHAR MulticastDelegateStart[] = TEXT("DECLARE_MULTICAST_DELEGATE");
 		if (!FCString::Strncmp(Identifier, MulticastDelegateStart, ARRAY_COUNT(MulticastDelegateStart) - 1))
@@ -103,9 +110,9 @@ namespace
 		}
 
 		// Failing that, we'll guess about it being a macro based on it being a fully-capitalized identifier.
-		while (TCHAR Ch = *Identifier++)
+		while (TCHAR Ch = *++Identifier)
 		{
-			if (Ch != TEXT('_') && (Ch < TEXT('A') || Ch > TEXT('Z')))
+			if (Ch != TEXT('_') && (Ch < TEXT('A') || Ch > TEXT('Z')) && (Ch < TEXT('0') || Ch > TEXT('9')))
 			{
 				return false;
 			}
@@ -114,6 +121,84 @@ namespace
 		return true;
 	}
 
+	/**
+	 * Tests if an identifier looks like a macro which doesn't have a following open parenthesis.
+	 *
+	 * @param HeaderParser  The parser to retrieve the next token.
+	 * @param Token         The token to test for being callable-macro-like.
+	 *
+	 * @return true if it looks like a non-callable macro, false otherwise.
+	 */
+	bool ProbablyAnUnknownObjectLikeMacro(FHeaderParser& HeaderParser, FToken Token)
+	{
+		// Non-identifiers are not macros
+		if (Token.TokenType != TOKEN_Identifier)
+		{
+			return false;
+		}
+
+		// Macros must start with a capitalized alphanumeric character or underscore
+		TCHAR FirstChar = Token.Identifier[0];
+		if (FirstChar != TEXT('_') && (FirstChar < TEXT('A') || FirstChar > TEXT('Z')))
+		{
+			return false;
+		}
+
+		// We'll guess about it being a macro based on it being fully-capitalized with at least one underscore.
+		const TCHAR* IdentPtr = Token.Identifier;
+		int32 UnderscoreCount = 0;
+		while (TCHAR Ch = *++IdentPtr)
+		{
+			if (Ch == TEXT('_'))
+			{
+				++UnderscoreCount;
+			}
+			else if ((Ch < TEXT('A') || Ch > TEXT('Z')) && (Ch < TEXT('0') || Ch > TEXT('9')))
+			{
+				return false;
+			}
+		}
+
+		// We look for at least one underscore as a convenient way of whitelisting many known macros
+		// like FORCEINLINE and CONSTEXPR, and non-macros like FPOV and TCHAR.
+		if (UnderscoreCount == 0)
+		{
+			return false;
+		}
+
+		// Identifiers which end in _API are known
+		if (IdentPtr - Token.Identifier > 4 && IdentPtr[-4] == TEXT('_') && IdentPtr[-3] == TEXT('A') && IdentPtr[-2] == TEXT('P') && IdentPtr[-1] == TEXT('I'))
+		{
+			return false;
+		}
+
+		// Ignore certain known macros or identifiers that look like macros.
+		// IMPORTANT: needs to be in lexicographical order.
+		static const TCHAR* Whitelist[] =
+		{
+			TEXT("FORCEINLINE_DEBUGGABLE"),
+			TEXT("FORCEINLINE_STATS"),
+			TEXT("SIZE_T")
+		};
+		if (Algo::FindSortedStringCaseInsensitive(Token.Identifier, Whitelist, ARRAY_COUNT(Whitelist)) >= 0)
+		{
+			return false;
+		}
+
+		// Check if there's an open parenthesis following the token.
+		//
+		// Rather than ungetting the bracket token, we unget the original identifier token,
+		// then get it again, so we don't lose any comments which may exist between the token 
+		// and the non-bracket.
+		FToken PossibleBracketToken;
+		HeaderParser.GetToken(PossibleBracketToken);
+		HeaderParser.UngetToken(Token);
+		HeaderParser.GetToken(Token);
+
+		bool bResult = PossibleBracketToken.TokenType != TOKEN_Symbol || FCString::Strcmp(PossibleBracketToken.Identifier, TEXT("("));
+		return bResult;
+	}
+
 	/**
 	 * Parse and validate an array of identifiers (inside FUNC_NetRequest, FUNC_NetResponse) 
 	 * @param FuncInfo function info for the current function
@@ -2225,6 +2310,10 @@ UScriptStruct* FHeaderParser::CompileStructDeclaration(FClasses& AllClasses)
 				}
 			}
 		}
+		else if (ProbablyAnUnknownObjectLikeMacro(*this, Token))
+		{
+			// skip it
+		}
 		else
 		{
 			if ( !Token.Matches( TEXT("}") ) )
@@ -2239,8 +2328,8 @@ UScriptStruct* FHeaderParser::CompileStructDeclaration(FClasses& AllClasses)
 			{
 				MatchSemi();
 				break;
-			}			
-		}		
+			}
+		}
 	}
 
 	// Validation
@@ -4860,9 +4949,9 @@ bool FHeaderParser::CompileDeclaration(FClasses& AllClasses, TArray<UDelegateFun
 		}
 	}
 
-	if (Token.Matches(TEXT("PRAGMA_DISABLE_DEPRECATION_WARNINGS")) || Token.Matches(TEXT("PRAGMA_ENABLE_DEPRECATION_WARNINGS")))
+	// Skip anything that looks like a macro followed by no bracket that we don't know about
+	if (ProbablyAnUnknownObjectLikeMacro(*this, Token))
 	{
-		// Skip these macros
 		return true;
 	}
 
@@ -8089,7 +8178,7 @@ bool FHeaderParser::DefaultValueStringCppFormatToInnerFormat(const UProperty* Pr
 			int64 Value;
 			if (FDefaultValueHelper::ParseInt64(CppForm, Value))
 			{
-				OutForm = LexicalConversion::ToString(Value);
+				OutForm = Lex::ToString(Value);
 				return EnumProp->GetUnderlyingProperty()->CanHoldValue(Value);
 			}
 		}
diff --git a/Engine/Source/Runtime/Analytics/Analytics/Public/AnalyticsConversion.h b/Engine/Source/Runtime/Analytics/Analytics/Public/AnalyticsConversion.h
--- a/Engine/Source/Runtime/Analytics/Analytics/Public/AnalyticsConversion.h
+++ b/Engine/Source/Runtime/Analytics/Analytics/Public/AnalyticsConversion.h
@@ -44,18 +44,18 @@ namespace AnalyticsConversion
 		return Value.ToString();
 	}
 
-	/** Double conversion. LexicalConversion is broken to doubles (won't use SanitizeFloat), so overload this directly. */
+	/** Double conversion. Lex is broken to doubles (won't use SanitizeFloat), so overload this directly. */
 	inline FString ToString(double Value)
 	{
 		return FString::SanitizeFloat(Value);
 	}
 
 
-	/** Lexical conversion. Allow any type that we have a LexicalConversion for. */
+	/** Lexical conversion. Allow any type that we have a Lex for. */
 	template <typename T>
 	inline typename TEnableIf<TIsArithmetic<T>::Value, FString>::Type ToString(T Value)
 	{
-		return LexicalConversion::ToSanitizedString(Value);
+		return Lex::ToSanitizedString(Value);
 	}
 
 	/** Array conversion. Creates comma-separated list. */
diff --git a/Engine/Source/Runtime/Analytics/AnalyticsET/Private/HttpServiceTracker.cpp b/Engine/Source/Runtime/Analytics/AnalyticsET/Private/HttpServiceTracker.cpp
--- a/Engine/Source/Runtime/Analytics/AnalyticsET/Private/HttpServiceTracker.cpp
+++ b/Engine/Source/Runtime/Analytics/AnalyticsET/Private/HttpServiceTracker.cpp
@@ -41,7 +41,7 @@ bool FHttpServiceTracker::Tick(float DeltaTime)
 				// one attribute per response code.
 				for (const auto& ResponseCodeMapPair : MetricsMapPair.Value.ResponseCodes)
 				{
-					Attrs.Emplace(FString(TEXT("Code-")) + LexicalConversion::ToString(ResponseCodeMapPair.Key), ResponseCodeMapPair.Value);
+					Attrs.Emplace(FString(TEXT("Code-")) + Lex::ToString(ResponseCodeMapPair.Key), ResponseCodeMapPair.Value);
 				}
 				AnalyticsProvider->RecordEvent(MetricsMapPair.Key.ToString(), Attrs);
 			}
diff --git a/Engine/Source/Runtime/Analytics/AnalyticsET/Private/IAnalyticsProviderET.cpp b/Engine/Source/Runtime/Analytics/AnalyticsET/Private/IAnalyticsProviderET.cpp
--- a/Engine/Source/Runtime/Analytics/AnalyticsET/Private/IAnalyticsProviderET.cpp
+++ b/Engine/Source/Runtime/Analytics/AnalyticsET/Private/IAnalyticsProviderET.cpp
@@ -36,7 +36,7 @@ struct FAnalyticsPerfTracker : FTickerObjectBase
 			LogFile.Serialize(TEXT("Date,CL,RunID,Time,WindowSeconds,ProfiledSeconds,Frames,Flushes,Events,Bytes"), ELogVerbosity::Log, FName());
 			LastSubmitTime = StartTime;
 			StartDate = FDateTime::UtcNow().ToIso8601();
-			CL = LexicalConversion::ToString(FEngineVersion::Current().GetChangelist());
+			CL = Lex::ToString(FEngineVersion::Current().GetChangelist());
 		}
 	}
 
diff --git a/Engine/Source/Runtime/AssetRegistry/Private/PackageReader.cpp b/Engine/Source/Runtime/AssetRegistry/Private/PackageReader.cpp
--- a/Engine/Source/Runtime/AssetRegistry/Private/PackageReader.cpp
+++ b/Engine/Source/Runtime/AssetRegistry/Private/PackageReader.cpp
@@ -103,9 +103,7 @@ bool FPackageReader::OpenPackageFile(EOpenPackageResult* OutErrorCode)
 	// check if this is a compressed package and decompress header 
 	if ( !!(PackageFileSummary.PackageFlags & PKG_StoreCompressed) )
 	{
-#if USE_NEW_ASYNC_IO
-		check(!"Package level compression cannot be used with the async io scheme.");
-#else
+		checkf(!GNewAsyncIO, TEXT("Package level compression cannot be used with the async io scheme."));
 		check(PackageFileSummary.CompressedChunks.Num() > 0);
 
 		int64 CurPos = Loader->Tell();
@@ -121,7 +119,6 @@ bool FPackageReader::OpenPackageFile(EOpenPackageResult* OutErrorCode)
 		}
 		
 		Seek(Loader->Tell());
-#endif
 	}
 
 	//make sure the filereader gets the correct version number (it defaults to latest version)
diff --git a/Engine/Source/Runtime/AssetRegistry/Public/AssetData.h b/Engine/Source/Runtime/AssetRegistry/Public/AssetData.h
--- a/Engine/Source/Runtime/AssetRegistry/Public/AssetData.h
+++ b/Engine/Source/Runtime/AssetRegistry/Public/AssetData.h
@@ -413,7 +413,7 @@ inline bool FAssetData::GetTagValue(const FName InTagName, ValueType& OutTagValu
 	if (const FString* FoundValue = TagsAndValues.Find(InTagName))
 	{
 		FMemory::Memzero(&OutTagValue, sizeof(ValueType));
-		LexicalConversion::FromString(OutTagValue, **FoundValue);
+		Lex::FromString(OutTagValue, **FoundValue);
 		return true;
 	}
 	return false;
@@ -444,7 +444,7 @@ inline ValueType FAssetData::GetTagValueRef(const FName InTagName) const
 	FMemory::Memzero(&TmpValue, sizeof(ValueType));
 	if (const FString* FoundValue = TagsAndValues.Find(InTagName))
 	{
-		LexicalConversion::FromString(TmpValue, **FoundValue);
+		Lex::FromString(TmpValue, **FoundValue);
 	}
 	return TmpValue;
 }
diff --git a/Engine/Source/Runtime/Core/Private/GenericPlatform/GenericPlatformFile.cpp b/Engine/Source/Runtime/Core/Private/GenericPlatform/GenericPlatformFile.cpp
--- a/Engine/Source/Runtime/Core/Private/GenericPlatform/GenericPlatformFile.cpp
+++ b/Engine/Source/Runtime/Core/Private/GenericPlatform/GenericPlatformFile.cpp
@@ -9,8 +9,6 @@
 #include "Async/AsyncWork.h"
 #include "UniquePtr.h"
 
-#if USE_NEW_ASYNC_IO
-
 #include "AsyncFileHandle.h"
 
 class FGenericBaseRequest;
@@ -294,6 +292,7 @@ bool FGenericReadRequest::CheckForPrecache()
 
 IAsyncReadFileHandle* IPlatformFile::OpenAsyncRead(const TCHAR* Filename)
 {
+	check(GNewAsyncIO);
 	return new FGenericAsyncReadFileHandle(this, Filename);
 }
 
@@ -301,8 +300,6 @@ DEFINE_STAT(STAT_AsyncFileMemory);
 DEFINE_STAT(STAT_AsyncFileHandles);
 DEFINE_STAT(STAT_AsyncFileRequests);
 
-#endif //USE_NEW_ASYNC_IO
-
 int64 IFileHandle::Size()
 {
 	int64 Current = Tell();
diff --git a/Engine/Source/Runtime/Core/Private/HAL/MallocLeakDetection.cpp b/Engine/Source/Runtime/Core/Private/HAL/MallocLeakDetection.cpp
--- a/Engine/Source/Runtime/Core/Private/HAL/MallocLeakDetection.cpp
+++ b/Engine/Source/Runtime/Core/Private/HAL/MallocLeakDetection.cpp
@@ -425,7 +425,7 @@ int32 FMallocLeakDetection::DumpPotentialLeakers(const FMallocLeakReportOptions&
 			LOG_OUTPUT(TEXT("AllocSize: %i KB, Num: %i, FirstFrame %i, LastFrame %i, KnownDeleter: %d, Baseline memory %.3fMB  Leak Rate %.2fB/frame"),
 				Callstack.Size / 1024
 				, Callstack.Count
-				, Callstack.FirstFame
+				, Callstack.FirstFrame
 				, Callstack.LastFrame
 				, KnownDeleter
 				, Callstack.Baseline / 1024.0f / 1024.0f
@@ -579,7 +579,7 @@ int32 FMallocLeakDetection::DumpOpenCallstacks(const FMallocLeakReportOptions& O
 		LOG_OUTPUT(TEXT("\nAllocSize: %i KB, Num: %i, FirstFrame %i, LastFrame %i, KnownDeleter: %d")
 			, Callstack.Size / 1024
 			, Callstack.Count
-			, Callstack.FirstFame
+			, Callstack.FirstFrame
 			, Callstack.LastFrame
 			, KnownDeleter);
 
@@ -671,7 +671,7 @@ void FMallocLeakDetection::Malloc(void* Ptr, SIZE_T Size)
 				bRecursive = true;
 				FCallstackTrack Callstack;
 				FPlatformStackWalk::CaptureStackBackTrace(Callstack.CallStack, FCallstackTrack::Depth);
-				Callstack.FirstFame = GFrameCounter;
+				Callstack.FirstFrame = GFrameCounter;
 				Callstack.LastFrame = GFrameCounter;
 				Callstack.Size = Size;
 				AddCallstack(Callstack);
diff --git a/Engine/Source/Runtime/Core/Private/HAL/MallocLeakDetection.h b/Engine/Source/Runtime/Core/Private/HAL/MallocLeakDetection.h
--- a/Engine/Source/Runtime/Core/Private/HAL/MallocLeakDetection.h
+++ b/Engine/Source/Runtime/Core/Private/HAL/MallocLeakDetection.h
@@ -51,7 +51,7 @@ class FMallocLeakDetection
 		}
 		static const int32 Depth = 32;
 		uint64 CallStack[Depth];
-		uint32 FirstFame;
+		uint32 FirstFrame;
 		uint32 LastFrame;
 		uint64 Size;
 		uint32 Count;
diff --git a/Engine/Source/Runtime/Core/Private/HAL/PlatformFileManager.cpp b/Engine/Source/Runtime/Core/Private/HAL/PlatformFileManager.cpp
--- a/Engine/Source/Runtime/Core/Private/HAL/PlatformFileManager.cpp
+++ b/Engine/Source/Runtime/Core/Private/HAL/PlatformFileManager.cpp
@@ -99,6 +99,52 @@ IPlatformFile* FPlatformFileManager::GetPlatformFile(const TCHAR* Name)
 	return PlatformFile;
 }
 
+void FPlatformFileManager::RemovePlatformFile(IPlatformFile* PlatformFileToRemove)
+{
+	check(TopmostPlatformFile != nullptr);
+	check(PlatformFileToRemove != nullptr);
+
+	IPlatformFile* HigherLevelPlatformFile = nullptr;
+	IPlatformFile* FoundElement = nullptr;
+	for (FoundElement = TopmostPlatformFile; FoundElement && FoundElement != PlatformFileToRemove; FoundElement = FoundElement->GetLowerLevel())
+	{
+		HigherLevelPlatformFile = FoundElement;
+	}
+	check(FoundElement == PlatformFileToRemove);
+	if (HigherLevelPlatformFile)
+	{
+		check(HigherLevelPlatformFile->GetLowerLevel() == PlatformFileToRemove);
+		HigherLevelPlatformFile->SetLowerLevel(PlatformFileToRemove->GetLowerLevel());
+	}
+	else
+	{
+		check(TopmostPlatformFile == PlatformFileToRemove);
+		check(PlatformFileToRemove->GetLowerLevel());
+		SetPlatformFile(*PlatformFileToRemove->GetLowerLevel());
+	}
+}
+
+void FPlatformFileManager::InitializeNewAsyncIO()
+{
+	// Removed the cached file wrapper because it doesn't work well with EDL
+	if (GEventDrivenLoaderEnabled)
+	{
+		IPlatformFile* CachedWrapper = FindPlatformFile(FCachedReadPlatformFile::GetTypeName());
+		if (CachedWrapper)
+		{
+			RemovePlatformFile(CachedWrapper);
+		}
+	}
+	// Make sure all platform wrappers know about new async IO and EDL
+	if (GNewAsyncIO)
+	{
+		for (IPlatformFile* ChainElement = TopmostPlatformFile; ChainElement; ChainElement = ChainElement->GetLowerLevel())
+		{
+			ChainElement->InitializeNewAsyncIO();
+		}
+	}
+}
+
 FPlatformFileManager& FPlatformFileManager::Get()
 {
 	static FPlatformFileManager Singleton;
diff --git a/Engine/Source/Runtime/Core/Private/HAL/ThreadingBase.cpp b/Engine/Source/Runtime/Core/Private/HAL/ThreadingBase.cpp
--- a/Engine/Source/Runtime/Core/Private/HAL/ThreadingBase.cpp
+++ b/Engine/Source/Runtime/Core/Private/HAL/ThreadingBase.cpp
@@ -15,9 +15,7 @@ DECLARE_DWORD_COUNTER_STAT( TEXT( "ThreadPoolDummyCounter" ), STAT_ThreadPoolDum
 /** The global thread pool */
 FQueuedThreadPool* GThreadPool = nullptr;
 
-#if USE_NEW_ASYNC_IO
 FQueuedThreadPool* GIOThreadPool = nullptr;
-#endif // USE_NEW_ASYNC_IO
 
 #if WITH_EDITOR
 FQueuedThreadPool* GLargeThreadPool = nullptr;
@@ -779,56 +777,3 @@ void FTlsAutoCleanup::Register()
 		RunnableThread->TlsInstances.Add( this );
 	}
 }
-
-void FMultiReaderSingleWriterGT::LockRead()
-{
-	if (!IsInGameThread())
-	{		
-		SCOPE_CYCLE_COUNTER(STAT_Sleep);
-
-		FThreadIdleStats::FScopeIdle Scope;
-		while (!CanRead())
-		{
-			FPlatformProcess::SleepNoStats(0.0f);
-		}
-	}
-	CriticalSection.ReadCounter.Increment();
-}
-
-void FMultiReaderSingleWriterGT::UnlockRead()
-{
-	if (CriticalSection.ReadCounter.Decrement() != 0)
-	{
-		return;
-	}
-
-	if (!IsInGameThread())
-	{
-		FPlatformAtomics::InterlockedExchange(&CriticalSection.Action, NoAction);
-	}
-}
-
-void FMultiReaderSingleWriterGT::LockWrite()
-{
-	check(IsInGameThread());
-
-	{
-		SCOPE_CYCLE_COUNTER(STAT_Sleep);
-
-		FThreadIdleStats::FScopeIdle Scope;
-		while (!CanWrite())
-		{
-			FPlatformProcess::SleepNoStats(0.0f);
-		}
-	}
-
-	CriticalSection.WriteCounter.Increment();
-}
-
-void FMultiReaderSingleWriterGT::UnlockWrite()
-{
-	if (CriticalSection.WriteCounter.Decrement() == 0)
-	{
-		FPlatformAtomics::InterlockedExchange(&CriticalSection.Action, NoAction);
-	}
-}
diff --git a/Engine/Source/Runtime/Core/Private/Misc/AES.cpp b/Engine/Source/Runtime/Core/Private/Misc/AES.cpp
--- a/Engine/Source/Runtime/Core/Private/Misc/AES.cpp
+++ b/Engine/Source/Runtime/Core/Private/Misc/AES.cpp
@@ -1188,11 +1188,12 @@ void FAES::DecryptData(uint8 *Contents, uint32 NumBytes)
 #endif
 }
 
-void FAES::DecryptData( uint8 *Contents, uint32 NumBytes, ANSICHAR* Key )
+void FAES::DecryptData( uint8 *Contents, uint32 NumBytes, const ANSICHAR* Key )
 {
 	uint32 rk[RKLENGTH(AES_KEYBITS)] = { 0 };
 	int32 nrounds;
 
+	check(Key != nullptr);
 	checkf( ( NumBytes & ( AESBlockSize - 1 ) ) == 0, TEXT( "NumBytes needs to tbe a multiple of 16 bytes" ) );
 	checkf( TCString<ANSICHAR>::Strlen( Key ) >= KEYLENGTH( AES_KEYBITS ), TEXT( "AES_KEY needs to be at least %d characters" ), KEYLENGTH( AES_KEYBITS ) );
 
diff --git a/Engine/Source/Runtime/Core/Private/Misc/CoreDelegates.cpp b/Engine/Source/Runtime/Core/Private/Misc/CoreDelegates.cpp
--- a/Engine/Source/Runtime/Core/Private/Misc/CoreDelegates.cpp
+++ b/Engine/Source/Runtime/Core/Private/Misc/CoreDelegates.cpp
@@ -35,6 +35,18 @@ FCoreDelegates::FOnHandleSystemEnsure FCoreDelegates::OnHandleSystemEnsure;
 FCoreDelegates::FOnHandleSystemError FCoreDelegates::OnHandleSystemError;
 FCoreDelegates::FOnActorLabelChanged FCoreDelegates::OnActorLabelChanged;
 
+FCoreDelegates::FPakEncryptionKeyDelegate& FCoreDelegates::GetPakEncryptionKeyDelegate()
+{
+	static FPakEncryptionKeyDelegate PakEncryptionKeyDelegate;
+	return PakEncryptionKeyDelegate;
+}
+
+FCoreDelegates::FPakSigningKeysDelegate& FCoreDelegates::GetPakSigningKeysDelegate()
+{
+	static FPakSigningKeysDelegate PakSigningKeysDelegate;
+	return PakSigningKeysDelegate;
+}
+
 #if WITH_EDITOR
 	FSimpleMulticastDelegate FCoreDelegates::PreModal;
 	FSimpleMulticastDelegate FCoreDelegates::PostModal;
@@ -105,3 +117,20 @@ FCoreDelegates::FConfigReadyForUse FCoreDelegates::ConfigReadyForUse;
 
 FSimpleMulticastDelegate FCoreDelegates::OnOutOfMemory;
 FCoreDelegates::FGetOnScreenMessagesDelegate FCoreDelegates::OnGetOnScreenMessages;
+
+void RegisterEncryptionKey(const char* InEncryptionKey)
+{
+	FCoreDelegates::GetPakEncryptionKeyDelegate().BindLambda([InEncryptionKey]() { return InEncryptionKey; });
+}
+
+void RegisterPakSigningKeys(const char* InExponent, const char* InModulus)
+{
+	static FString Exponent(ANSI_TO_TCHAR(InExponent));
+	static FString Modulus(ANSI_TO_TCHAR(InModulus));
+
+	FCoreDelegates::GetPakSigningKeysDelegate().BindLambda([](FString& OutExponent, FString& OutModulus)
+	{
+		OutExponent = Exponent;
+		OutModulus = Modulus;
+	});
+}
\ No newline at end of file
diff --git a/Engine/Source/Runtime/Core/Private/Misc/CoreGlobals.cpp b/Engine/Source/Runtime/Core/Private/Misc/CoreGlobals.cpp
--- a/Engine/Source/Runtime/Core/Private/Misc/CoreGlobals.cpp
+++ b/Engine/Source/Runtime/Core/Private/Misc/CoreGlobals.cpp
@@ -225,6 +225,13 @@ double					GStartTime						= FPlatformTime::InitTiming();
 FString					GSystemStartTime;
 /** Whether we are still in the initial loading proces.														*/
 bool					GIsInitialLoad					= true;
+/* Whether we are using the new async IO */
+bool					GNewAsyncIO = false;
+/* Whether we are using the event driven loader */
+bool					GEventDrivenLoaderEnabled = false;
+//@todoio put this in some kind of API
+bool					GPakCache_AcceptPrecacheRequests = true;
+
 /** Steadily increasing frame counter.																		*/
 uint64					GFrameCounter					= 0;
 uint64					GLastGCFrame					= 0;
diff --git a/Engine/Source/Runtime/Core/Private/Misc/SecureHash.cpp b/Engine/Source/Runtime/Core/Private/Misc/SecureHash.cpp
--- a/Engine/Source/Runtime/Core/Private/Misc/SecureHash.cpp
+++ b/Engine/Source/Runtime/Core/Private/Misc/SecureHash.cpp
@@ -458,7 +458,7 @@ void FMD5::Decode( uint32* output, const uint8* input, int32 len )
 // 		(((uint32)input[j+2]) << 16) | (((uint32)input[j+3]) << 24);
 // }
 
-namespace LexicalConversion
+namespace Lex
 {
 	FString ToString(const FMD5Hash& Hash)
 	{
diff --git a/Engine/Source/Runtime/Core/Private/Misc/StringFormatter.cpp b/Engine/Source/Runtime/Core/Private/Misc/StringFormatter.cpp
--- a/Engine/Source/Runtime/Core/Private/Misc/StringFormatter.cpp
+++ b/Engine/Source/Runtime/Core/Private/Misc/StringFormatter.cpp
@@ -31,9 +31,9 @@ void AppendToString(const FStringFormatArg& Arg, FString& StringToAppendTo)
 {
 	switch(Arg.Type)
 	{
-		case FStringFormatArg::Int: 			StringToAppendTo.Append(LexicalConversion::ToString(Arg.IntValue)); break;
-		case FStringFormatArg::UInt: 			StringToAppendTo.Append(LexicalConversion::ToString(Arg.UIntValue)); break;
-		case FStringFormatArg::Double: 			StringToAppendTo.Append(LexicalConversion::ToString(Arg.DoubleValue)); break;
+		case FStringFormatArg::Int: 			StringToAppendTo.Append(Lex::ToString(Arg.IntValue)); break;
+		case FStringFormatArg::UInt: 			StringToAppendTo.Append(Lex::ToString(Arg.UIntValue)); break;
+		case FStringFormatArg::Double: 			StringToAppendTo.Append(Lex::ToString(Arg.DoubleValue)); break;
 		case FStringFormatArg::String: 			StringToAppendTo.AppendChars(*Arg.StringValue, Arg.StringValue.Len()); break;
 		case FStringFormatArg::StringLiteral: 	StringToAppendTo += Arg.StringLiteralValue; break;
 	}
diff --git a/Engine/Source/Runtime/Core/Private/Modules/ModuleManager.cpp b/Engine/Source/Runtime/Core/Private/Modules/ModuleManager.cpp
--- a/Engine/Source/Runtime/Core/Private/Modules/ModuleManager.cpp
+++ b/Engine/Source/Runtime/Core/Private/Modules/ModuleManager.cpp
@@ -25,7 +25,7 @@ FModuleManager::ModuleInfoPtr FModuleManager::FindModule(FName InModuleName)
 {
 	FModuleManager::ModuleInfoPtr Result = nullptr;
 
-	FReadScopeLock Lock(&ModulesCriticalSection);
+	FScopeLock Lock(&ModulesCriticalSection);
 	if ( FModuleManager::ModuleInfoRef* FoundModule = Modules.Find(InModuleName))
 	{
 		Result = *FoundModule;
@@ -36,7 +36,7 @@ FModuleManager::ModuleInfoPtr FModuleManager::FindModule(FName InModuleName)
 
 FModuleManager::ModuleInfoRef FModuleManager::FindModuleChecked(FName InModuleName)
 {
-	FReadScopeLock Lock(&ModulesCriticalSection);
+	FScopeLock Lock(&ModulesCriticalSection);
 	return Modules.FindChecked(InModuleName);
 }
 
@@ -203,7 +203,7 @@ bool FindNewestModuleFile(TArray<FString>& FilesToSearch, const FDateTime& Newer
 void FModuleManager::AddModuleToModulesList(const FName InModuleName, FModuleManager::ModuleInfoRef& InModuleInfo)
 {
 	{
-		FWriteScopeLock Lock(&ModulesCriticalSection);
+		FScopeLock Lock(&ModulesCriticalSection);
 
 		// Update hash table
 		Modules.Add(InModuleName, InModuleInfo);
@@ -852,7 +852,7 @@ bool FModuleManager::QueryModule( const FName InModuleName, FModuleStatus& OutMo
 void FModuleManager::QueryModules( TArray< FModuleStatus >& OutModuleStatuses ) const
 {
 	OutModuleStatuses.Reset();
-	FReadScopeLock Lock(&ModulesCriticalSection);
+	FScopeLock Lock(&ModulesCriticalSection);
 	for (const TPair<FName, ModuleInfoRef>& ModuleIt : Modules)
 	{
 		const FModuleInfo& CurModule = *ModuleIt.Value;
@@ -1351,6 +1351,8 @@ FModuleManager::ModuleInfoRef FModuleManager::GetOrCreateModule(FName InModuleNa
 
 int32 FModuleManager::GetModuleCount() const
 {
-	FReadScopeLock Lock(&ModulesCriticalSection);
+	// Theoretically thread safe but by the time we return new modules could've been added
+	// so no point in locking here. ModulesCriticalSection should be locked by the caller
+	// if it wants to rely on the returned value.
 	return Modules.Num();
 }
diff --git a/Engine/Source/Runtime/Core/Private/Serialization/Archive.cpp b/Engine/Source/Runtime/Core/Private/Serialization/Archive.cpp
--- a/Engine/Source/Runtime/Core/Private/Serialization/Archive.cpp
+++ b/Engine/Source/Runtime/Core/Private/Serialization/Archive.cpp
@@ -57,6 +57,7 @@ FArchive::FArchive(const FArchive& ArchiveToCopy)
 	// Don't know why this is set to false, but this is what the original copying code did
 	ArIsFilterEditorOnly  = false;
 
+	bCustomVersionsAreReset = ArchiveToCopy.bCustomVersionsAreReset;
 	CustomVersionContainer = new FCustomVersionContainer(*ArchiveToCopy.CustomVersionContainer);
 }
 
@@ -71,6 +72,7 @@ FArchive& FArchive::operator=(const FArchive& ArchiveToCopy)
 	// Don't know why this is set to false, but this is what the original copying code did
 	ArIsFilterEditorOnly  = false;
 
+	bCustomVersionsAreReset = ArchiveToCopy.bCustomVersionsAreReset;
 	*CustomVersionContainer = *ArchiveToCopy.CustomVersionContainer;
 	return *this;
 }
@@ -272,6 +274,40 @@ FArchive& FArchive::operator<<(struct FWeakObjectPtr& Value)
 	return *this;
 }
 
+#if WITH_EDITOR
+FArchive& FArchive::operator<<( bool& D )
+{
+	// Serialize bool as if it were UBOOL (legacy, 32 bit int).
+	uint32 OldUBoolValue;
+#if DEVIRTUALIZE_FLinkerLoad_Serialize
+	const uint8 * RESTRICT Src = this->ActiveFPLB->StartFastPathLoadBuffer;
+	if (Src + sizeof(uint32) <= this->ActiveFPLB->EndFastPathLoadBuffer)
+	{
+#if PLATFORM_SUPPORTS_UNALIGNED_INT_LOADS
+		OldUBoolValue = *(uint32* RESTRICT)Src;
+#else
+		static_assert(sizeof(uint32) == 4, "assuming sizeof(uint32) == 4");
+		OldUBoolValue = Src[0] | Src[1] | Src[2] | Src[3];
+#endif
+		this->ActiveFPLB->StartFastPathLoadBuffer += 4;
+	}
+	else
+#endif
+	{
+		OldUBoolValue = D ? 1 : 0;
+		this->Serialize(&OldUBoolValue, sizeof(OldUBoolValue));
+	}
+	if (OldUBoolValue > 1)
+	{
+		UE_LOG( LogSerialization, Error, TEXT("Invalid boolean encountered while reading archive - stream is most likely corrupted."));
+
+		this->ArIsError = true;
+	}
+	D = !!OldUBoolValue;
+	return *this;
+}
+#endif
+
 const FCustomVersionContainer& FArchive::GetCustomVersions() const
 {
 	if (bCustomVersionsAreReset)
diff --git a/Engine/Source/Runtime/Core/Private/Serialization/AsyncIOSystemBase.cpp b/Engine/Source/Runtime/Core/Private/Serialization/AsyncIOSystemBase.cpp
--- a/Engine/Source/Runtime/Core/Private/Serialization/AsyncIOSystemBase.cpp
+++ b/Engine/Source/Runtime/Core/Private/Serialization/AsyncIOSystemBase.cpp
@@ -1,7 +1,6 @@
 // Copyright 1998-2017 Epic Games, Inc. All Rights Reserved.
 
 #include "Serialization/AsyncIOSystemBase.h"
-#if !USE_NEW_ASYNC_IO
 
 /*=============================================================================
 	AsyncIOSystemBase.h: Base implementation of the async IO system
@@ -892,6 +891,8 @@ static FAsyncIOSystemBase* AsyncIOSystem = NULL;
 
 FIOSystem& FIOSystem::Get()
 {
+	check(GConfig); // Otherwise GNewAsyncIO may not be initialized properly
+	check(!GNewAsyncIO); // We don't want to use this old class with the new async IO
 	if (!AsyncIOThread)
 	{
 		check(!AsyncIOSystem);
@@ -928,5 +929,3 @@ bool FIOSystem::HasShutdown()
 {
 	return AsyncIOThread == nullptr || AsyncIOThread == (FRunnableThread*)-1;
 }
-
-#endif
diff --git a/Engine/Source/Runtime/Core/Private/Serialization/AsyncIOSystemBase.h b/Engine/Source/Runtime/Core/Private/Serialization/AsyncIOSystemBase.h
--- a/Engine/Source/Runtime/Core/Private/Serialization/AsyncIOSystemBase.h
+++ b/Engine/Source/Runtime/Core/Private/Serialization/AsyncIOSystemBase.h
@@ -16,8 +16,6 @@
 #include "Misc/SingleThreadRunnable.h"
 #include "HAL/IOBase.h"
 
-#if !USE_NEW_ASYNC_IO
-
 /**
  * Base implementation of an async IO system allowing most of the code to be shared across platforms.
  */
@@ -449,5 +447,3 @@ protected:
 	/** Low level file system that we use for our requests.											*/
 	IPlatformFile&					LowLevel;
 };
-
-#endif
diff --git a/Engine/Source/Runtime/Core/Public/Async/AsyncFileHandle.h b/Engine/Source/Runtime/Core/Public/Async/AsyncFileHandle.h
--- a/Engine/Source/Runtime/Core/Public/Async/AsyncFileHandle.h
+++ b/Engine/Source/Runtime/Core/Public/Async/AsyncFileHandle.h
@@ -7,17 +7,10 @@
 #include "Templates/Function.h"
 #include "Stats/Stats.h"
 #include "GenericPlatform/GenericPlatformFile.h"
-
-#if !defined(USE_NEW_ASYNC_IO)
-#error "USE_NEW_ASYNC_IO must be defined."
-#endif
+#include "Function.h"
 
 class IAsyncReadRequest;
 
-#if USE_NEW_ASYNC_IO
-
-#include "Function.h"
-
 DECLARE_MEMORY_STAT_EXTERN(TEXT("Async File Handle Memory"), STAT_AsyncFileMemory, STATGROUP_Memory, CORE_API);
 DECLARE_DWORD_ACCUMULATOR_STAT_EXTERN(TEXT("Num Async File Handles"), STAT_AsyncFileHandles, STATGROUP_Memory, CORE_API);
 DECLARE_DWORD_ACCUMULATOR_STAT_EXTERN(TEXT("Num Async File Requests"), STAT_AsyncFileRequests, STATGROUP_Memory, CORE_API);
@@ -193,6 +186,3 @@ public:
 	**/
 	virtual IAsyncReadRequest* ReadRequest(int64 Offset, int64 BytesToRead, EAsyncIOPriority Priority = AIOP_Normal, FAsyncFileCallBack* CompleteCallback = nullptr, uint8* UserSuppliedMemory = nullptr) = 0;
 };
-
-
-#endif // USE_NEW_ASYNC_IO
diff --git a/Engine/Source/Runtime/Core/Public/Containers/Algo/IsSorted.h b/Engine/Source/Runtime/Core/Public/Containers/Algo/IsSorted.h
--- a/Engine/Source/Runtime/Core/Public/Containers/Algo/IsSorted.h
+++ b/Engine/Source/Runtime/Core/Public/Containers/Algo/IsSorted.h
@@ -3,6 +3,7 @@
 #pragma once
 
 #include "CoreTypes.h"
+#include "Traits/ElementType.h"
 
 namespace AlgoImpl
 {
@@ -114,7 +115,7 @@ namespace Algo
 	template <typename ContainerType>
 	FORCEINLINE bool IsSorted(const ContainerType& Container)
 	{
-		return AlgoImpl::IsSorted(Container.GetData(), Container.Num(), AlgoImpl::TLess<typename ContainerType::ElementType>());
+		return AlgoImpl::IsSorted(Container.GetData(), Container.Num(), AlgoImpl::TLess<typename TElementType<ContainerType>::Type>());
 	}
 
 	/**
diff --git a/Engine/Source/Runtime/Core/Public/Containers/UnrealString.h b/Engine/Source/Runtime/Core/Public/Containers/UnrealString.h
--- a/Engine/Source/Runtime/Core/Public/Containers/UnrealString.h
+++ b/Engine/Source/Runtime/Core/Public/Containers/UnrealString.h
@@ -1875,7 +1875,7 @@ inline int32 HexToBytes( const FString& HexString, uint8* OutBytes )
 }
 
 /** Namespace that houses lexical conversion for various types. User defined conversions can be implemented externally */
-namespace LexicalConversion
+namespace Lex
 {
 	/**
 	 *	Expected functions in this namespace are as follows:
@@ -1920,6 +1920,16 @@ namespace LexicalConversion
 		return Value ? TEXT("true") : TEXT("false");
 	}
 
+	FORCEINLINE FString ToString(FString&& Str)
+	{
+		return MoveTemp(Str);
+	}
+
+	FORCEINLINE FString ToString(const FString& Str)
+	{
+		return Str;
+	}
+
 	/** Helper template to convert to sanitized strings */
 	template<typename T>
 	FString ToSanitizedString(const T& Value)
@@ -1957,17 +1967,21 @@ namespace LexicalConversion
 	}
 }
 
-/** Shorthand legacy use for LexicalConversion functions */
+// Deprecated alias for old LexicalConversion namespace.
+// Namespace alias deprecation doesn't work on our compilers, so we can't actually mark it with the DEPRECATED() macro.
+namespace LexicalConversion = Lex;
+
+/** Shorthand legacy use for Lex functions */
 template<typename T>
 struct TTypeToString
 {
-	static FString ToString(const T& Value)				{ return LexicalConversion::ToString(Value); }
-	static FString ToSanitizedString(const T& Value)	{ return LexicalConversion::ToSanitizedString(Value); }
+	static FString ToString(const T& Value)				{ return Lex::ToString(Value); }
+	static FString ToSanitizedString(const T& Value)	{ return Lex::ToSanitizedString(Value); }
 };
 template<typename T>
 struct TTypeFromString
 {
-	static void FromString(T& Value, const TCHAR* Buffer) { return LexicalConversion::FromString(Value, Buffer); }
+	static void FromString(T& Value, const TCHAR* Buffer) { return Lex::FromString(Value, Buffer); }
 };
 
 /*----------------------------------------------------------------------------
diff --git a/Engine/Source/Runtime/Core/Public/CoreGlobals.h b/Engine/Source/Runtime/Core/Public/CoreGlobals.h
--- a/Engine/Source/Runtime/Core/Public/CoreGlobals.h
+++ b/Engine/Source/Runtime/Core/Public/CoreGlobals.h
@@ -288,6 +288,15 @@ extern CORE_API FString GSystemStartTime;
 /** Whether we are still in the initial loading process. */
 extern CORE_API bool GIsInitialLoad;
 
+/* Whether we are using the new async IO */
+extern CORE_API bool GNewAsyncIO;
+
+/* Whether we are using the event driven loader */
+extern CORE_API bool GEventDrivenLoaderEnabled;
+
+//@todoio put this in some kind of API
+extern CORE_API bool GPakCache_AcceptPrecacheRequests;
+
 #if WITH_HOT_RELOAD_CTORS
 /** true when we are retrieving VTablePtr from UClass */
 extern CORE_API bool GIsRetrievingVTablePtr;
diff --git a/Engine/Source/Runtime/Core/Public/Delegates/DelegateSignatureImpl.inl b/Engine/Source/Runtime/Core/Public/Delegates/DelegateSignatureImpl.inl
--- a/Engine/Source/Runtime/Core/Public/Delegates/DelegateSignatureImpl.inl
+++ b/Engine/Source/Runtime/Core/Public/Delegates/DelegateSignatureImpl.inl
@@ -1055,7 +1055,7 @@ public:
 		// Store the function name.  The incoming function name was generated by a macro and includes the method's class name.
 		this->FunctionName = InFunctionName;
 
-		ensureMsgf(this->IsBound(), TEXT("Unable to bind delegate to '%s' (function might not be marked as a UFUNCTION)"), *InFunctionName.ToString());
+		ensureMsgf(this->IsBound(), TEXT("Unable to bind delegate to '%s' (function might not be marked as a UFUNCTION or object may be pending kill)"), *InFunctionName.ToString());
 	}
 
 	friend uint32 GetTypeHash(const TBaseDynamicDelegate& Key)
diff --git a/Engine/Source/Runtime/Core/Public/Delegates/Tuple.h b/Engine/Source/Runtime/Core/Public/Delegates/Tuple.h
--- a/Engine/Source/Runtime/Core/Public/Delegates/Tuple.h
+++ b/Engine/Source/Runtime/Core/Public/Delegates/Tuple.h
@@ -26,6 +26,11 @@ struct TNthTypeFromParameterPack<0, T, OtherTypes...>
 template <typename T, uint32 Index>
 struct TTupleElement
 {
+	TTupleElement()
+		: Value()
+	{
+	}
+
 	template <typename... ArgTypes>
 	explicit TTupleElement(ArgTypes&&... Args)
 		: Value(Forward<ArgTypes>(Args)...)
@@ -75,6 +80,11 @@ struct TTupleImpl;
 template <uint32... Indices, typename... Types>
 struct TTupleImpl<TIntegerSequence<uint32, Indices...>, Types...> : TTupleElement<Types, Indices>...
 {
+	TTupleImpl()
+		: TTupleElement<Types, Indices>()...
+	{
+	}
+
 	template <typename... ArgTypes>
 	explicit TTupleImpl(ArgTypes&&... Args)
 		: TTupleElement<Types, Indices>(Forward<ArgTypes>(Args))...
diff --git a/Engine/Source/Runtime/Core/Public/GenericPlatform/GenericPlatformFile.h b/Engine/Source/Runtime/Core/Public/GenericPlatform/GenericPlatformFile.h
--- a/Engine/Source/Runtime/Core/Public/GenericPlatform/GenericPlatformFile.h
+++ b/Engine/Source/Runtime/Core/Public/GenericPlatform/GenericPlatformFile.h
@@ -14,10 +14,6 @@
 
 class IAsyncReadFileHandle;
 
-#if USE_NEW_ASYNC_IO
-class IAsyncReadFileHandle;
-#endif // USE_NEW_ASYNC_IO
-
 /**
 * Enum for async IO priorities.
 */
@@ -211,6 +207,11 @@ public:
 	 */
 	virtual void		InitializeAfterSetActive() { }
 
+	/**
+	* Performs initialization of the platform file after the new async IO has been enabled
+	*/
+	virtual void		InitializeNewAsyncIO() { }
+
 	/**
 	 * Identifies any platform specific paths that are guaranteed to be local (i.e. cache, scratch space)
 	 */
@@ -232,7 +233,9 @@ public:
 
 	/** Gets the platform file wrapped by this file. */
 	virtual IPlatformFile* GetLowerLevel() = 0;
-		/** Gets this platform file type name. */
+	/** Sets the platform file wrapped by this file. */
+	virtual void SetLowerLevel(IPlatformFile* NewLowerLevel) = 0;
+	/** Gets this platform file type name. */
 	virtual const TCHAR* GetName() const = 0;
 	/** Return true if the file exists. **/
 	virtual bool		FileExists(const TCHAR* Filename) = 0;
@@ -324,14 +327,12 @@ public:
 	/////////// Generally, these do not need to be implemented per platform.
 	/////////////////////////////////////////////////////////////////////////////////////////////////////////
 
-#if USE_NEW_ASYNC_IO
 	/** Open a file for async reading. This call does not hit the disk or block.
 	*
 	* @param Filename file to be opened
 	* @return Close the file by delete'ing the handle. A non-null return value does not mean the file exists, since that may not be determined yet.
 	*/
 	virtual IAsyncReadFileHandle* OpenAsyncRead(const TCHAR* Filename);
-#endif // USE_NEW_ASYNC_IO
 
 	virtual void GetTimeStampPair(const TCHAR* PathA, const TCHAR* PathB, FDateTime& OutTimeStampA, FDateTime& OutTimeStampB);
 
@@ -442,6 +443,10 @@ public:
 	{
 		return nullptr;
 	}
+	virtual void SetLowerLevel(IPlatformFile* NewLowerLevel) override
+	{
+		check(false); // can't override wrapped platform file for physical platform file
+	}
 	virtual const TCHAR* GetName() const override
 	{
 		return IPlatformFile::GetPhysicalTypeName();
diff --git a/Engine/Source/Runtime/Core/Public/HAL/IOBase.h b/Engine/Source/Runtime/Core/Public/HAL/IOBase.h
--- a/Engine/Source/Runtime/Core/Public/HAL/IOBase.h
+++ b/Engine/Source/Runtime/Core/Public/HAL/IOBase.h
@@ -5,7 +5,6 @@
 #include "CoreTypes.h"
 #include "GenericPlatform/GenericPlatformFile.h"
 
-#if !USE_NEW_ASYNC_IO
 /*-----------------------------------------------------------------------------
 	FIOSystem.
 -----------------------------------------------------------------------------*/
@@ -136,4 +135,3 @@ struct CORE_API FIOSystem
 	virtual void FlushLog() = 0;
 };
 
-#endif
diff --git a/Engine/Source/Runtime/Core/Public/HAL/IPlatformFileCachedWrapper.h b/Engine/Source/Runtime/Core/Public/HAL/IPlatformFileCachedWrapper.h
--- a/Engine/Source/Runtime/Core/Public/HAL/IPlatformFileCachedWrapper.h
+++ b/Engine/Source/Runtime/Core/Public/HAL/IPlatformFileCachedWrapper.h
@@ -245,9 +245,6 @@ public:
 	}
 	virtual bool ShouldBeUsed(IPlatformFile* Inner, const TCHAR* CmdLine) const override
 	{
-#if USE_EVENT_DRIVEN_ASYNC_LOAD
-		return false; // this stuff is awful for the EDL
-#else
 		// Default to false on platforms that already do platform file level caching
 		bool bResult = !PLATFORM_PS4 && !PLATFORM_WINDOWS && FPlatformProperties::RequiresCookedData();
 
@@ -268,13 +265,15 @@ public:
 		}
 #endif
 		return bResult;
-#endif
 	}
 	IPlatformFile* GetLowerLevel() override
 	{
 		return LowerLevel;
 	}
-
+	virtual void SetLowerLevel(IPlatformFile* NewLowerLevel) override
+	{
+		LowerLevel = NewLowerLevel;
+	}
 	virtual const TCHAR* GetName() const override
 	{
 		return FCachedReadPlatformFile::GetTypeName();
@@ -397,10 +396,8 @@ public:
 	{
 		return LowerLevel->SendMessageToServer(Message, Handler);
 	}
-#if USE_NEW_ASYNC_IO
 	virtual IAsyncReadFileHandle* OpenAsyncRead(const TCHAR* Filename) override
 	{
 		return LowerLevel->OpenAsyncRead(Filename);
 	}
-#endif // USE_NEW_ASYNC_IO
 };
diff --git a/Engine/Source/Runtime/Core/Public/HAL/IPlatformFileLogWrapper.h b/Engine/Source/Runtime/Core/Public/HAL/IPlatformFileLogWrapper.h
--- a/Engine/Source/Runtime/Core/Public/HAL/IPlatformFileLogWrapper.h
+++ b/Engine/Source/Runtime/Core/Public/HAL/IPlatformFileLogWrapper.h
@@ -131,7 +131,10 @@ public:
 	{
 		return LowerLevel;
 	}
-
+	virtual void SetLowerLevel(IPlatformFile* NewLowerLevel) override
+	{
+		LowerLevel = NewLowerLevel;
+	}
 	virtual const TCHAR* GetName() const override
 	{
 		return FLoggedPlatformFile::GetTypeName();
@@ -443,7 +446,6 @@ public:
 	}
 	void HandleDumpCommand(const TCHAR* Cmd, FOutputDevice& Ar);
 #endif
-#if USE_NEW_ASYNC_IO
 	virtual IAsyncReadFileHandle* OpenAsyncRead(const TCHAR* Filename) override
 	{
 		FString DataStr = FString::Printf(TEXT("OpenAsyncRead %s"), Filename);
@@ -456,6 +458,4 @@ public:
 		//@todo no wrapped logging for async file handles (yet)
 		return Result;
 	}
-#endif // USE_NEW_ASYNC_IO
-
 };
diff --git a/Engine/Source/Runtime/Core/Public/HAL/IPlatformFileOpenLogWrapper.h b/Engine/Source/Runtime/Core/Public/HAL/IPlatformFileOpenLogWrapper.h
--- a/Engine/Source/Runtime/Core/Public/HAL/IPlatformFileOpenLogWrapper.h
+++ b/Engine/Source/Runtime/Core/Public/HAL/IPlatformFileOpenLogWrapper.h
@@ -96,6 +96,10 @@ public:
 	{
 		return LowerLevel;
 	}
+	virtual void SetLowerLevel(IPlatformFile* NewLowerLevel) override
+	{
+		LowerLevel = NewLowerLevel;
+	}
 	static const TCHAR* GetTypeName()
 	{
 		return TEXT("FileOpenLog");
@@ -227,7 +231,6 @@ public:
 	{
 		return LowerLevel->SendMessageToServer(Message, Handler);
 	}
-#if USE_NEW_ASYNC_IO
 	virtual IAsyncReadFileHandle* OpenAsyncRead(const TCHAR* Filename) override
 	{
 		IAsyncReadFileHandle* Result = LowerLevel->OpenAsyncRead(Filename);
@@ -247,7 +250,6 @@ public:
 		}
 		return Result;
 	}
-#endif // USE_NEW_ASYNC_IO
 };
 
 #endif // !UE_BUILD_SHIPPING
diff --git a/Engine/Source/Runtime/Core/Public/HAL/IPlatformFileProfilerWrapper.h b/Engine/Source/Runtime/Core/Public/HAL/IPlatformFileProfilerWrapper.h
--- a/Engine/Source/Runtime/Core/Public/HAL/IPlatformFileProfilerWrapper.h
+++ b/Engine/Source/Runtime/Core/Public/HAL/IPlatformFileProfilerWrapper.h
@@ -251,6 +251,10 @@ public:
 	{
 		return LowerLevel;
 	}
+	virtual void SetLowerLevel(IPlatformFile* NewLowerLevel) override
+	{
+		LowerLevel = NewLowerLevel;
+	}
 
 	double GetStartTime() const
 	{
@@ -483,12 +487,10 @@ public:
 		OpStat->Duration += FPlatformTime::Seconds() * 1000.0 - OpStat->LastOpTime;
 		return Result;
 	}
-#if USE_NEW_ASYNC_IO
 	virtual IAsyncReadFileHandle* OpenAsyncRead(const TCHAR* Filename) override
 	{
 		return LowerLevel->OpenAsyncRead(Filename);
 	}
-#endif // USE_NEW_ASYNC_IO
 
 	//static void CreateProfileVisualizer
 };
@@ -588,6 +590,10 @@ public:
 	{
 		return LowerLevel;
 	}
+	virtual void SetLowerLevel(IPlatformFile* NewLowerLevel) override
+	{
+		LowerLevel = NewLowerLevel;
+	}
 	static const TCHAR* GetTypeName()
 	{
 		return TEXT("FileReadStats");
@@ -691,13 +697,11 @@ public:
 	{
 		return LowerLevel->CopyFile( To, From, ReadFlags, WriteFlags );
 	}
-#if USE_NEW_ASYNC_IO
 	virtual IAsyncReadFileHandle* OpenAsyncRead(const TCHAR* Filename) override
 	{
 		//@todo no wrapped async handles (yet)
 		return LowerLevel->OpenAsyncRead(Filename);
 	}
-#endif // USE_NEW_ASYNC_IO
 };
 
 
diff --git a/Engine/Source/Runtime/Core/Public/HAL/PlatformFilemanager.h b/Engine/Source/Runtime/Core/Public/HAL/PlatformFilemanager.h
--- a/Engine/Source/Runtime/Core/Public/HAL/PlatformFilemanager.h
+++ b/Engine/Source/Runtime/Core/Public/HAL/PlatformFilemanager.h
@@ -48,8 +48,27 @@ public:
 	 */
 	IPlatformFile* GetPlatformFile( const TCHAR* Name );
 
+	/**
+	* Permorms additional initialization when the new async IO is enabled.
+	*/
+	void InitializeNewAsyncIO();
+
 	/**
 	 * Gets FPlatformFileManager Singleton.
 	 */
 	static FPlatformFileManager& Get( );
+
+private:
+
+	/**
+	* Removes the specified file wrapper from the platform file wrapper chain.
+	*
+	* THIS IS EXTREMELY DANGEROUS AFTER THE ENGINE HAS BEEN INITIALIZED AS WE MAY BE MODIFYING
+	* THE WRAPPER CHAIN WHILE THINGS ARE BEING LOADED
+	*
+	* @param Name of the platform file to create.
+	* @return Platform file instance of the platform file type was found, nullptr otherwise.
+	*/
+	void RemovePlatformFile(IPlatformFile* PlatformFileToRemove);
+
 };
diff --git a/Engine/Source/Runtime/Core/Public/Math/UnitConversion.h b/Engine/Source/Runtime/Core/Public/Math/UnitConversion.h
--- a/Engine/Source/Runtime/Core/Public/Math/UnitConversion.h
+++ b/Engine/Source/Runtime/Core/Public/Math/UnitConversion.h
@@ -193,7 +193,7 @@ bool operator==(const FNumericUnit<NumericType>& LHS, const FNumericUnit<OtherTy
 template<typename NumericType, typename OtherType>
 bool operator!=(const FNumericUnit<NumericType>& LHS, const FNumericUnit<OtherType>& RHS);
 
-namespace LexicalConversion
+namespace Lex
 {
 	template<typename T>
 	FString ToString(const FNumericUnit<T>& NumericUnit);
diff --git a/Engine/Source/Runtime/Core/Public/Math/UnitConversion.inl b/Engine/Source/Runtime/Core/Public/Math/UnitConversion.inl
--- a/Engine/Source/Runtime/Core/Public/Math/UnitConversion.inl
+++ b/Engine/Source/Runtime/Core/Public/Math/UnitConversion.inl
@@ -292,7 +292,7 @@ TOptional<FNumericUnit<NumericType>> FNumericUnit<NumericType>::TryParseString(c
 	}
 
 	NumericType NewValue;
-	LexicalConversion::FromString(NewValue, InSource);
+	Lex::FromString(NewValue, InSource);
 
 	// Now parse the units
 	while(FChar::IsWhitespace(*NumberEnd)) ++NumberEnd;
@@ -412,7 +412,7 @@ template <typename NumericType>
 struct TNumericLimits<FNumericUnit<NumericType>> : public TNumericLimits<NumericType>
 { };
 
-namespace LexicalConversion
+namespace Lex
 {
 	template<typename T>
 	FString ToString(const FNumericUnit<T>& NumericUnit)
diff --git a/Engine/Source/Runtime/Core/Public/Misc/AES.h b/Engine/Source/Runtime/Core/Public/Misc/AES.h
--- a/Engine/Source/Runtime/Core/Public/Misc/AES.h
+++ b/Engine/Source/Runtime/Core/Public/Misc/AES.h
@@ -30,5 +30,5 @@ struct CORE_API FAES
 	 * @param NumBytes the size of the buffer
 	 * @param Key a null terminated string that is a 32 byte multiple length
 	 */
-	static void DecryptData(uint8* Contents, uint32 NumBytes, ANSICHAR* Key);
+	static void DecryptData(uint8* Contents, uint32 NumBytes, const ANSICHAR* Key);
 };
diff --git a/Engine/Source/Runtime/Core/Public/Misc/CoreDelegates.h b/Engine/Source/Runtime/Core/Public/Misc/CoreDelegates.h
--- a/Engine/Source/Runtime/Core/Public/Misc/CoreDelegates.h
+++ b/Engine/Source/Runtime/Core/Public/Misc/CoreDelegates.h
@@ -71,6 +71,12 @@ public:
 	// Callback for handling accepting invitations - generally for engine code
 	DECLARE_MULTICAST_DELEGATE_TwoParams(FOnInviteAccepted, const FString&, const FString&);
 
+	// Callback for accessing pak encryption key, if it exists
+	DECLARE_DELEGATE_RetVal(const ANSICHAR*, FPakEncryptionKeyDelegate);
+
+	// Callback for gathering pak signing keys, if they exist
+	DECLARE_DELEGATE_TwoParams(FPakSigningKeysDelegate, FString&, FString&);
+
 	// Callback for handling the Controller connection / disconnection
 	// first param is true for a connection, false for a disconnection.
 	// second param is UserID, third is UserIndex / ControllerId.
@@ -106,6 +112,9 @@ public:
 	// Called when an actor label is changed
 	static FOnActorLabelChanged OnActorLabelChanged;
 
+	static FPakEncryptionKeyDelegate& GetPakEncryptionKeyDelegate();
+	static FPakSigningKeysDelegate& GetPakSigningKeysDelegate();
+
 #if WITH_EDITOR
 	// Called before the editor displays a modal window, allowing other windows the opportunity to disable themselves to avoid reentrant calls
 	static FSimpleMulticastDelegate PreModal;
diff --git a/Engine/Source/Runtime/Core/Public/Misc/CoreMiscDefines.h b/Engine/Source/Runtime/Core/Public/Misc/CoreMiscDefines.h
--- a/Engine/Source/Runtime/Core/Public/Misc/CoreMiscDefines.h
+++ b/Engine/Source/Runtime/Core/Public/Misc/CoreMiscDefines.h
@@ -27,13 +27,6 @@
 /** This controls if metadata for compiled in classes is unpacked and setup at boot time. Meta data is not normally used except by the editor. **/
 #define WITH_METADATA (WITH_EDITORONLY_DATA && WITH_EDITOR)
 
-/** Setup the new async IO system */
-#if !defined(USE_NEW_ASYNC_IO)
-#define USE_NEW_ASYNC_IO 0
-#elif USE_NEW_ASYNC_IO && WITH_EDITORONLY_DATA
-#error USE_NEW_ASYNC_IO can not be used with WITH_EDITORONLY_DATA
-#endif
-
 // Set up optimization control macros, now that we have both the build settings and the platform macros
 #define PRAGMA_DISABLE_OPTIMIZATION		PRAGMA_DISABLE_OPTIMIZATION_ACTUAL
 #if UE_BUILD_DEBUG
diff --git a/Engine/Source/Runtime/Core/Public/Misc/QueuedThreadPool.h b/Engine/Source/Runtime/Core/Public/Misc/QueuedThreadPool.h
--- a/Engine/Source/Runtime/Core/Public/Misc/QueuedThreadPool.h
+++ b/Engine/Source/Runtime/Core/Public/Misc/QueuedThreadPool.h
@@ -89,9 +89,7 @@ public:
  */
 extern CORE_API FQueuedThreadPool* GThreadPool;
 
-#if USE_NEW_ASYNC_IO
 extern CORE_API FQueuedThreadPool* GIOThreadPool;
-#endif // USE_NEW_ASYNC_IO
 
 #if WITH_EDITOR
 extern CORE_API FQueuedThreadPool* GLargeThreadPool;
diff --git a/Engine/Source/Runtime/Core/Public/Misc/SecureHash.h b/Engine/Source/Runtime/Core/Public/Misc/SecureHash.h
--- a/Engine/Source/Runtime/Core/Public/Misc/SecureHash.h
+++ b/Engine/Source/Runtime/Core/Public/Misc/SecureHash.h
@@ -99,7 +99,7 @@ private:
 
 struct FMD5Hash;
 
-namespace LexicalConversion
+namespace Lex
 {
 	CORE_API FString ToString(const FMD5Hash& Hash);
 	CORE_API void FromString(FMD5Hash& Hash, const TCHAR* Buffer);
@@ -159,8 +159,8 @@ private:
 	/** The bytes this hash comprises */
 	uint8 Bytes[16];
 
-	friend FString LexicalConversion::ToString(const FMD5Hash&);
-	friend void LexicalConversion::FromString(FMD5Hash& Hash, const TCHAR*);
+	friend FString Lex::ToString(const FMD5Hash&);
+	friend void Lex::FromString(FMD5Hash& Hash, const TCHAR*);
 };
 
 /*-----------------------------------------------------------------------------
diff --git a/Engine/Source/Runtime/Core/Public/Modules/ModuleManager.h b/Engine/Source/Runtime/Core/Public/Modules/ModuleManager.h
--- a/Engine/Source/Runtime/Core/Public/Modules/ModuleManager.h
+++ b/Engine/Source/Runtime/Core/Public/Modules/ModuleManager.h
@@ -13,7 +13,6 @@
 #include "Delegates/Delegate.h"
 #include "Misc/Optional.h"
 #include "Misc/CoreMisc.h"
-#include "Misc/MultiReaderSingleWriterGT.h"
 #include "Modules/ModuleInterface.h"
 #include "Modules/Boilerplate/ModuleBoilerplate.h"
 
@@ -611,7 +610,7 @@ private:
 	TArray<FString> GameBinariesDirectories;
 
 	/** Critical section object controlling R/W access to Modules. */
-	mutable FMultiReaderSingleWriterGT ModulesCriticalSection;
+	mutable FCriticalSection ModulesCriticalSection;
 };
 
 /**
diff --git a/Engine/Source/Runtime/Core/Public/ProfilingDebugging/CookStats.h b/Engine/Source/Runtime/Core/Public/ProfilingDebugging/CookStats.h
--- a/Engine/Source/Runtime/Core/Public/ProfilingDebugging/CookStats.h
+++ b/Engine/Source/Runtime/Core/Public/ProfilingDebugging/CookStats.h
@@ -123,27 +123,23 @@ public:
 		return Attrs;
 	}
 
-	/** ToString abstraction that can convert anything into a string. Differs from LexicalConversion in that it is legal to pass a string or TCHAR* to it. */
-	static const FString& ToString(const FString& InStr) { return InStr; }
-	static FString ToString(const TCHAR* InStr) { return FString(InStr); }
-	template <typename T> static FString ToString(T&&Val) { return LexicalConversion::ToString(Forward<T>(Val)); }
 private:
 	template <typename KeyType>
 	static void ImplCreateKeyValueArray(TArray<FCookStatsManager::StringKeyValue>& Attrs, KeyType&& Key)
 	{
-		Attrs.Emplace(ToString(Forward<KeyType>(Key)), TEXT(""));
+		Attrs.Emplace(Lex::ToString(Forward<KeyType>(Key)), TEXT(""));
 	}
 
 	template <typename KeyType, typename ValueType>
 	static void ImplCreateKeyValueArray(TArray<FCookStatsManager::StringKeyValue>& Attrs, KeyType&& Key, ValueType&& Value)
 	{
-		Attrs.Emplace(ToString(Forward<KeyType>(Key)), ToString(Forward<ValueType>(Value)));
+		Attrs.Emplace(Lex::ToString(Forward<KeyType>(Key)), Lex::ToString(Forward<ValueType>(Value)));
 	}
 
 	template <typename KeyType, typename ValueType, typename...ArgTypes>
 	static void ImplCreateKeyValueArray(TArray<FCookStatsManager::StringKeyValue>& Attrs, KeyType&& Key, ValueType&& Value, ArgTypes&&...Args)
 	{
-		Attrs.Emplace(ToString(Forward<KeyType>(Key)), ToString(Forward<ValueType>(Value)));
+		Attrs.Emplace(Lex::ToString(Forward<KeyType>(Key)), Lex::ToString(Forward<ValueType>(Value)));
 		ImplCreateKeyValueArray(Attrs, Forward<ArgTypes>(Args)...);
 	}
 
diff --git a/Engine/Source/Runtime/Core/Public/Serialization/Archive.h b/Engine/Source/Runtime/Core/Public/Serialization/Archive.h
--- a/Engine/Source/Runtime/Core/Public/Serialization/Archive.h
+++ b/Engine/Source/Runtime/Core/Public/Serialization/Archive.h
@@ -21,15 +21,8 @@ template<class TEnum> class TEnumAsByte;
 //@todoio if this is off, then we should leave the package file format completely unchanged....!!!!! this is really important to fix before we merge to main
 //#define COOK_FOR_EVENT_DRIVEN_LOAD (1)
 
-#define SPLIT_COOKED_FILES (1)
-
-#if !defined(USE_NEW_ASYNC_IO)
-#error "USE_NEW_ASYNC_IO must be defined"
-#endif
-
-// We only use event driven loading if we have cooked for it and it is a cooked platform
-#define USE_EVENT_DRIVEN_ASYNC_LOAD (!WITH_EDITORONLY_DATA && USE_NEW_ASYNC_IO)
-#define DEVIRTUALIZE_FLinkerLoad_Serialize (USE_EVENT_DRIVEN_ASYNC_LOAD)
+#define EVENT_DRIVEN_ASYNC_LOAD_ACTIVE_AT_RUNTIME (!GIsInitialLoad) // set to (!GIsInitialLoad) to avoid using the EDL at boot time
+#define DEVIRTUALIZE_FLinkerLoad_Serialize (!WITH_EDITORONLY_DATA)
 
 /**
  * Base class for archives that can be used for loading, saving, and garbage
@@ -279,10 +272,12 @@ public:
 	 * @param Ar The archive to serialize from or to.
 	 * @param Value The value to serialize.
 	 */
+#if WITH_EDITOR
+	FArchive& operator<<( bool& D );
+#else
 	FORCEINLINE friend FArchive& operator<<( FArchive& Ar, bool& D )
 	{
 		// Serialize bool as if it were UBOOL (legacy, 32 bit int).
-
 #if DEVIRTUALIZE_FLinkerLoad_Serialize
 		const uint8 * RESTRICT Src = Ar.ActiveFPLB->StartFastPathLoadBuffer;
 		if (Src + sizeof(uint32) <= Ar.ActiveFPLB->EndFastPathLoadBuffer)
@@ -304,6 +299,7 @@ public:
 		}
 		return Ar;
 	}
+#endif
 
 	/**
 	 * Serializes a signed 32-bit integer value from or into an archive.
diff --git a/Engine/Source/Runtime/Core/Public/Stats/StatsData.h b/Engine/Source/Runtime/Core/Public/Stats/StatsData.h
--- a/Engine/Source/Runtime/Core/Public/Stats/StatsData.h
+++ b/Engine/Source/Runtime/Core/Public/Stats/StatsData.h
@@ -69,7 +69,7 @@ struct CORE_API FStatConstants
 	static const FStatNameAndInfo AdvanceFrame;
 };
 
-namespace LexicalConversion
+namespace Lex
 {
 	inline void FromString(FName& OutValue, const TCHAR* Buffer )
 	{
@@ -78,7 +78,7 @@ namespace LexicalConversion
 }
 
 /** Parse a typed value into the specified out parameter.
- * 	Expects to find a FromString function that takes a reference to T. Defaults are provided in the LexicalConversion namespace.
+ * 	Expects to find a FromString function that takes a reference to T. Defaults are provided in the Lex namespace.
  */
 template <typename T>
 void ParseTypedValue( const TCHAR* Stream, const TCHAR* Match, T& Out )
@@ -86,7 +86,7 @@ void ParseTypedValue( const TCHAR* Stream, const TCHAR* Match, T& Out )
 	TCHAR Temp[64] = TEXT( "" );
 	if( FParse::Value( Stream, Match, Temp, ARRAY_COUNT( Temp ) ) )
 	{
-		using namespace LexicalConversion;
+		using namespace Lex;
 		FromString( Out, Temp );
 	}
 }
diff --git a/Engine/Source/Runtime/CoreUObject/Private/Blueprint/BlueprintSupport.cpp b/Engine/Source/Runtime/CoreUObject/Private/Blueprint/BlueprintSupport.cpp
--- a/Engine/Source/Runtime/CoreUObject/Private/Blueprint/BlueprintSupport.cpp
+++ b/Engine/Source/Runtime/CoreUObject/Private/Blueprint/BlueprintSupport.cpp
@@ -1750,9 +1750,7 @@ void FLinkerLoad::CreateDynamicTypeLoader()
 		OuterImport->ObjectName = Import.ObjectRef.PackageName;
 
 		if ((Import.ObjectRef.ClassName == DynamicClassName) 
-#if USE_EVENT_DRIVEN_ASYNC_LOAD
-			&& GIsInitialLoad
-#endif 
+			&& (!GEventDrivenLoaderEnabled || !EVENT_DRIVEN_ASYNC_LOAD_ACTIVE_AT_RUNTIME)
 			&& (Import.ObjectRef.ClassPackageName == DynamicClassPackageName))
 		{
 			const FString DynamicClassPath = Import.ObjectRef.PackageName.ToString() + TEXT(".") + Import.ObjectRef.ObjectName.ToString();
@@ -1780,8 +1778,8 @@ void FLinkerLoad::CreateDynamicTypeLoader()
 		DynamicTypeExport->ObjectFlags |= RF_Public;
 	}
 
-#if USE_EVENT_DRIVEN_ASYNC_LOAD
-
+	if (GEventDrivenLoaderEnabled)
+	{
 	const FString DynamicTypePath = GetExportPathName(DynamicTypeExportIndex);
 	const FName DynamicTypeClassName = GetDynamicTypeClassName(*DynamicTypePath);
 	ensure(DynamicTypeClassName != NAME_None);
@@ -1891,11 +1889,11 @@ void FLinkerLoad::CreateDynamicTypeLoader()
 			HandleDependencyTypeForExport(EDependencyType::CreateBeforeCreate);
 		}
 	}
-#endif
-
-#if !USE_EVENT_DRIVEN_ASYNC_LOAD
+	}
+	else
+	{
 	LinkerRoot->SetPackageFlags(LinkerRoot->GetPackageFlags() | PKG_CompiledIn);
-#endif
+	}
 }
 
 /*******************************************************************************
@@ -2270,8 +2268,7 @@ void FConvertedBlueprintsDependencies::FillUsedAssetsInDynamicClass(UDynamicClas
 	TArray<FBlueprintDependencyData> UsedAssetdData;
 	GetUsedAssets(UsedAssetdData);
 
-#if USE_EVENT_DRIVEN_ASYNC_LOAD
-	if (!GIsInitialLoad)
+	if (GEventDrivenLoaderEnabled && EVENT_DRIVEN_ASYNC_LOAD_ACTIVE_AT_RUNTIME)
 	{
 		FLinkerLoad* Linker = DynamicClass->GetOutermost()->LinkerLoad;
 		if (Linker)
@@ -2290,7 +2287,6 @@ void FConvertedBlueprintsDependencies::FillUsedAssetsInDynamicClass(UDynamicClas
 		}
 		check(0);
 	}
-#endif
 
 	for (FBlueprintDependencyData& ItData : UsedAssetdData)
 	{
diff --git a/Engine/Source/Runtime/CoreUObject/Private/Misc/GCObjectReferencer.cpp b/Engine/Source/Runtime/CoreUObject/Private/Misc/GCObjectReferencer.cpp
--- a/Engine/Source/Runtime/CoreUObject/Private/Misc/GCObjectReferencer.cpp
+++ b/Engine/Source/Runtime/CoreUObject/Private/Misc/GCObjectReferencer.cpp
@@ -32,8 +32,7 @@ void UGCObjectReferencer::AddReferencedObjects(UObject* InThis, FReferenceCollec
 void UGCObjectReferencer::AddObject(FGCObject* Object)
 {
 	check(Object);
-	check(GObjUnhashUnreachableIsInProgress || GObjIncrementalPurgeIsInProgress || !IsGarbageCollecting());
-	FScopeLock ReferencedObjectsLock(&ReferencedObjectsCritical);
+	check(GObjUnhashUnreachableIsInProgress || GObjIncrementalPurgeIsInProgress || !IsGarbageCollecting());	FScopeLock ReferencedObjectsLock(&ReferencedObjectsCritical);
 	// Make sure there are no duplicates. Should be impossible...
 	ReferencedObjects.AddUnique(Object);
 }
@@ -41,9 +40,7 @@ void UGCObjectReferencer::AddObject(FGCObject* Object)
 void UGCObjectReferencer::RemoveObject(FGCObject* Object)
 {
 	check(Object);
-	// Only allow to remove FGCObjects during GC when purging objects.
-	check(GObjUnhashUnreachableIsInProgress || GObjIncrementalPurgeIsInProgress || !IsGarbageCollecting());
-	FScopeLock ReferencedObjectsLock(&ReferencedObjectsCritical);
+	check(GObjUnhashUnreachableIsInProgress || GObjIncrementalPurgeIsInProgress || !IsGarbageCollecting());	FScopeLock ReferencedObjectsLock(&ReferencedObjectsCritical);
 	ReferencedObjects.Remove(Object);
 }
 
diff --git a/Engine/Source/Runtime/CoreUObject/Private/Misc/StartupPackages.cpp b/Engine/Source/Runtime/CoreUObject/Private/Misc/StartupPackages.cpp
--- a/Engine/Source/Runtime/CoreUObject/Private/Misc/StartupPackages.cpp
+++ b/Engine/Source/Runtime/CoreUObject/Private/Misc/StartupPackages.cpp
@@ -33,9 +33,9 @@ void FStartupPackages::GetStartupPackageNames(TArray<FString>& PackageNames, con
 	}
 }
 
-#if !USE_NEW_ASYNC_IO
 static void AsyncPreloadPackageList(const TArray<FString>& PackageNames)
 {
+	check(!GNewAsyncIO);
 	// Iterate over all native script packages and preload them.
 	for (int32 PackageIndex=0; PackageIndex<PackageNames.Num(); PackageIndex++)
 	{
@@ -43,7 +43,6 @@ static void AsyncPreloadPackageList(const TArray<FString>& PackageNames)
 		FLinkerLoad::AsyncPreloadPackage(*PackageNames[PackageIndex]);
 	}
 }
-#endif
 
 void FStartupPackages::LoadPackageList(const TArray<FString>& PackageNames)
 {
@@ -69,21 +68,5 @@ bool FStartupPackages::LoadAll()
 		FStartupPackages::GetStartupPackageNames(StartupPackages);
 	}
 
-//@todoio this doesn't seem to work with !USE_NEW_ASYNC_IO anymore...and honestly I can't see how it ever worked well in UE4
-#if 0 && !USE_NEW_ASYNC_IO
-	// should startup packages load from memory?
-	bool bSerializeStartupPackagesFromMemory = false;
-	GConfig->GetBool(TEXT("Engine.StartupPackages"), TEXT("bSerializeStartupPackagesFromMemory"), bSerializeStartupPackagesFromMemory, GEngineIni);
-
-	if (bSerializeStartupPackagesFromMemory)
-	{
-		// @todo: this needs proper support in the cooker
-		if (FPlatformProperties::RequiresCookedData())
-		{
-			// kick them off to be preloaded
-			AsyncPreloadPackageList(StartupPackages);
-		}
-	}
-#endif
 	return bReturn;
 }
diff --git a/Engine/Source/Runtime/CoreUObject/Private/Misc/WorldCompositionUtility.cpp b/Engine/Source/Runtime/CoreUObject/Private/Misc/WorldCompositionUtility.cpp
--- a/Engine/Source/Runtime/CoreUObject/Private/Misc/WorldCompositionUtility.cpp
+++ b/Engine/Source/Runtime/CoreUObject/Private/Misc/WorldCompositionUtility.cpp
@@ -98,9 +98,7 @@ bool FWorldTileInfo::Read(const FString& InPackageFileName, FWorldTileInfo& OutI
 	{
 		if (!!(FileSummary.PackageFlags & PKG_StoreCompressed))
 		{
-#if USE_NEW_ASYNC_IO
-			check(!"Package level compression cannot be used with the async io scheme.");
-#else
+			checkf(!GNewAsyncIO, TEXT("Package level compression cannot be used with the async io scheme."));
 			check(FileSummary.CompressedChunks.Num() > 0);
 			if (!FileReader->SetCompressionMap(&FileSummary.CompressedChunks, (ECompressionFlags)FileSummary.CompressionFlags))
 			{
@@ -108,7 +106,6 @@ bool FWorldTileInfo::Read(const FString& InPackageFileName, FWorldTileInfo& OutI
 				check(!FileReader->IsError());
 				verify(FileReader->SetCompressionMap(&FileSummary.CompressedChunks, (ECompressionFlags)FileSummary.CompressionFlags));
 			}
-#endif // USE_NEW_ASYNC_IO
 		}
 				
 		// Seek the the part of the file where the structure lives
