Copying //UE4/Dev-Core to //UE4/Dev-Main (Source: //UE4/Dev-Core @ 3283640)

diff --git a/Engine/Source/Programs/CrashReporter/CrashReportWebSite/Controllers/BuggsController.cs b/Engine/Source/Programs/CrashReporter/CrashReportWebSite/Controllers/BuggsController.cs
--- a/Engine/Source/Programs/CrashReporter/CrashReportWebSite/Controllers/BuggsController.cs
+++ b/Engine/Source/Programs/CrashReporter/CrashReportWebSite/Controllers/BuggsController.cs
@@ -1,5 +1,4 @@
 ﻿// Copyright 1998-2017 Epic Games, Inc. All Rights Reserved.
-
 using System;
 using System.Collections.Generic;
 using System.Diagnostics;
@@ -9,6 +8,7 @@ using System.Linq.Expressions;
 using System.Net;
 using System.Web;
 using System.Web.Mvc;
+using System.Web.UI;
 using Microsoft.Practices.ObjectBuilder2;
 using Newtonsoft.Json.Linq;
 using Tools.CrashReporter.CrashReportWebSite.DataModels;
@@ -22,14 +22,13 @@ namespace Tools.CrashReporter.CrashReportWebSite.Controllers
 	/// </summary>
 	public class BuggsController : Controller
 	{
-        private IUnitOfWork _unitOfWork;
+	    private int pageSize = 50;
 
 		/// <summary>
 		/// An empty constructor.
 		/// </summary>
-		public BuggsController(IUnitOfWork unitOfWork)
+		public BuggsController()
 		{
-		    _unitOfWork = unitOfWork;
 		}
 
 		/// <summary>
@@ -63,49 +62,39 @@ namespace Tools.CrashReporter.CrashReportWebSite.Controllers
 			    var displayFunctionNames = buggsForm["DisplayFunctionNames"] == "true";
 			    var displayFileNames = buggsForm["DisplayFileNames"] == "true";
 			    var displayFilePathNames = false;
+
 				if( buggsForm["DisplayFilePathNames"] == "true" )
 				{
 					displayFilePathNames = true;
 					displayFileNames = false;
 				}
-			    var displayUnformattedCallStack = buggsForm["DisplayUnformattedCallStack"] == "true";
 
-			    // Create a new view and populate with crashes
-				List<Crash> crashes = null;
-				var model = new BuggViewModel();
-				var newBugg = _unitOfWork.BuggRepository.GetById( id );
-				if( newBugg == null )
-				{
-					return RedirectToAction( "Index" );
-				}
-                crashes = newBugg.Crashes.OrderByDescending(data => data.TimeOfCrash).ToList();
-                newBugg.CrashesInTimeFrameAll = crashes.Count;
-                newBugg.CrashesInTimeFrameGroup = crashes.Count;
-			    newBugg.NumberOfCrashes = crashes.Count;
+			    var displayUnformattedCallStack = buggsForm["DisplayUnformattedCallStack"] == "true";
 
-                // Handle 'CopyToJira' button
+			    var model = GetResult(id);
+                model.Bugg.PrepareBuggForJira(model.Crashes);
+			    // Handle 'CopyToJira' button
                 var buggIdToBeAddedToJira = 0;
                 foreach (var entry in buggsForm.Cast<object>().Where(entry => entry.ToString().Contains("CopyToJira-")))
                 {
                     int.TryParse(entry.ToString().Substring("CopyToJira-".Length), out buggIdToBeAddedToJira);
                     break;
                 }
-
-                newBugg.PrepareBuggForJira(crashes);
                 if (buggIdToBeAddedToJira != 0)
                 {
-                    newBugg.JiraProject = buggsForm["JiraProject"];
-                    newBugg.CopyToJira();
+                   
+                    model.Bugg.JiraProject = buggsForm["JiraProject"];
+                    model.Bugg.CopyToJira();
                 }
 
                 var jc = JiraConnection.Get();
 				var bValidJira = false;
 
 				// Verify valid JiraID, this may be still a TTP
-				if( !string.IsNullOrEmpty( newBugg.TTPID ) )
+				if( !string.IsNullOrEmpty( model.Bugg.TTPID ) )
 				{
 					var jira = 0;
-                    int.TryParse(newBugg.TTPID, out jira);
+                    int.TryParse(model.Bugg.TTPID, out jira);
 
                     if (jira == 0)
 					{
@@ -116,7 +105,7 @@ namespace Tools.CrashReporter.CrashReportWebSite.Controllers
 				if( jc.CanBeUsed() && bValidJira )
 				{
 					// Grab the data form JIRA.
-					var jiraSearchQuery = "key = " + newBugg.TTPID;
+					var jiraSearchQuery = "key = " + model.Bugg.TTPID;
 					var jiraResults = new Dictionary<string, Dictionary<string, object>>();
 
 					try
@@ -135,11 +124,11 @@ namespace Tools.CrashReporter.CrashReportWebSite.Controllers
 					}
 					catch (System.Exception)
 					{
-                        newBugg.JiraSummary = "JIRA MISMATCH";
-                        newBugg.JiraComponentsText = "JIRA MISMATCH";
-                        newBugg.JiraResolution = "JIRA MISMATCH";
-                        newBugg.JiraFixVersionsText = "JIRA MISMATCH";
-                        newBugg.JiraFixCL = "JIRA MISMATCH";
+                        model.Bugg.JiraSummary = "JIRA MISMATCH";
+                        model.Bugg.JiraComponentsText = "JIRA MISMATCH";
+                        model.Bugg.JiraResolution = "JIRA MISMATCH";
+                        model.Bugg.JiraFixVersionsText = "JIRA MISMATCH";
+                        model.Bugg.JiraFixCL = "JIRA MISMATCH";
 					}
 
 					// Jira Key, Summary, Components, Resolution, Fix version, Fix changelist
@@ -171,13 +160,13 @@ namespace Tools.CrashReporter.CrashReportWebSite.Controllers
 
                         //Conversion to ado.net entity framework
 
-                        newBugg.JiraSummary = summary;
-                        newBugg.JiraComponentsText = componentsText;
-                        newBugg.JiraResolution = resolution;
-                        newBugg.JiraFixVersionsText = fixVersionsText;
+                        model.Bugg.JiraSummary = summary;
+                        model.Bugg.JiraComponentsText = componentsText;
+                        model.Bugg.JiraResolution = resolution;
+                        model.Bugg.JiraFixVersionsText = fixVersionsText;
 						if( fixCl != 0 )
 						{
-							newBugg.FixedChangeList = fixCl.ToString();
+							model.Bugg.FixedChangeList = fixCl.ToString();
 						}
 					}
 				}
@@ -187,30 +176,34 @@ namespace Tools.CrashReporter.CrashReportWebSite.Controllers
 				{
 					if( !string.IsNullOrEmpty( buggsForm["SetStatus"] ) )
 					{
-						newBugg.Status = buggsForm["SetStatus"];
-                        newBugg.Crashes.ForEach(data =>data.Status = buggsForm["SetStatus"]);
+						model.Bugg.Status = buggsForm["SetStatus"];
+                        model.Bugg.Crashes.ForEach(data =>data.Status = buggsForm["SetStatus"]);
 					}
 
 					if( !string.IsNullOrEmpty( buggsForm["SetFixedIn"] ) )
 					{
-						newBugg.FixedChangeList = buggsForm["SetFixedIn"];
-                        newBugg.Crashes.ForEach(data => data.FixedChangeList = buggsForm["SetFixedIn"]);
+						model.Bugg.FixedChangeList = buggsForm["SetFixedIn"];
+                        model.Bugg.Crashes.ForEach(data => data.FixedChangeList = buggsForm["SetFixedIn"]);
 					}
 
 					if( !string.IsNullOrEmpty( buggsForm["SetTTP"] ) )
 					{
-						newBugg.TTPID = buggsForm["SetTTP"];
-                        newBugg.Crashes.ForEach(data => data.Jira = buggsForm["SetTTP"]);
+						model.Bugg.TTPID = buggsForm["SetTTP"];
+                        model.Bugg.Crashes.ForEach(data => data.Jira = buggsForm["SetTTP"]);
 					}
-
-                    _unitOfWork.BuggRepository.Update(newBugg);
-                    _unitOfWork.Save();
+				    using (var unitOfWork = new UnitOfWork(new CrashReportEntities()))
+				    {
+				        unitOfWork.BuggRepository.Update(model.Bugg);
+				        unitOfWork.Save();
+				    }
 				}
 
 				// Set up the view model with the crash data
-				model.Crashes = crashes;
-			    model.Bugg = newBugg;
+			    int page;
+                int.TryParse(buggsForm["Page"], out page);
 
+				model.Crashes = model.Crashes.Skip(page * pageSize).Take(pageSize).ToList();
+			    
 				var newCrash = model.Crashes.FirstOrDefault();
 				if( newCrash != null )
 				{
@@ -235,26 +228,54 @@ namespace Tools.CrashReporter.CrashReportWebSite.Controllers
 
 					model.LatestCrashSummary = newCrash.Summary;
 				}
+
                 model.LatestCrashSummary = newCrash.Summary;
                 model.Bugg.LatestCrashSummary = newCrash.Summary;
 				model.GenerationTime = logTimer.GetElapsedSeconds().ToString( "F2" );
+
+                //Populate Jira Projects
+                var jiraConnection = JiraConnection.Get();
+                var response = jiraConnection.JiraRequest("/issue/createmeta", JiraConnection.JiraMethod.GET, null, HttpStatusCode.OK);
+
+                using (var responseReader = new StreamReader(response.GetResponseStream()))
+                {
+                    var responseText = responseReader.ReadToEnd();
+
+                    JObject jsonObject = JObject.Parse(responseText);
+
+                    JToken fields = jsonObject["projects"];
+
+                    foreach (var project in fields)
+                    {
+                        model.JiraProjects.Add(new SelectListItem() { Text = project["name"].ToObject<string>(), Value = project["key"].ToObject<string>() });
+                    }
+                }
+
+			    model.PagingInfo = new PagingInfo
+			    {
+			        CurrentPage = page,
+			        PageSize = pageSize,
+			        TotalResults = model.Bugg.NumberOfCrashes.Value
+			    };
+
 				return View( "Show", model );
 			}
 		}
+        
+        #region Private Methods
 
         /// <summary>
         /// Retrieve all Buggs matching the search criteria.
         /// </summary>
         /// <param name="FormData">The incoming form of search criteria from the client.</param>
         /// <returns>A view to display the filtered Buggs.</returns>
-        public BuggsViewModel GetResults(FormHelper FormData)
+        private BuggsViewModel GetResults(FormHelper FormData)
         {
             // Right now we take a Result IQueryable starting with ListAll() Buggs then we widdle away the result set by tacking on
             // Linq queries. Essentially it's Results.ListAll().Where().Where().Where().Where().Where().Where()
             // Could possibly create more optimized queries when we know exactly what we're querying
             // The downside is that if we add more parameters each query may need to be updated.... Or we just add new case statements
             // The other downside is that there's less code reuse, but that may be worth it.
-
             var fromDate = FormData.DateFrom;
             var toDate = FormData.DateTo.AddDays(1);
 
@@ -262,148 +283,195 @@ namespace Tools.CrashReporter.CrashReportWebSite.Controllers
             var skip = (FormData.Page - 1) * FormData.PageSize;
             var take = FormData.PageSize;
 
-            var userGroupId = _unitOfWork.UserGroupRepository.First(data => data.Name == FormData.UserGroup).Id;
-
-            // Get every Bugg.
-            var resultsAll = _unitOfWork.BuggRepository.ListAll();
+            var sortedResultsList = new List<Bugg>();
+            var groupCounts = new SortedDictionary<string, int>();
+            int totalCountedRecords = 0;
 
-            // Look at all Buggs that are still 'open' i.e. the last crash occurred in our date range.
-            results = FilterByDate(resultsAll, FormData.DateFrom, FormData.DateTo);
-
-            // Filter results by build version.
-            if (!string.IsNullOrEmpty(FormData.VersionName))
+            using (var unitOfWork = new UnitOfWork(new CrashReportEntities()))
             {
-                results = FilterByBuildVersion(results, FormData.VersionName);
-            }
+                var userGroupId = unitOfWork.UserGroupRepository.First(data => data.Name == FormData.UserGroup).Id;
 
-            // Filter by BranchName
-            if (!string.IsNullOrEmpty(FormData.BranchName))
-            {
-                results =
-                    results.Where(
-                        data => data.Crashes.Any(da => da.Branch.Equals(FormData.BranchName)));
-            }
+                // Get every Bugg.
+                var resultsAll = unitOfWork.BuggRepository.ListAll();
 
-            // Filter by PlatformName
-            if (!string.IsNullOrEmpty(FormData.PlatformName))
-            {
-                results =
-                    results.Where(
-                        data => data.Crashes.Any(da => da.PlatformName.Equals(FormData.PlatformName)));
-            }
+                // Look at all Buggs that are still 'open' i.e. the last crash occurred in our date range.
+                results = FilterByDate(resultsAll, FormData.DateFrom, FormData.DateTo);
 
-            // Run at the end
-            if (!string.IsNullOrEmpty(FormData.SearchQuery))
-            {
-                results = FilterByCallstack(results, FormData.SearchQuery);
-            }
+                // Filter results by build version.
+                if (!string.IsNullOrEmpty(FormData.VersionName))
+                {
+                    results = FilterByBuildVersion(results, FormData.VersionName);
+                }
 
-            // Filter by Crash Type
-            if (FormData.CrashType != "All")
-            {
-                switch (FormData.CrashType)
+                // Filter by BranchName
+                if (!string.IsNullOrEmpty(FormData.BranchName))
                 {
-                    case "Crashes":
-                        results = results.Where(buggInstance => buggInstance.CrashType == 1);
-                        break;
-                    case "Assert":
-                        results = results.Where(buggInstance => buggInstance.CrashType == 2);
-                        break;
-                    case "Ensure":
-                        results = results.Where(buggInstance => buggInstance.CrashType == 3);
-                        break;
-                    case "CrashesAsserts":
-                        results = results.Where(buggInstance => buggInstance.CrashType == 1 || buggInstance.CrashType == 2);
-                        break;
+                    results =
+                        results.Where(
+                            data => data.Crashes.Any(da => da.Branch.Equals(FormData.BranchName)));
                 }
-            }
 
-            // Filter by user group if present
-            if(!string.IsNullOrEmpty(FormData.UserGroup))
-                results = FilterByUserGroup(results, FormData.UserGroup);
+                // Filter by PlatformName
+                if (!string.IsNullOrEmpty(FormData.PlatformName))
+                {
+                    results =
+                        results.Where(
+                            data => data.Crashes.Any(da => da.PlatformName.Equals(FormData.PlatformName)));
+                }
 
-            if (!string.IsNullOrEmpty(FormData.JiraId))
-            {
-                results = FilterByJira(results, FormData.JiraId);
-            }
-                
-            // Grab just the results we want to display on this page
-            var totalCountedRecords = results.Count();
-                
-            var crashesByUserGroupCounts = _unitOfWork.CrashRepository.ListAll().Where(data =>
-                data.User.UserGroupId == userGroupId && 
-                data.TimeOfCrash >= fromDate &&
-                data.TimeOfCrash <= toDate &&
-                data.PatternId != null)
-                .GroupBy(data => data.PatternId)
-                .Select(data => new { data.Key, Count = data.Count() })
-                .OrderByDescending(data => data.Count)
-                .ToDictionary(data => data.Key, data => data.Count);
-
-            var crashesInTimeFrameCounts = _unitOfWork.CrashRepository.ListAll().Where(data =>
-                data.TimeOfCrash >= fromDate &&
-                data.TimeOfCrash <= toDate
-                && data.PatternId != null)
-                .GroupBy(data => data.PatternId)
-                .Select(data => new { data.Key, Count = data.Count() })
-                .OrderByDescending(data => data.Count)
-                .ToDictionary(data => data.Key, data => data.Count);
-
-            var affectedUsers = _unitOfWork.CrashRepository.ListAll().Where(data =>
-                data.User.UserGroupId == userGroupId &&
-                data.TimeOfCrash >= fromDate &&
-                data.TimeOfCrash <= toDate &&
-                data.PatternId != null).Select(data => new {PatternId = data.PatternId, ComputerName = data.ComputerName}).Distinct()
-                .GroupBy(data => data.PatternId)
-                .Select(data => new { data.Key, Count = data.Count() })
-                .OrderByDescending(data => data.Count)
-                .ToDictionary(data => data.Key, data => data.Count);
-                
-            results = GetSortedResults(results, FormData.SortTerm, FormData.SortTerm == "descending",
-            FormData.DateFrom, FormData.DateTo, FormData.UserGroup);
-                
-            var sortedResultsList = 
-                results.ToList();
+                // Run at the end
+                if (!string.IsNullOrEmpty(FormData.SearchQuery))
+                {
+                    try
+                    {
+                        var queryString = HttpUtility.HtmlDecode(FormData.SearchQuery.ToString()) ?? "";
+
+                        // Take out terms starting with a -
+                        var terms = queryString.Split("-, ;+".ToCharArray(), StringSplitOptions.RemoveEmptyEntries);
+                        var allFuncionCallIds = new HashSet<int>();
+                        foreach (var term in terms)
+                        {
+                            var functionCallIds = unitOfWork.FunctionRepository.Get(functionCallInstance => functionCallInstance.Call.Contains(term)).Select(x => x.Id).ToList();
+                            foreach (var id in functionCallIds)
+                            {
+                                allFuncionCallIds.Add(id);
+                            }
+                        }
+
+                        // Search for all function ids. OR operation, not very efficient, but for searching for one function should be ok.
+                        results = allFuncionCallIds.Aggregate(results, (current, id) => current.Where(x => x.Pattern.Contains(id + "+") || x.Pattern.Contains("+" + id)));
+                    }
+                    catch (Exception ex)
+                    {
+                        Debug.WriteLine("Exception in Search: " + ex.ToString());
+                    }
+                }
 
-            // Get UserGroup ResultCounts
-            var groups = results.SelectMany(data => data.UserGroups).GroupBy(data => data.Name).Select(data => new { Key = data.Key, Count = data.Count()}).ToDictionary(x => x.Key, y => y.Count);
-            var groupCounts = new SortedDictionary<string, int>(groups);
-                
-            foreach (var bugg in sortedResultsList)
-            {
-                if (bugg.PatternId.HasValue && crashesByUserGroupCounts.ContainsKey(bugg.PatternId.Value))
+                // Filter by Crash Type
+                if (FormData.CrashType != "All")
                 {
-                    bugg.CrashesInTimeFrameGroup = crashesByUserGroupCounts[bugg.PatternId.Value];
+                    switch (FormData.CrashType)
+                    {
+                        case "Crashes":
+                            results = results.Where(buggInstance => buggInstance.CrashType == 1);
+                            break;
+                        case "Assert":
+                            results = results.Where(buggInstance => buggInstance.CrashType == 2);
+                            break;
+                        case "Ensure":
+                            results = results.Where(buggInstance => buggInstance.CrashType == 3);
+                            break;
+                        case "CrashesAsserts":
+                            results =
+                                results.Where(buggInstance => buggInstance.CrashType == 1 || buggInstance.CrashType == 2);
+                            break;
+                    }
                 }
-                else
-                    bugg.CrashesInTimeFrameGroup = 0;
 
-                if (bugg.PatternId.HasValue && crashesInTimeFrameCounts.ContainsKey(bugg.PatternId.Value))
+                // Filter by user group if present
+                if (!string.IsNullOrEmpty(FormData.UserGroup))
+                    results = FilterByUserGroup(results, FormData.UserGroup);
+
+                if (!string.IsNullOrEmpty(FormData.JiraId))
                 {
-                    bugg.CrashesInTimeFrameAll = crashesInTimeFrameCounts[bugg.PatternId.Value];
+                    results = FilterByJira(results, FormData.JiraId);
                 }
-                else
-                    bugg.CrashesInTimeFrameAll = 0;
 
-                if (bugg.PatternId.HasValue && affectedUsers.ContainsKey(bugg.PatternId.Value))
+                // Grab just the results we want to display on this page
+                totalCountedRecords = results.Count();
+
+                var crashesByUserGroupCounts = unitOfWork.CrashRepository.ListAll().Where(data =>
+                    data.User.UserGroupId == userGroupId &&
+                    data.TimeOfCrash >= fromDate &&
+                    data.TimeOfCrash <= toDate &&
+                    data.PatternId != null)
+                    .GroupBy(data => data.PatternId)
+                    .Select(data => new {data.Key, Count = data.Count()})
+                    .OrderByDescending(data => data.Count)
+                    .ToDictionary(data => data.Key, data => data.Count);
+
+                var crashesInTimeFrameCounts = unitOfWork.CrashRepository.ListAll().Where(data =>
+                    data.TimeOfCrash >= fromDate &&
+                    data.TimeOfCrash <= toDate
+                    && data.PatternId != null)
+                    .GroupBy(data => data.PatternId)
+                    .Select(data => new {data.Key, Count = data.Count()})
+                    .OrderByDescending(data => data.Count)
+                    .ToDictionary(data => data.Key, data => data.Count);
+
+                var affectedUsers = unitOfWork.CrashRepository.ListAll().Where(data =>
+                    data.User.UserGroupId == userGroupId &&
+                    data.TimeOfCrash >= fromDate &&
+                    data.TimeOfCrash <= toDate &&
+                    data.PatternId != null)
+                    .Select(data => new {PatternId = data.PatternId, ComputerName = data.ComputerName})
+                    .Distinct()
+                    .GroupBy(data => data.PatternId)
+                    .Select(data => new {data.Key, Count = data.Count()})
+                    .OrderByDescending(data => data.Count)
+                    .ToDictionary(data => data.Key, data => data.Count);
+
+                results = GetSortedResults(results, FormData.SortTerm, FormData.SortTerm == "descending",
+                    FormData.DateFrom, FormData.DateTo, FormData.UserGroup);
+
+                sortedResultsList =
+                    results.ToList();
+
+                // Get UserGroup ResultCounts
+                var groups =
+                    results.SelectMany(data => data.UserGroups)
+                        .GroupBy(data => data.Name)
+                        .Select(data => new {Key = data.Key, Count = data.Count()})
+                        .ToDictionary(x => x.Key, y => y.Count);
+                
+                groupCounts = new SortedDictionary<string, int>(groups);
+
+                foreach (var bugg in sortedResultsList)
                 {
-                    bugg.NumberOfUniqueMachines = affectedUsers[bugg.PatternId.Value];
+                    if (bugg.PatternId.HasValue && crashesByUserGroupCounts.ContainsKey(bugg.PatternId.Value))
+                    {
+                        bugg.CrashesInTimeFrameGroup = crashesByUserGroupCounts[bugg.PatternId.Value];
+                    }
+                    else
+                        bugg.CrashesInTimeFrameGroup = 0;
+
+                    if (bugg.PatternId.HasValue && crashesInTimeFrameCounts.ContainsKey(bugg.PatternId.Value))
+                    {
+                        bugg.CrashesInTimeFrameAll = crashesInTimeFrameCounts[bugg.PatternId.Value];
+                    }
+                    else
+                        bugg.CrashesInTimeFrameAll = 0;
+
+                    if (bugg.PatternId.HasValue && affectedUsers.ContainsKey(bugg.PatternId.Value))
+                    {
+                        bugg.NumberOfUniqueMachines = affectedUsers[bugg.PatternId.Value];
+                    }
+                    else
+                        bugg.NumberOfUniqueMachines = 0;
                 }
-                else
-                    bugg.NumberOfUniqueMachines = 0;
-            }
 
-            sortedResultsList = sortedResultsList.OrderByDescending(data => data.CrashesInTimeFrameGroup).Skip(skip)
+                sortedResultsList = sortedResultsList.OrderByDescending(data => data.CrashesInTimeFrameGroup).Skip(skip)
                     .Take(totalCountedRecords >= skip + take ? take : totalCountedRecords).ToList();
 
-            foreach (var bugg in sortedResultsList)
+                foreach (var bugg in sortedResultsList)
+                {
+                    var crash =
+                        unitOfWork.CrashRepository.First(data => data.BuggId == bugg.Id);
+                    bugg.FunctionCalls = crash.GetCallStack().GetFunctionCalls();
+                }
+            }
+
+            List<SelectListItem> branchNames;
+            List<SelectListItem> versionNames;
+            List<SelectListItem> platformNames;
+            using (var unitOfWork = new UnitOfWork(new CrashReportEntities()))
             {
-                var crash =
-                    _unitOfWork.CrashRepository.First(data => data.BuggId == bugg.Id);
-                bugg.FunctionCalls = crash.GetCallStack().GetFunctionCalls();
+                branchNames = unitOfWork.CrashRepository.GetBranchesAsListItems();
+                versionNames = unitOfWork.CrashRepository.GetVersionsAsListItems();
+                platformNames = unitOfWork.CrashRepository.GetPlatformsAsListItems();
             }
 
-            return new BuggsViewModel(_unitOfWork.CrashRepository)
+            return new BuggsViewModel()
             {
                 Results = sortedResultsList,
                 PagingInfo = new PagingInfo { CurrentPage = FormData.Page, PageSize = FormData.PageSize, TotalResults = totalCountedRecords },
@@ -412,6 +480,9 @@ namespace Tools.CrashReporter.CrashReportWebSite.Controllers
                 UserGroup = FormData.UserGroup,
                 CrashType = FormData.CrashType,
                 SearchQuery = FormData.SearchQuery,
+                BranchNames = branchNames,
+                VersionNames = versionNames,
+                PlatformNames = platformNames,
                 DateFrom = (long)(FormData.DateFrom - CrashesViewModel.Epoch).TotalMilliseconds,
                 DateTo = (long)(FormData.DateTo - CrashesViewModel.Epoch).TotalMilliseconds,
                 VersionName = FormData.VersionName,
@@ -422,13 +493,38 @@ namespace Tools.CrashReporter.CrashReportWebSite.Controllers
             };
         }
 
+        /// <summary>
+        /// Retrieve all Data for a single bug given by the id
+        /// </summary>
+        /// <param name="buggsForm"></param>
+        /// <param name="id"></param>
+        /// <returns></returns>
+        private BuggViewModel GetResult(int id)
+	    {
+            var model = new BuggViewModel();
+            using (var unitOfWork = new UnitOfWork(new CrashReportEntities()))
+            {
+                // Create a new view and populate with crashes
+                List<Crash> crashes = null;
+                model.Bugg = unitOfWork.BuggRepository.GetById(id);
+
+                crashes = model.Bugg.Crashes.OrderByDescending(data => data.TimeOfCrash).ToList();
+                model.Bugg.CrashesInTimeFrameAll = crashes.Count;
+                model.Bugg.CrashesInTimeFrameGroup = crashes.Count;
+                model.Bugg.NumberOfCrashes = crashes.Count;
+                model.Crashes = crashes;
+            }
+
+            return model;
+	    }
+
         /// <summary>
         /// Filter a set of Buggs by a jira ticket.
         /// </summary>
         /// <param name="results">The unfiltered set of Buggs</param>
         /// <param name="jira">The ticket by which to filter the list of buggs</param>
         /// <returns></returns>
-        public IQueryable<Bugg> FilterByJira(IQueryable<Bugg> results, string jira)
+        private IQueryable<Bugg> FilterByJira(IQueryable<Bugg> results, string jira)
         {
             return results.Where(bugg => bugg.TTPID == jira);
         }
@@ -440,7 +536,7 @@ namespace Tools.CrashReporter.CrashReportWebSite.Controllers
         /// <param name="dateFrom">The earliest date to filter by.</param>
         /// <param name="dateTo">The latest date to filter by.</param>
         /// <returns>The set of Buggs between the earliest and latest date.</returns>
-        public IQueryable<Bugg> FilterByDate(IQueryable<Bugg> results, DateTime dateFrom, DateTime dateTo)
+        private IQueryable<Bugg> FilterByDate(IQueryable<Bugg> results, DateTime dateFrom, DateTime dateTo)
         {
             dateTo = dateTo.AddDays(1);
             var buggsInTimeFrame =
@@ -456,7 +552,7 @@ namespace Tools.CrashReporter.CrashReportWebSite.Controllers
         /// <param name="results">The unfiltered set of Buggs.</param>
         /// <param name="buildVersion">The build version to filter by.</param>
         /// <returns>The set of Buggs that matches specified build version</returns>
-        public IQueryable<Bugg> FilterByBuildVersion(IQueryable<Bugg> results, string buildVersion)
+        private IQueryable<Bugg> FilterByBuildVersion(IQueryable<Bugg> results, string buildVersion)
         {
             if (!string.IsNullOrEmpty(buildVersion))
             {
@@ -471,7 +567,7 @@ namespace Tools.CrashReporter.CrashReportWebSite.Controllers
         /// <param name="setOfBuggs">The unfiltered set of Buggs.</param>
         /// <param name="groupName">The user group name to filter by.</param>
         /// <returns>The set of Buggs by users in the requested user group.</returns>
-        public IQueryable<Bugg> FilterByUserGroup(IQueryable<Bugg> setOfBuggs, string groupName)
+        private IQueryable<Bugg> FilterByUserGroup(IQueryable<Bugg> setOfBuggs, string groupName)
         {
             return setOfBuggs.Where(data => data.UserGroups.Any(ug => ug.Name == groupName));
         }
@@ -486,7 +582,7 @@ namespace Tools.CrashReporter.CrashReportWebSite.Controllers
         /// <param name="dateTo">The date of the most recent Bugg to examine.</param>
         /// <param name="groupName">The user group name to filter by.</param>
         /// <returns>A sorted container of Buggs.</returns>
-        public IOrderedQueryable<Bugg> GetSortedResults(IQueryable<Bugg> results, string sortTerm, bool bSortDescending, DateTime dateFrom, DateTime dateTo, string groupName)
+        private IOrderedQueryable<Bugg> GetSortedResults(IQueryable<Bugg> results, string sortTerm, bool bSortDescending, DateTime dateFrom, DateTime dateTo, string groupName)
         {
             IOrderedQueryable<Bugg> orderedResults = null;
 
@@ -549,74 +645,18 @@ namespace Tools.CrashReporter.CrashReportWebSite.Controllers
         /// <param name="predicate"></param>
         /// <param name="bDescending"></param>
         /// <returns></returns>
-        public IOrderedQueryable<Bugg> EnumerableOrderBy<TKey>(IQueryable<Bugg> query, Expression<Func<Bugg, TKey>> predicate, bool bDescending)
+        private IOrderedQueryable<Bugg> EnumerableOrderBy<TKey>(IQueryable<Bugg> query, Expression<Func<Bugg, TKey>> predicate, bool bDescending)
         {
             return bDescending ? query.OrderByDescending(predicate) : query.OrderBy(predicate);
         }
-
-        /// <summary>
-        /// Filter the list of Buggs by a callstack entry.
-        /// </summary>
-        public IQueryable<Bugg> FilterByCallstack(IQueryable<Bugg> results, string callstackEntry)
-        {
-            try
-            {
-                var queryString = HttpUtility.HtmlDecode(callstackEntry.ToString()) ?? "";
-
-                // Take out terms starting with a -
-                var terms = queryString.Split("-, ;+".ToCharArray(), StringSplitOptions.RemoveEmptyEntries);
-                var allFuncionCallIds = new HashSet<int>();
-                foreach (var term in terms)
-                {
-                    var functionCallIds = _unitOfWork.FunctionRepository.Get(functionCallInstance => functionCallInstance.Call.Contains(term)).Select(x => x.Id).ToList();
-                    foreach (var id in functionCallIds)
-                    {
-                        allFuncionCallIds.Add(id);
-                    }
-                }
-
-                // Search for all function ids. OR operation, not very efficient, but for searching for one function should be ok.
-                results = allFuncionCallIds.Aggregate(results, (current, id) => current.Where(x => x.Pattern.Contains(id + "+") || x.Pattern.Contains("+" + id)));
-            }
-            catch (Exception ex)
-            {
-                Debug.WriteLine("Exception in Search: " + ex.ToString());
-            }
-            return results;
-        }
-
+        
         /// <summary>Dispose the controller - safely getting rid of database connections</summary>
         /// <param name="disposing"></param>
         protected override void Dispose(bool disposing)
         {
-            _unitOfWork.Dispose();
             base.Dispose(disposing);
         }
 
-        /// <summary>
-        /// 
-        /// </summary>
-	    public void TestJiraMethod()
-	    {
-	        var jiraConnection = JiraConnection.Get();
-
-	        if (!jiraConnection.CanBeUsed())
-	            return;
-
-            var response = jiraConnection.JiraRequest(
-                "/issue/createmeta?projectKeys=OR&issuetypeNames=Bug&expand=projects.issuetypes.fields",
-                JiraConnection.JiraMethod.GET, null, HttpStatusCode.OK);
-
-            using (var responseReader = new StreamReader(response.GetResponseStream()))
-            {
-                var responseText = responseReader.ReadToEnd();
-
-                JObject jsonObject = JObject.Parse(responseText);
-
-                var fields = jsonObject["projects"][0]["issuetypes"][0]["fields"];
-
-                var requiredFields = fields.Where(c => c["required"].Value<bool>() == true);
-            }
-	    }
-	}
+        #endregion
+    }
 }
diff --git a/Engine/Source/Programs/CrashReporter/CrashReportWebSite/Controllers/CrashesController.cs b/Engine/Source/Programs/CrashReporter/CrashReportWebSite/Controllers/CrashesController.cs
--- a/Engine/Source/Programs/CrashReporter/CrashReportWebSite/Controllers/CrashesController.cs
+++ b/Engine/Source/Programs/CrashReporter/CrashReportWebSite/Controllers/CrashesController.cs
@@ -27,8 +27,7 @@ namespace Tools.CrashReporter.CrashReportWebSite.Controllers
 	public class CrashesController : Controller
 	{
         //Ugly instantiation of crash repository will replace with dependency injection BEFORE this gets anywhere near live.
-	    private IUnitOfWork _unitOfWork;
-	    private SlackWriter _slackWriter;
+	    private readonly SlackWriter _slackWriter;
         
         /// <summary> Special user name, currently used to mark crashes from UE4 releases. </summary>
         const string UserNameAnonymous = "Anonymous";
@@ -36,10 +35,8 @@ namespace Tools.CrashReporter.CrashReportWebSite.Controllers
 		/// <summary>
 		/// An empty constructor.
 		/// </summary>
-		public CrashesController(IUnitOfWork unitOfWork)
+		public CrashesController()
 		{
-		    _unitOfWork = unitOfWork;
-
 		    _slackWriter = new SlackWriter()
 		    {
 		        WebhookUrl = Settings.Default.SlackWebhookUrl,
@@ -54,51 +51,55 @@ namespace Tools.CrashReporter.CrashReportWebSite.Controllers
 		/// </summary>
 		/// <param name="crashesForm">A form of user data passed up from the client.</param>
 		/// <returns>A view to display a list of crash reports.</returns>
-		public ActionResult Index( FormCollection crashesForm )
+		public ActionResult Index(FormCollection crashesForm )
 		{
             using( var logTimer = new FAutoScopedLogTimer( this.GetType().ToString(), bCreateNewLog: true ) )
 			{
 				// Handle any edits made in the Set form fields
-				foreach( var Entry in crashesForm )
-				{
-					int Id = 0;
-					if( int.TryParse( Entry.ToString(), out Id ) )
-					{
-                        Crash currentCrash = _unitOfWork.CrashRepository.GetById(Id);
-						if( currentCrash != null )
-						{
-							if( !string.IsNullOrEmpty( crashesForm["SetStatus"] ) )
-							{
-								currentCrash.Status = crashesForm["SetStatus"];
-							}
-
-							if( !string.IsNullOrEmpty( crashesForm["SetFixedIn"] ) )
-							{
-								currentCrash.FixedChangeList = crashesForm["SetFixedIn"];
-							}
-
-							if( !string.IsNullOrEmpty( crashesForm["SetTTP"] ) )
-							{
-								currentCrash.Jira = crashesForm["SetTTP"];
-							}
-						}
-					}
-
-				    _unitOfWork.Save();
-				}
+                //foreach( var Entry in crashesForm )
+                //{
+                //    int Id = 0;
+                //    if( int.TryParse( Entry.ToString(), out Id ) )
+                //    {
+                //        Crash currentCrash = _unitOfWork.CrashRepository.GetById(Id);
+                //        if( currentCrash != null )
+                //        {
+                //            if( !string.IsNullOrEmpty( crashesForm["SetStatus"] ) )
+                //            {
+                //                currentCrash.Status = crashesForm["SetStatus"];
+                //            }
+
+                //            if( !string.IsNullOrEmpty( crashesForm["SetFixedIn"] ) )
+                //            {
+                //                currentCrash.FixedChangeList = crashesForm["SetFixedIn"];
+                //            }
+
+                //            if( !string.IsNullOrEmpty( crashesForm["SetTTP"] ) )
+                //            {
+                //                currentCrash.Jira = crashesForm["SetTTP"];
+                //            }
+                //        }
+                //    }
+
+                //    _unitOfWork.Save();
+                //}
 
 				// <STATUS>
 
 				// Parse the contents of the query string, and populate the form
 				var formData = new FormHelper( Request, crashesForm, "TimeOfCrash" );
 				var result = GetResults( formData );
-                result.BranchNames = _unitOfWork.CrashRepository.GetBranchesAsListItems();
-                result.VersionNames = _unitOfWork.CrashRepository.GetVersionsAsListItems();
-                result.PlatformNames = _unitOfWork.CrashRepository.GetPlatformsAsListItems();
-                result.EngineModes = _unitOfWork.CrashRepository.GetEngineModesAsListItems();
-                result.EngineVersions = _unitOfWork.CrashRepository.GetEngineVersionsAsListItems();
 
-				// Add the FromCollection to the CrashesViewModel since we don't need it for the get results function but we do want to post it back to the page.
+			    using (var unitOfWork = new UnitOfWork(new CrashReportEntities()))
+			    {
+			        result.BranchNames = unitOfWork.CrashRepository.GetBranchesAsListItems();
+                    result.VersionNames = unitOfWork.CrashRepository.GetVersionsAsListItems();
+                    result.PlatformNames = unitOfWork.CrashRepository.GetPlatformsAsListItems();
+                    result.EngineModes = unitOfWork.CrashRepository.GetEngineModesAsListItems();
+                    result.EngineVersions = unitOfWork.CrashRepository.GetEngineVersionsAsListItems();
+			    }
+
+                // Add the FromCollection to the CrashesViewModel since we don't need it for the get results function but we do want to post it back to the page.
 				result.FormCollection = crashesForm;
 				result.GenerationTime = logTimer.GetElapsedSeconds().ToString( "F2" );
 				return View( "Index", result );
@@ -111,16 +112,26 @@ namespace Tools.CrashReporter.CrashReportWebSite.Controllers
 		/// <param name="crashesForm">A form of user data passed up from the client.</param>
 		/// <param name="id">The unique id of the crash we wish to show the details of.</param>
 		/// <returns>A view to show crash details.</returns>
-		public ActionResult Show( FormCollection crashesForm, int id )
+		public ActionResult Show(FormCollection crashesForm, int id)
 		{
 			using( var logTimer = new FAutoScopedLogTimer( this.GetType().ToString() + "(CrashId=" + id + ")", bCreateNewLog: true ) )
 			{
 				CallStackContainer currentCallStack = null;
+                Crash currentCrash = null;
+			    User crashUser = null;
+			    Bugg currentBugg = null;
+			    UserGroup currentUserGroup = null;
 
 				// Update the selected crash based on the form contents
-                var currentCrash = _unitOfWork.CrashRepository.GetById(id);
-                
-				if( currentCrash == null )
+			    using (var unitOfWork = new UnitOfWork(new CrashReportEntities()))
+			    {
+                    currentCrash = unitOfWork.CrashRepository.GetById(id);
+			        crashUser = currentCrash.User;
+			        currentBugg = currentCrash.Bugg;
+			        currentUserGroup = crashUser.UserGroup;
+			    }
+
+			    if( currentCrash == null )
 				{
 					return RedirectToAction( "Index" );
 				}
@@ -153,11 +164,13 @@ namespace Tools.CrashReporter.CrashReportWebSite.Controllers
 				{
 					currentCrash.Description = FormValue;
 				}
+			    using (var unitOfWork = new UnitOfWork(new CrashReportEntities()))
+			    {
+			        unitOfWork.CrashRepository.Update(currentCrash);
+			        unitOfWork.Save();
+			    }
 
-                _unitOfWork.CrashRepository.Update(currentCrash);
-                _unitOfWork.Save();
-
-				currentCallStack = new CallStackContainer( currentCrash );
+			    currentCallStack = new CallStackContainer( currentCrash );
 			    currentCrash.Module = currentCallStack.GetModuleName();
 
 				//Set call stack properties
@@ -169,7 +182,7 @@ namespace Tools.CrashReporter.CrashReportWebSite.Controllers
 
 				currentCrash.CallStackContainer = new CallStackContainer(currentCrash);
 
-				var Model = new CrashViewModel { Crash = currentCrash, CallStack = currentCallStack };
+				var Model = new CrashViewModel { Crash = currentCrash, User = crashUser, Bugg = currentBugg, UserGroup = currentUserGroup, CallStack = currentCallStack };
 				Model.GenerationTime = logTimer.GetElapsedSeconds().ToString( "F2" );
 				return View( "Show", Model );
 			}
@@ -180,7 +193,7 @@ namespace Tools.CrashReporter.CrashReportWebSite.Controllers
 		/// </summary>
 		/// <param name="id">Unused.</param>
 		/// <returns>The row id of the newly added crash.</returns>
-		public ActionResult AddCrash( int id )
+		public ActionResult AddCrash(int id)
 		{
             var newCrashResult = new CrashReporterResult();
             CrashDescription newCrash;
@@ -307,7 +320,7 @@ namespace Tools.CrashReporter.CrashReportWebSite.Controllers
 			}
 
 			var returnResult = XmlHandler.ToXmlString<CrashReporterResult>( newCrashResult );
-                
+
 			return Content( returnResult, "text/xml" );
 		}
 
@@ -324,49 +337,59 @@ namespace Tools.CrashReporter.CrashReportWebSite.Controllers
             var skip = (formData.Page - 1) * formData.PageSize;
             var take = formData.PageSize;
 
-            resultsQuery = ConstructQueryForFiltering(formData);
-
-            // Filter by data and get as enumerable.
-            resultsQuery = FilterByDate(resultsQuery, formData.DateFrom, formData.DateTo);
-
-            // Filter by BuggId
-            if (!string.IsNullOrEmpty(formData.BuggId))
+            Dictionary<string, int> groupCounts;
+            var resultCount = 0;
+            using (var unitOfWork = new UnitOfWork(new CrashReportEntities()))
             {
-                var buggId = 0;
-                var bValid = int.TryParse(formData.BuggId, out buggId);
+                resultsQuery = ConstructQueryForFiltering(unitOfWork, formData);
 
-                if (bValid)
+                // Filter by data and get as enumerable.
+                resultsQuery = FilterByDate(resultsQuery, formData.DateFrom, formData.DateTo);
+
+                // Filter by BuggId
+                if (!string.IsNullOrEmpty(formData.BuggId))
                 {
-                    var newBugg = _unitOfWork.BuggRepository.GetById(buggId);
+                    var buggId = 0;
+                    var bValid = int.TryParse(formData.BuggId, out buggId);
 
-                    if (newBugg != null)
+                    if (bValid)
                     {
-                        resultsQuery = resultsQuery.Where(data => data.PatternId == newBugg.PatternId);
+                        var newBugg = unitOfWork.BuggRepository.GetById(buggId);
+
+                        if (newBugg != null)
+                        {
+                            resultsQuery = resultsQuery.Where(data => data.PatternId == newBugg.PatternId);
+                        }
                     }
                 }
-            }
 
-            var countsQuery = resultsQuery.GroupBy(data => data.User.UserGroup).Select(data => new { Key = data.Key.Name, Count = data.Count()});
-            var groupCounts = countsQuery.OrderBy(data => data.Key).ToDictionary(data => data.Key, data => data.Count);
+                var countsQuery =
+                    resultsQuery.GroupBy(data => data.User.UserGroup)
+                        .Select(data => new {Key = data.Key.Name, Count = data.Count()});
+                groupCounts = countsQuery.OrderBy(data => data.Key)
+                    .ToDictionary(data => data.Key, data => data.Count);
 
-            // Filter by user group if present
-            var userGroupId = !string.IsNullOrEmpty(formData.UserGroup) ? _unitOfWork.UserGroupRepository.First(data => data.Name == formData.UserGroup).Id : 1;
-            resultsQuery = resultsQuery.Where(data => data.User.UserGroupId == userGroupId);
+                // Filter by user group if present
+                var userGroupId = !string.IsNullOrEmpty(formData.UserGroup)
+                    ? unitOfWork.UserGroupRepository.First(data => data.Name == formData.UserGroup).Id
+                    : 1;
+                resultsQuery = resultsQuery.Where(data => data.User.UserGroupId == userGroupId);
 
-            var orderedQuery = GetSortedQuery(resultsQuery, formData.SortTerm ?? "TimeOfCrash", formData.SortOrder == "Descending");
+                var orderedQuery = GetSortedQuery(resultsQuery, formData.SortTerm ?? "TimeOfCrash",
+                    formData.SortOrder == "Descending");
 
-            // Grab just the results we want to display on this page
-            results = orderedQuery.Skip(skip).Take(take).ToList();
+                // Grab just the results we want to display on this page
+                results = orderedQuery.Skip(skip).Take(take).ToList();
 
-            // Get the Count for pagination
-            var resultCount = 0;
-            resultCount = orderedQuery.Count();
+                // Get the Count for pagination
+                resultCount = orderedQuery.Count();
 
-            // Process call stack for display
-            foreach (var crashInstance in results)
-            {
-                // Put call stacks into an list so we can access them line by line in the view
-                crashInstance.CallStackContainer = new CallStackContainer(crashInstance);
+                // Process call stack for display
+                foreach (var crashInstance in results)
+                {
+                    // Put call stacks into an list so we can access them line by line in the view
+                    crashInstance.CallStackContainer = new CallStackContainer(crashInstance);
+                }
             }
 
             return new CrashesViewModel
@@ -454,11 +477,11 @@ namespace Tools.CrashReporter.CrashReportWebSite.Controllers
 	    }
 
         /// <summary>Constructs query for filtering.</summary>
-        private IQueryable<Crash> ConstructQueryForFiltering(FormHelper formData)
+        private IQueryable<Crash> ConstructQueryForFiltering(IUnitOfWork unitOfWork, FormHelper formData)
         {
             //Instead of returning a queryable and filtering it here we should construct a filtering expression and pass that to the repository.
             //I don't like that data handling is taking place in the controller.
-            var results = _unitOfWork.CrashRepository.ListAll();
+            var results = unitOfWork.CrashRepository.ListAll();
 
             // Grab Results 
             var queryString = HttpUtility.HtmlDecode(formData.SearchQuery);
@@ -473,6 +496,11 @@ namespace Tools.CrashReporter.CrashReportWebSite.Controllers
                 results = results.Where(data => data.RawCallStack.Contains(formData.SearchQuery));
             }
 
+            if (formData.IsVanilla.HasValue)
+            {
+                results = results.Where(data => data.IsVanilla == !data.IsVanilla.HasValue || formData.IsVanilla.Value);
+            }
+
             // Filter by Crash Type
             if (formData.CrashType != "All")
             {
@@ -497,15 +525,14 @@ namespace Tools.CrashReporter.CrashReportWebSite.Controllers
             if (!string.IsNullOrEmpty(formData.UsernameQuery))
             {
                 var decodedUsername = HttpUtility.HtmlDecode(formData.UsernameQuery).ToLower();
-                var user = _unitOfWork.UserRepository.First(data => data.UserName.Contains(decodedUsername));
+                var user = unitOfWork.UserRepository.First(data => data.UserName.Contains(decodedUsername));
                 if (user != null)
                 {
-                    var userId = user.Id; 
+                    var userId = user.Id;
                     results = (
                         from crashDetail in results
                         where crashDetail.UserNameId == userId
-                        select crashDetail
-                        );
+                        select crashDetail);
                 }
             }
 
@@ -678,21 +705,25 @@ namespace Tools.CrashReporter.CrashReportWebSite.Controllers
 	            ChangeListVersion = description.BuiltFromCL.ToString(),
 	            CommandLine = description.CommandLine,
 	            EngineMode = description.EngineMode,
-	            ComputerName = description.MachineGuid
+	            ComputerName = description.MachineGuid,
+                IsVanilla = description.EngineModeEx.ToLower() == "vanilla"
 	        };
 
-	        //if there's a valid username assign the associated UserNameId else use "anonymous".
-	        var userName = (!string.IsNullOrEmpty(description.UserName)) ? description.UserName : UserNameAnonymous;
-	        var user = _unitOfWork.UserRepository.GetByUserName(userName);
-
-	        if (user != null)
+	        using (var unitOfWork = new UnitOfWork(new CrashReportEntities()))
 	        {
-	            newCrash.UserNameId = user.Id;
-	            newCrash.UserName = user.UserName;
-	        }
-	        else
-	        {
-	            newCrash.User = new User() {UserName = description.UserName, UserGroupId = 5};
+	            //if there's a valid username assign the associated UserNameId else use "anonymous".
+	            var userName = (!string.IsNullOrEmpty(description.UserName)) ? description.UserName : UserNameAnonymous;
+	            var user = unitOfWork.UserRepository.GetByUserName(userName);
+
+	            if (user != null)
+	            {
+	                newCrash.UserNameId = user.Id;
+	                newCrash.UserName = user.UserName;
+	            }
+	            else
+	            {
+	                newCrash.User = new User() {UserName = description.UserName, UserGroupId = 5};
+	            }
 	        }
 
 	        //If there's a valid EpicAccountId assign that.
@@ -706,6 +737,8 @@ namespace Tools.CrashReporter.CrashReportWebSite.Controllers
 	        {
 	            newCrash.Description = string.Join(Environment.NewLine, description.UserDescription);
 	        }
+	        if (newCrash.Description.Length > 4095)
+	            newCrash.Description = newCrash.Description.Substring(0, 4095);
 
 	        newCrash.EngineMode = description.EngineMode;
 	        newCrash.GameName = description.GameName;
@@ -797,183 +830,195 @@ namespace Tools.CrashReporter.CrashReportWebSite.Controllers
             if(newCrash.CommandLine == null)
                 newCrash.CommandLine = "";
 
-            _unitOfWork.Dispose();
-            _unitOfWork = new UnitOfWork(new CrashReportEntities());
-	        var callStackRepository = _unitOfWork.CallstackRepository;
-
-	        try
+	        using (var unitOfWork = new UnitOfWork(new CrashReportEntities()))
 	        {
-	            var crashRepo = _unitOfWork.CrashRepository;
-                //if we don't have any callstack data then insert the crash and return
-	            if (string.IsNullOrEmpty(newCrash.Pattern))
-	            {
-	                crashRepo.Save(newCrash);
-                    _unitOfWork.Save();
-	                return newCrash;
-	            }
+	            var callStackRepository = unitOfWork.CallstackRepository;
 
-                //If this isn't a new pattern then link it to our crash data model
-                if (callStackRepository.Any(data => data.Pattern == newCrash.Pattern))
-	            {
-                    var callstackPattern = callStackRepository.First(data => data.Pattern == newCrash.Pattern);
-	                newCrash.PatternId = callstackPattern.id;
-	            }
-	            else
+	            try
 	            {
-	                //if this is a new callstack pattern then insert into data model and create a new bugg.
-	                var callstackPattern = new CallStackPattern { Pattern = newCrash.Pattern };
-	                callStackRepository.Save(callstackPattern);
-                    _unitOfWork.Save();
-	                newCrash.PatternId = callstackPattern.id;
-	            }
+	                var crashRepo = unitOfWork.CrashRepository;
 
-                //Mask out the line number and File path from our error message.
-                var errorMessageString = description.ErrorMessage != null ? String.Join("", description.ErrorMessage) : "";
+	                //if we don't have any callstack data then insert the crash and return
+	                if (string.IsNullOrEmpty(newCrash.Pattern))
+	                {
+	                    crashRepo.Save(newCrash);
+	                    unitOfWork.Save();
+	                    return newCrash;
+	                }
 
-                //Create our masking regular expressions
-                var fileRegex = new Regex(@"(\[File:).*?(])");//Match the filename out the file name
-                var lineRegex = new Regex(@"(\[Line:).*?(])");//Match the line no.
+	                //If this isn't a new pattern then link it to our crash data model
+	                if (callStackRepository.Any(data => data.Pattern == newCrash.Pattern))
+	                {
+	                    var callstackPattern = callStackRepository.First(data => data.Pattern == newCrash.Pattern);
+	                    newCrash.PatternId = callstackPattern.id;
+	                }
+	                else
+	                {
+	                    //if this is a new callstack pattern then insert into data model and create a new bugg.
+	                    var callstackPattern = new CallStackPattern {Pattern = newCrash.Pattern};
+	                    callStackRepository.Save(callstackPattern);
+	                    unitOfWork.Save();
+	                    newCrash.PatternId = callstackPattern.id;
+	                }
 
-                /**
+	                //Mask out the line number and File path from our error message.
+	                var errorMessageString = description.ErrorMessage != null
+	                    ? String.Join("", description.ErrorMessage)
+	                    : "";
+
+	                //Create our masking regular expressions
+	                var fileRegex = new Regex(@"(\[File:).*?(])"); //Match the filename out the file name
+	                var lineRegex = new Regex(@"(\[Line:).*?(])"); //Match the line no.
+
+	                /**
                  * Regex to match ints of two characters or longer
                  * First term ((?<=\s)|(-)) : Positive look behind, match if preceeded by whitespace or if first character is '-'
                  * Second term (\d{3,}) match three or more decimal chracters in a row.
                  * Third term (?=(\s|$)) positive look ahead match if followed by whitespace or end of line/file.
                  */
-                var intRegex = new Regex(@"-?\d{3,}");
-                
-                /**
+	                var intRegex = new Regex(@"-?\d{3,}");
+
+	                /**
                  * Regular expression for masking out floats
                  */
-                var floatRegex = new Regex(@"-?\d+\.\d+");
+	                var floatRegex = new Regex(@"-?\d+\.\d+");
 
-                /**
+	                /**
                  * Regular expression for masking out hexadecimal numbers
                  */
-                var hexRegex = new Regex(@"0x[\da-fA-F]+");
-
-                //mask out terms matches by our regex's
-                var trimmedError = fileRegex.Replace(errorMessageString, "");
-                trimmedError = lineRegex.Replace(trimmedError, "");
-                trimmedError = floatRegex.Replace(trimmedError, "");
-                trimmedError = hexRegex.Replace(trimmedError, "");
-                trimmedError = intRegex.Replace(trimmedError, "");
-                if(trimmedError.Length > 450)
-                    trimmedError = trimmedError.Substring(0, 450);//error message is used as an index - trim to max index length
-                
-                //Check to see if the masked error message is unique
-	            ErrorMessage errorMessage = null;
-                if (_unitOfWork.ErrorMessageRepository.Any(data => data.Message.ToLower().Contains(trimmedError.ToLower())))
-                {
-                    errorMessage =
-                        _unitOfWork.ErrorMessageRepository.First(data => data.Message.ToLower().Contains(trimmedError.ToLower()));
-                }
-                else
-                {
-                    //if it's a new message then add it to the database.
-                    errorMessage = new ErrorMessage() { Message = trimmedError };
-                    _unitOfWork.ErrorMessageRepository.Save(errorMessage);
-                    _unitOfWork.Save();
-                }
+	                var hexRegex = new Regex(@"0x[\da-fA-F]+");
+
+	                //mask out terms matches by our regex's
+	                var trimmedError = fileRegex.Replace(errorMessageString, "");
+	                trimmedError = lineRegex.Replace(trimmedError, "");
+	                trimmedError = floatRegex.Replace(trimmedError, "");
+	                trimmedError = hexRegex.Replace(trimmedError, "");
+	                trimmedError = intRegex.Replace(trimmedError, "");
+	                if (trimmedError.Length > 450)
+	                    trimmedError = trimmedError.Substring(0, 450);
+	                    //error message is used as an index - trim to max index length
+
+	                //Check to see if the masked error message is unique
+	                ErrorMessage errorMessage = null;
+	                if (
+	                    unitOfWork.ErrorMessageRepository.Any(
+	                        data => data.Message.Equals(trimmedError, StringComparison.InvariantCultureIgnoreCase)))
+	                {
+	                    errorMessage =
+	                        unitOfWork.ErrorMessageRepository.First(
+	                            data => data.Message.Equals(trimmedError, StringComparison.InvariantCultureIgnoreCase));
+	                }
+	                else
+	                {
+	                    //if it's a new message then add it to the database.
+	                    errorMessage = new ErrorMessage() {Message = trimmedError.ToLower(CultureInfo.InvariantCulture)};
+	                    unitOfWork.ErrorMessageRepository.Save(errorMessage);
+	                    unitOfWork.Save();
+	                }
 
-	            //Check for an existing bugg with this pattern and error message / no error message
-                if ( _unitOfWork.BuggRepository.Any(data =>
-                    (data.PatternId == newCrash.PatternId) && (data.ErrorMessageId == errorMessage.Id || data.ErrorMessageId == null) ))
-                {
-                    //if a bugg exists for this pattern update the bugg data
-                    var bugg = _unitOfWork.BuggRepository.First(data => data.PatternId == newCrash.PatternId);
-                    bugg.PatternId = newCrash.PatternId;
+	                //Check for an existing bugg with this pattern and error message / no error message
+	                if (unitOfWork.BuggRepository.Any(data =>
+	                    (data.PatternId == newCrash.PatternId) &&
+	                    (data.ErrorMessageId == errorMessage.Id || data.ErrorMessageId == null)))
+	                {
+	                    //if a bugg exists for this pattern update the bugg data
+	                    var bugg = unitOfWork.BuggRepository.First(data => data.PatternId == newCrash.PatternId);
+	                    bugg.PatternId = newCrash.PatternId;
+
+	                    bugg.CrashType = newCrash.CrashType;
+	                    bugg.ErrorMessageId = errorMessage.Id;
 
-                    bugg.CrashType = newCrash.CrashType;
-                    bugg.ErrorMessageId = errorMessage.Id;
+	                    //also update the bugg data while we're here
+	                    bugg.TimeOfLastCrash = newCrash.TimeOfCrash;
+	                    bugg.NumberOfCrashes = bugg.NumberOfCrashes + 1;
 
-	                //also update the bugg data while we're here
-	                bugg.TimeOfLastCrash = newCrash.TimeOfCrash;
-	                
-	                if (String.Compare(newCrash.BuildVersion, bugg.BuildVersion,
+	                    if (String.Compare(newCrash.BuildVersion, bugg.BuildVersion,
 	                        StringComparison.Ordinal) != 1)
+	                        bugg.BuildVersion = newCrash.BuildVersion;
+
+	                    unitOfWork.Save();
+
+	                    //if a bugg exists update this crash from the bugg
+	                    //buggs are authoritative in this case
+	                    newCrash.BuggId = bugg.Id;
+	                    newCrash.Jira = bugg.TTPID;
+	                    newCrash.FixedChangeList = bugg.FixedChangeList;
+	                    newCrash.Status = bugg.Status;
+	                    newCrash.Buggs.Add(bugg);
+	                    unitOfWork.CrashRepository.Save(newCrash);
+	                    unitOfWork.Save();
+	                }
+	                else
+	                {
+	                    //if there's no bugg for this pattern create a new bugg and insert into the data store.
+	                    var bugg = new Bugg();
+	                    bugg.TimeOfFirstCrash = newCrash.TimeOfCrash;
+	                    bugg.TimeOfLastCrash = newCrash.TimeOfCrash;
+	                    bugg.TTPID = newCrash.Jira;
+	                    bugg.Pattern = newCrash.Pattern;
+	                    bugg.PatternId = newCrash.PatternId;
+	                    bugg.NumberOfCrashes = 1;
+	                    bugg.NumberOfUsers = 1;
+	                    bugg.NumberOfUniqueMachines = 1;
 	                    bugg.BuildVersion = newCrash.BuildVersion;
-
-                    _unitOfWork.Save();
-
-	                //if a bugg exists update this crash from the bugg
-	                //buggs are authoritative in this case
-                    newCrash.BuggId = bugg.Id;
-	                newCrash.Jira = bugg.TTPID;
-	                newCrash.FixedChangeList = bugg.FixedChangeList;
-	                newCrash.Status = bugg.Status;
-                    newCrash.Buggs.Add(bugg);
-	                _unitOfWork.CrashRepository.Save(newCrash);
-	                _unitOfWork.Save();
+	                    bugg.CrashType = newCrash.CrashType;
+	                    bugg.Status = newCrash.Status;
+	                    bugg.FixedChangeList = newCrash.FixedChangeList;
+	                    bugg.ErrorMessageId = errorMessage.Id;
+	                    bugg.NumberOfCrashes = 1;
+	                    newCrash.Bugg = bugg;
+	                    newCrash.Buggs.Add(bugg);
+	                    unitOfWork.BuggRepository.Save(bugg);
+	                    unitOfWork.CrashRepository.Save(newCrash);
+	                    unitOfWork.Save();
+	                }
 	            }
-	            else
+	            catch (DbEntityValidationException dbentEx)
 	            {
-	                //if there's no bugg for this pattern create a new bugg and insert into the data store.
-	                var bugg = new Bugg();
-	                bugg.TimeOfFirstCrash = newCrash.TimeOfCrash;
-	                bugg.TimeOfLastCrash = newCrash.TimeOfCrash;
-	                bugg.TTPID = newCrash.Jira;
-                    bugg.Pattern = newCrash.Pattern;
-	                bugg.PatternId = newCrash.PatternId;
-	                bugg.NumberOfCrashes = 1;
-	                bugg.NumberOfUsers = 1;
-	                bugg.NumberOfUniqueMachines = 1;
-	                bugg.BuildVersion = newCrash.BuildVersion;
-	                bugg.CrashType = newCrash.CrashType;
-	                bugg.Status = newCrash.Status;
-                    bugg.FixedChangeList = newCrash.FixedChangeList;
-	                bugg.ErrorMessageId = errorMessage.Id;
-                    newCrash.Bugg = bugg;
-                    newCrash.Buggs.Add(bugg);
-                    _unitOfWork.BuggRepository.Save(bugg);
-                    _unitOfWork.CrashRepository.Save(newCrash);
-                    _unitOfWork.Save();
-	            }
-	        }
-            catch (DbEntityValidationException dbentEx)
-            {
-                var messageBuilder = new StringBuilder();
-                messageBuilder.AppendLine("Db Entity Validation Exception Exception was:");
-                messageBuilder.AppendLine(dbentEx.ToString());
+	                var messageBuilder = new StringBuilder();
+	                messageBuilder.AppendLine("Db Entity Validation Exception Exception was:");
+	                messageBuilder.AppendLine(dbentEx.ToString());
 
-                var innerEx = dbentEx.InnerException;
-                while (innerEx != null)
-                {
-                    messageBuilder.AppendLine("Inner Exception : " + innerEx.Message);
-                    innerEx = innerEx.InnerException;
-                }
+	                var innerEx = dbentEx.InnerException;
+	                while (innerEx != null)
+	                {
+	                    messageBuilder.AppendLine("Inner Exception : " + innerEx.Message);
+	                    innerEx = innerEx.InnerException;
+	                }
 
-                if (dbentEx.EntityValidationErrors != null)
-                {
-                    messageBuilder.AppendLine("Validation Errors : ");
-                    foreach (var valErr in dbentEx.EntityValidationErrors)
-                    {
-                        messageBuilder.AppendLine(valErr.ValidationErrors.Select(data => data.ErrorMessage).Aggregate((current, next) => current + "; /n" + next));
-                    }
-                }
-                FLogger.Global.WriteException(messageBuilder.ToString());
-                _slackWriter.Write("Create Crash Exception : " + messageBuilder.ToString());
-                throw;
-            }
-	        catch (Exception ex)
-	        {
-                var messageBuilder = new StringBuilder();
-                messageBuilder.AppendLine("Create Crash Exception : ");
-                messageBuilder.AppendLine(ex.Message.ToString());
-                var innerEx = ex.InnerException;
-                while (innerEx != null)
-                {
-                    messageBuilder.AppendLine("Inner Exception : " + innerEx.Message);
-                    innerEx = innerEx.InnerException;
-                }
+	                if (dbentEx.EntityValidationErrors != null)
+	                {
+	                    messageBuilder.AppendLine("Validation Errors : ");
+	                    foreach (var valErr in dbentEx.EntityValidationErrors)
+	                    {
+	                        messageBuilder.AppendLine(
+	                            valErr.ValidationErrors.Select(data => data.ErrorMessage)
+	                                .Aggregate((current, next) => current + "; /n" + next));
+	                    }
+	                }
+	                FLogger.Global.WriteException(messageBuilder.ToString());
+	                _slackWriter.Write("Create Crash Exception : " + messageBuilder.ToString());
+	                throw;
+	            }
+	            catch (Exception ex)
+	            {
+	                var messageBuilder = new StringBuilder();
+	                messageBuilder.AppendLine("Create Crash Exception : ");
+	                messageBuilder.AppendLine(ex.Message.ToString());
+	                var innerEx = ex.InnerException;
+	                while (innerEx != null)
+	                {
+	                    messageBuilder.AppendLine("Inner Exception : " + innerEx.Message);
+	                    innerEx = innerEx.InnerException;
+	                }
 
-	            FLogger.Global.WriteException("Create Crash Exception : " +
-	                                          messageBuilder.ToString());
-                _slackWriter.Write("Create Crash Exception : " + messageBuilder.ToString());
-	            throw;
+	                FLogger.Global.WriteException("Create Crash Exception : " +
+	                                              messageBuilder.ToString());
+	                _slackWriter.Write("Create Crash Exception : " + messageBuilder.ToString());
+	                throw;
+	            }
 	        }
-
-            return newCrash;
+	        return newCrash;
 	    }
 
 	    /// <summary>
@@ -991,27 +1036,34 @@ namespace Tools.CrashReporter.CrashReportWebSite.Controllers
 
 	            try
 	            {
-	                foreach (var entry in callstack.CallStackEntries.Take(CallStackContainer.MaxLinesToParse))
+	                using (var unitOfWork = new UnitOfWork(new CrashReportEntities()))
 	                {
-	                    FunctionCall currentFunctionCall;
 
-	                    var csEntry = entry;
-                        var functionCall = _unitOfWork.FunctionRepository.First(f => f.Call == csEntry.FunctionName);
-                        if (functionCall != null)
-                        {
-                            currentFunctionCall = functionCall;
-                        }
-	                    else
+	                    foreach (var entry in callstack.CallStackEntries.Take(CallStackContainer.MaxLinesToParse))
 	                    {
-	                        currentFunctionCall = new FunctionCall { Call = csEntry.FunctionName };
-	                        _unitOfWork.FunctionRepository.Save(currentFunctionCall);
-	                        _unitOfWork.Save();
+	                        FunctionCall currentFunctionCall;
+
+	                        var csEntry = entry;
+	                        var functionCall = unitOfWork.FunctionRepository.First(f => f.Call == csEntry.FunctionName);
+	                        if (functionCall != null)
+	                        {
+	                            currentFunctionCall = functionCall;
+	                        }
+	                        else
+	                        {
+	                            var call = csEntry.FunctionName;
+	                            if (call.Length > 899)
+	                                call = call.Substring(0, 899);
+                                currentFunctionCall = new FunctionCall { Call = call };
+	                            unitOfWork.FunctionRepository.Save(currentFunctionCall);
+	                            unitOfWork.Save();
+	                        }
+
+	                        patternList.Add(currentFunctionCall.Id.ToString());
 	                    }
 
-	                    patternList.Add(currentFunctionCall.Id.ToString());
+	                    newCrash.Pattern = string.Join("+", patternList);
 	                }
-
-	                newCrash.Pattern = string.Join("+", patternList);
 	            }
 	            catch (Exception ex)
 	            {
@@ -1031,20 +1083,25 @@ namespace Tools.CrashReporter.CrashReportWebSite.Controllers
 	        }
 	    }
 
+        /// <summary></summary>
+        /// <param name="disposing"></param>
 	    protected override void Dispose(bool disposing)
 	    {
-	        if (disposing)
-	        {
-	            if (_unitOfWork != null)
-	            {
-	                _unitOfWork.Dispose();
-	                _unitOfWork = null;
-	            }
-	        }
+            //if (disposing)
+            //{
+            //    if (_unitOfWork != null)
+            //    {
+            //        _unitOfWork.Dispose();
+            //        _unitOfWork = null;
+            //    }
+            //}
 
 	        base.Dispose(disposing);
 	    }
 
+        /// <summary>
+        /// 
+        /// </summary>
 	    public void TestAddCrash()
 	    {
 	        string crashContext;
diff --git a/Engine/Source/Programs/CrashReporter/CrashReportWebSite/Controllers/HomeController.cs b/Engine/Source/Programs/CrashReporter/CrashReportWebSite/Controllers/HomeController.cs
--- a/Engine/Source/Programs/CrashReporter/CrashReportWebSite/Controllers/HomeController.cs
+++ b/Engine/Source/Programs/CrashReporter/CrashReportWebSite/Controllers/HomeController.cs
@@ -1,5 +1,6 @@
 ﻿// Copyright 1998-2017 Epic Games, Inc. All Rights Reserved.
 using System.Web.Mvc;
+using Tools.CrashReporter.CrashReportWebSite.DataModels;
 using Tools.CrashReporter.CrashReportWebSite.DataModels.Repositories;
 using Tools.CrashReporter.CrashReportWebSite.ViewModels;
 
@@ -8,14 +9,12 @@ namespace Tools.CrashReporter.CrashReportWebSite.Controllers
 	/// <summary>Controls the start page.</summary>
 	public class HomeController : Controller
 	{
-	    private readonly IUnitOfWork _unitOfWork;
-
         /// <summary>
         /// 
         /// </summary>
 	    public HomeController(IUnitOfWork unitOfWork)
 	    {
-	        _unitOfWork = unitOfWork;
+	        
 	    }
 
 		/// <summary>
@@ -24,23 +23,27 @@ namespace Tools.CrashReporter.CrashReportWebSite.Controllers
 		public ActionResult Index()
 		{
             using (var logTimer = new FAutoScopedLogTimer( this.GetType().ToString(), bCreateNewLog: true ))
-			{
-				var result = new CrashesViewModel
-				{
-				    BranchNames = _unitOfWork.CrashRepository.GetBranchesAsListItems(),
-				    VersionNames = _unitOfWork.CrashRepository.GetVersionsAsListItems(),
-				    PlatformNames = _unitOfWork.CrashRepository.GetPlatformsAsListItems(),
-				    EngineModes = _unitOfWork.CrashRepository.GetEngineModesAsListItems(),
-				    EngineVersions = _unitOfWork.CrashRepository.GetEngineVersionsAsListItems(),
-				    GenerationTime = logTimer.GetElapsedSeconds().ToString("F2")
-				};
+            {
+                var result = new CrashesViewModel();
+                using (var entitile = new CrashReportEntities())
+                {
+                    var unitOfWork = new UnitOfWork(entitile);
+                
+                    result.BranchNames = unitOfWork.CrashRepository.GetBranchesAsListItems();
+                    result.VersionNames = unitOfWork.CrashRepository.GetVersionsAsListItems();
+                    result.PlatformNames = unitOfWork.CrashRepository.GetPlatformsAsListItems();
+                    result.EngineModes = unitOfWork.CrashRepository.GetEngineModesAsListItems();
+                    result.EngineVersions = unitOfWork.CrashRepository.GetEngineVersionsAsListItems();
+                    result.IsVanilla = null;
+                    result.GenerationTime = logTimer.GetElapsedSeconds().ToString("F2");
+                }
+
 			    return View( "Index", result );
 			}
 		}
 
         protected override void Dispose(bool disposing)
         {
-            _unitOfWork.Dispose();
             base.Dispose(disposing);
         }
 	}
diff --git a/Engine/Source/Programs/CrashReporter/CrashReportWebSite/Controllers/UsersController.cs b/Engine/Source/Programs/CrashReporter/CrashReportWebSite/Controllers/UsersController.cs
--- a/Engine/Source/Programs/CrashReporter/CrashReportWebSite/Controllers/UsersController.cs
+++ b/Engine/Source/Programs/CrashReporter/CrashReportWebSite/Controllers/UsersController.cs
@@ -1,7 +1,8 @@
-﻿// Copyright 1998-2017 Epic Games, Inc. All Rights Reserved.
+﻿// Copyright 1998-2016 Epic Games, Inc. All Rights Reserved.
 
 using System.Linq;
 using System.Web.Mvc;
+using Tools.CrashReporter.CrashReportWebSite.DataModels;
 using Tools.CrashReporter.CrashReportWebSite.DataModels.Repositories;
 using Tools.CrashReporter.CrashReportWebSite.ViewModels;
 
@@ -12,7 +13,6 @@ namespace Tools.CrashReporter.CrashReportWebSite.Controllers
 	/// </summary>
 	public class UsersController : Controller
 	{
-	    private IUnitOfWork _unitOfWork;
 
         /// <summary>
         /// 
@@ -20,7 +20,6 @@ namespace Tools.CrashReporter.CrashReportWebSite.Controllers
         /// <param name="unitOfWork"></param>
 	    public UsersController(IUnitOfWork unitOfWork)
 	    {
-	        _unitOfWork = unitOfWork;
 	    }
 
 	    /// <summary>
@@ -32,49 +31,65 @@ namespace Tools.CrashReporter.CrashReportWebSite.Controllers
 	    /// <returns>A view to show a list of users in the current user group.</returns>
 	    public ActionResult Index( FormCollection formData, string userName, string userGroup = "General" )
         {
-            var group = _unitOfWork.UserGroupRepository.First(data => data.Name == userGroup);
+	        UsersViewModel model;
+	        using (var unitOfWork = new UnitOfWork(new CrashReportEntities()))
+	        {
+	            var group = unitOfWork.UserGroupRepository.First(data => data.Name == userGroup);
 
-            if (!string.IsNullOrWhiteSpace(userName))
-            {
-                var user = _unitOfWork.UserRepository.GetByUserName(userName);
-                if (user != null && group != null)
-                {
-                    user.UserGroup = group;
-                    _unitOfWork.UserRepository.Update(user);
-                    _unitOfWork.Save();
-                }
-            }
+	            if (!string.IsNullOrWhiteSpace(userName))
+	            {
+	                var user = unitOfWork.UserRepository.GetByUserName(userName);
+	                if (user != null && group != null)
+	                {
+	                    user.UserGroup = group;
+	                    unitOfWork.UserRepository.Update(user);
+	                    unitOfWork.Save();
+	                }
+	            }
 
-	        var formhelper = new FormHelper(Request, formData, "");
-	        var skip = (formhelper.Page -1) * formhelper.PageSize;
-	        var take = formhelper.PageSize;
+	            var formhelper = new FormHelper(Request, formData, "");
+	            var skip = (formhelper.Page - 1)*formhelper.PageSize;
+	            var take = formhelper.PageSize;
 
-			var model = new UsersViewModel
-			{
-			    UserGroup = userGroup,
-			    User = userName,
-			    Users = _unitOfWork.UserRepository.ListAll().Where(data => data.UserGroupId == group.Id).OrderBy(data => data.UserName)
-                    .Skip(skip)
-                    .Take(take)
-			        .Select(data => new UserViewModel() {Name = data.UserName, UserGroup = data.UserGroup.Name})
-			        .ToList()
-			};
+	            model = new UsersViewModel
+	            {
+	                UserGroup = userGroup,
+	                User = userName,
+	                Users =
+	                    unitOfWork.UserRepository.ListAll()
+	                        .Where(data => data.UserGroupId == @group.Id)
+	                        .OrderBy(data => data.UserName)
+	                        .Skip(skip)
+	                        .Take(take)
+	                        .Select(data => new UserViewModel() {Name = data.UserName, UserGroup = data.UserGroup.Name})
+	                        .ToList()
+	            };
 
-		    var groupCounts =
-		        _unitOfWork.UserRepository.ListAll()
-		            .GroupBy(data => data.UserGroup)
-		            .Select(data => new { Key = data.Key, Count = data.Count() }).ToDictionary( groupCount => groupCount.Key.Name, groupCount => groupCount.Count );
-            
-		    model.GroupSelectList = groupCounts.Select(listItem => new SelectListItem{Selected = listItem.Key == userGroup, Text = listItem.Key, Value = listItem.Key}).ToList();
-            model.GroupCounts = groupCounts;
-	        model.PagingInfo = new PagingInfo
-	        {
-	            CurrentPage = formhelper.Page,
-	            PageSize = formhelper.PageSize,
-	            TotalResults = _unitOfWork.UserRepository.ListAll().Count(data => data.UserGroupId == group.Id)
-	        };
+	            var groupCounts =
+	                unitOfWork.UserRepository.ListAll()
+	                    .GroupBy(data => data.UserGroup)
+	                    .Select(data => new {Key = data.Key, Count = data.Count()})
+	                    .ToDictionary(groupCount => groupCount.Key.Name, groupCount => groupCount.Count);
 
-			return View( "Index", model );
+	            model.GroupSelectList =
+	                groupCounts.Select(
+	                    listItem =>
+	                        new SelectListItem
+	                        {
+	                            Selected = listItem.Key == userGroup,
+	                            Text = listItem.Key,
+	                            Value = listItem.Key
+	                        }).ToList();
+	            model.GroupCounts = groupCounts;
+	            model.PagingInfo = new PagingInfo
+	            {
+	                CurrentPage = formhelper.Page,
+	                PageSize = formhelper.PageSize,
+	                TotalResults = unitOfWork.UserRepository.ListAll().Count(data => data.UserGroupId == group.Id)
+	            };
+	        }
+
+	        return View( "Index", model );
 		}
 
         /// <summary>
@@ -84,13 +99,19 @@ namespace Tools.CrashReporter.CrashReportWebSite.Controllers
         /// <returns></returns>
 	    public JsonResult AutocompleteUser(string userName)
         {
-            var users = _unitOfWork.UserRepository.Get(data => data.UserName.Contains(userName)).Take(10).Select(data => new { UserName = data.UserName, Group = data.UserGroup.Name}).ToList();
-            return Json(users, JsonRequestBehavior.AllowGet);
+            using (var unitOfWork = new UnitOfWork(new CrashReportEntities()))
+            {
+                var users =
+                    unitOfWork.UserRepository.Get(data => data.UserName.Contains(userName))
+                        .Take(10)
+                        .Select(data => new {UserName = data.UserName, Group = data.UserGroup.Name})
+                        .ToList();
+                return Json(users, JsonRequestBehavior.AllowGet);
+            }
         }
 
         protected override void Dispose(bool disposing)
         {
-            _unitOfWork.Dispose();
             base.Dispose(disposing);
         }
 	}
diff --git a/Engine/Source/Programs/CrashReporter/CrashReportWebSite/DataModels/BuggDerived.cs b/Engine/Source/Programs/CrashReporter/CrashReportWebSite/DataModels/BuggDerived.cs
--- a/Engine/Source/Programs/CrashReporter/CrashReportWebSite/DataModels/BuggDerived.cs
+++ b/Engine/Source/Programs/CrashReporter/CrashReportWebSite/DataModels/BuggDerived.cs
@@ -2,9 +2,12 @@
 
 using System;
 using System.Collections.Generic;
+using System.IO;
 using System.Linq;
+using System.Net;
 using System.Web;
 using Newtonsoft.Json;
+using Newtonsoft.Json.Linq;
 using Tools.CrashReporter.CrashReportCommon;
 using Tools.CrashReporter.CrashReportWebSite.Properties;
 
@@ -104,7 +107,7 @@ namespace Tools.CrashReporter.CrashReportWebSite.DataModels
             var jiraConnection = JiraConnection.Get();
 
             this.AffectedVersions = new SortedSet<string>();
-            this.AffectedMajorVersions = new SortedSet<string>(); // 4.4, 4.5 and so
+            this.AffectedMajorVersions = new SortedSet<string>(); // 4.4, 4.5 and so on
             this.BranchesFoundIn = new SortedSet<string>();
             this.AffectedPlatforms = new SortedSet<string>();
             var hashSetDescriptions = new HashSet<string>();
@@ -152,7 +155,7 @@ namespace Tools.CrashReporter.CrashReportWebSite.DataModels
                 }
             }
 
-            // CopyToJira 
+            //ToJiraDescriptons
             foreach (var line in hashSetDescriptions)
             {
                 var listItem = "- " + HttpUtility.HtmlEncode(line);
@@ -255,7 +258,7 @@ namespace Tools.CrashReporter.CrashReportWebSite.DataModels
         }
 
         /// <summary>
-        /// 
+        /// Create a new Jira for this bug
         /// </summary>
         public void CopyToJira()
         {
@@ -274,7 +277,7 @@ namespace Tools.CrashReporter.CrashReportWebSite.DataModels
                     issueFields = CreateOrionIssue(jc);
                     break;
                 default:
-                    issueFields = CreateGeneralIssue(jc);
+                    issueFields = CreateProjectIssue(jc, this.JiraProject);
                     break;
             }
 
@@ -321,7 +324,7 @@ namespace Tools.CrashReporter.CrashReportWebSite.DataModels
             
             // Callstack customfield_11807
             string JiraCallstack = "{noformat}" + string.Join("\r\n", ToJiraFunctionCalls) + "{noformat}";
-            fields.Add("customfield_11807", JiraCallstack);								// Callstack
+            fields.Add("customfield_11807", JiraCallstack);								    // Callstack
 
             string BuggLink = "http://crashreporter/Buggs/Show/" + Id;
             fields.Add("customfield_11205", BuggLink);
@@ -341,18 +344,13 @@ namespace Tools.CrashReporter.CrashReportWebSite.DataModels
             fields.Add("description", string.Join("\r\n", ToJiraDescriptions)); // Description
             fields.Add("project", new Dictionary<string, object> { { "id", 10600 } });
             fields.Add("issuetype", new Dictionary<string, object> { { "id", "1" } }); // Bug
+
             //branch found in - required = false
             fields.Add("customfield_11201", ToBranchName.ToList());
 
-            //found CL = required = false
-            //fields.Add("customfield_11500", th );
-
             //Platforms - required = false
             fields.Add("customfield_11203", new List<object>() { new Dictionary<string, object>() { { "self", Settings.Default.JiraDeploymentAddress + "/customFieldOption/11425" }, { "value", "PC" } } });
-
-            //repro rate required = false
-            //fields.Add("customfield_11900", this.Buggs_Crashes.Count);
-
+            
             // Callstack customfield_11807
             string JiraCallstack = "{noformat}" + string.Join("\r\n", ToJiraFunctionCalls) + "{noformat}";
             fields.Add("customfield_11807", JiraCallstack);         // Callstack
@@ -360,18 +358,108 @@ namespace Tools.CrashReporter.CrashReportWebSite.DataModels
             return fields;
         }
 
+        /// <summary>
+        /// 
+        /// </summary>
+        /// <param name="jc"></param>
+        /// <returns></returns>
         public Dictionary<string, object> CreateOrionIssue(JiraConnection jc)
         {
             var fields = new Dictionary<string, object>();
 
             fields.Add("summary", "[CrashReport] " + ToJiraSummary); // Call Stack, Line 1
             fields.Add("description", string.Join("\r\n", ToJiraDescriptions)); // Description
-            fields.Add("project", new Dictionary<string, object> { { "id", 10700 } });
+            fields.Add("project", new Dictionary<string, object> { { "id", 10700 } }); 
             fields.Add("issuetype", new Dictionary<string, object> { { "id", "1" } }); // Bug
             fields.Add("components", new object[]{new Dictionary<string, object>{{"id", "14604"}}});
 
             return fields;
-        } 
+        }
+
+        /// <summary>
+        /// Create an issue description from a project ID
+        /// </summary>
+        /// <param name="jc"></param>
+        /// <param name="projectId"></param>
+        /// <returns></returns>
+        public Dictionary<string, object> CreateProjectIssue(JiraConnection jc, string projectId)
+        {
+            var response = jc.JiraRequest("/issue/createmeta?projectKeys="+ projectId +"&expand=projects.issuetypes.fields",
+                JiraConnection.JiraMethod.GET, null, HttpStatusCode.OK);
+
+            var issueFields = new Dictionary<string, object>();
+            issueFields.Add("summary", "[CrashReport] " + ToJiraSummary); // Call Stack, Line 1
+            issueFields.Add("description", string.Join("\r\n", ToJiraDescriptions)); // Description
+
+            using (var responseReader = new StreamReader(response.GetResponseStream()))
+            {
+                var responseText = responseReader.ReadToEnd();
+
+                JObject jsonObject = JObject.Parse(responseText);
+
+                bool fields = jsonObject["projects"][0]["issuetypes"].Any();
+
+                issueFields.Add("project", new Dictionary<string, object> { { "id", jsonObject["projects"][0]["id"].ToObject<int>() } });
+
+                if (!fields) 
+                    return issueFields;
+
+                foreach (
+                    JToken issuetype in
+                        jsonObject["projects"][0]["issuetypes"])//.[0]["fields"])
+                {
+                    //Only fill required fields
+                    if (issuetype["subtask"].ToObject<bool>() == true)
+                        continue;
+
+                    JToken field = issuetype["fields"];
+
+                    if (field["required"].ToObject<bool>() == true)
+                    {
+                        //don't fill fields with a default value
+                        if (field["hasDefaultValue"].ToObject<bool>() == true)
+                            continue;
+
+                        JToken outValue;
+
+                        //if (field.Value["schema"]["type"].ToObject<string>() == "array")
+                        //{
+                        //    if (field.Value["allowedValues"] != null)
+                        //    {
+                        //        //don't add the same key twice
+                        //        if (issueFields.ContainsKey(field.Key))
+                        //            continue;
+
+                        //        issueFields.Add(field.Key, new object[]
+                        //        {
+                        //            new Dictionary<string, object>
+                        //            {
+                        //                {"id", field.Value["allowedValues"][0]["id"].ToObject<string>()}
+                        //            }
+                        //        });
+                        //    }
+                        //}
+                        //else
+                        //{
+                        //    if (obj.TryGetValue("allowedValues", out outValue))
+                        //    {
+                        //        //don't add the same key twice
+                        //        if (issueFields.ContainsKey(field.Key))
+                        //            continue;
+
+                        //        issueFields.Add(field.Key,
+                        //            new Dictionary<string, object>
+                        //            {
+                        //                {"id", field.Value["allowedValues"][0]["id"].ToObject<int>()}
+                        //            });
+                        //    }
+                        //}
+                    }
+                }
+            }
+
+            return issueFields;
+        }
 
         /// <summary> Returns concatenated string of fields from the specified JIRA list. </summary>
         public string GetFieldsFrom(List<object> jiraList, string fieldName)
@@ -410,22 +498,22 @@ namespace Tools.CrashReporter.CrashReportWebSite.DataModels
             Tooltip += "LatestOSAffected: " + LatestOSAffected + NL;
 
             // "name"
-            string JiraVersions = GetFieldsFrom(ToJiraVersions, "name");
+            var JiraVersions = GetFieldsFrom(ToJiraVersions, "name");
             Tooltip += "JiraVersions: " + JiraVersions + NL;
 
             // "value"
-            string JiraBranches = "";
+            var jiraBranches = "";
             foreach (var Branch in ToJiraBranches)
             {
-                JiraBranches += Branch + ", ";
+                jiraBranches += Branch + ", ";
             }
-            Tooltip += "JiraBranches: " + JiraBranches + NL;
+            Tooltip += "JiraBranches: " + jiraBranches + NL;
 
             // "value"
-            string JiraPlatforms = GetFieldsFrom(ToJiraPlatforms, "value");
+            var JiraPlatforms = GetFieldsFrom(ToJiraPlatforms, "value");
             Tooltip += "JiraPlatforms: " + JiraPlatforms + NL;
 
-            string JiraCallstack = "Callstack:" + NL + string.Join(NL, ToJiraFunctionCalls) + NL;
+            var JiraCallstack = "Callstack:" + NL + string.Join(NL, ToJiraFunctionCalls) + NL;
             Tooltip += JiraCallstack;
 
             return Tooltip;
diff --git a/Engine/Source/Programs/CrashReporter/CrashReportWebSite/DataModels/Crash.cs b/Engine/Source/Programs/CrashReporter/CrashReportWebSite/DataModels/Crash.cs
--- a/Engine/Source/Programs/CrashReporter/CrashReportWebSite/DataModels/Crash.cs
+++ b/Engine/Source/Programs/CrashReporter/CrashReportWebSite/DataModels/Crash.cs
@@ -62,6 +62,7 @@ namespace Tools.CrashReporter.CrashReportWebSite.DataModels
         public Nullable<int> ErrorMessageId { get; set; }
         public Nullable<int> BuggId { get; set; }
         public Nullable<bool> ProcessFailed { get; set; }
+        public Nullable<bool> IsVanilla { get; set; }
     
         public virtual CallStackPattern CallStackPattern { get; set; }
         public virtual User User { get; set; }
diff --git a/Engine/Source/Programs/CrashReporter/CrashReportWebSite/DataModels/Repositories/BuggRespository.cs b/Engine/Source/Programs/CrashReporter/CrashReportWebSite/DataModels/Repositories/BuggRespository.cs
--- a/Engine/Source/Programs/CrashReporter/CrashReportWebSite/DataModels/Repositories/BuggRespository.cs
+++ b/Engine/Source/Programs/CrashReporter/CrashReportWebSite/DataModels/Repositories/BuggRespository.cs
@@ -1,6 +1,4 @@
-﻿// Copyright 1998-2017 Epic Games, Inc. All Rights Reserved.
-
-using System;
+﻿using System;
 using System.Collections.Generic;
 using System.Data.Entity;
 using System.Linq;
@@ -13,7 +11,7 @@ namespace Tools.CrashReporter.CrashReportWebSite.DataModels.Repositories
     /// </summary>
     public class BuggRepository : IBuggRepository
     {
-        private readonly CrashReportEntities _entityContext;
+        private CrashReportEntities _entityContext;
 
         /// <summary>
         /// Constructor
diff --git a/Engine/Source/Programs/CrashReporter/CrashReportWebSite/DataModels/Repositories/CallStackRepository.cs b/Engine/Source/Programs/CrashReporter/CrashReportWebSite/DataModels/Repositories/CallStackRepository.cs
--- a/Engine/Source/Programs/CrashReporter/CrashReportWebSite/DataModels/Repositories/CallStackRepository.cs
+++ b/Engine/Source/Programs/CrashReporter/CrashReportWebSite/DataModels/Repositories/CallStackRepository.cs
@@ -1,6 +1,4 @@
-﻿// Copyright 1998-2017 Epic Games, Inc. All Rights Reserved.
-
-using System;
+﻿using System;
 using System.Collections.Generic;
 using System.Data.Entity;
 using System.Linq;
@@ -13,7 +11,7 @@ namespace Tools.CrashReporter.CrashReportWebSite.DataModels.Repositories
     /// </summary>
     public class CallStackRepository : ICallstackRepository
     {
-        private readonly CrashReportEntities _entityContext;
+        private CrashReportEntities _entityContext;
 
         /// <summary>
         /// Constructor
