Fix CIS warning due to unused variable in CrashReportProcess.

diff --git a/Engine/Source/Programs/CrashReporter/CrashReportProcess/ReportWatcher.cs b/Engine/Source/Programs/CrashReporter/CrashReportProcess/ReportWatcher.cs
--- a/Engine/Source/Programs/CrashReporter/CrashReportProcess/ReportWatcher.cs
+++ b/Engine/Source/Programs/CrashReporter/CrashReportProcess/ReportWatcher.cs
@@ -145,7 +145,7 @@ namespace Tools.CrashReporter.CrashReportProcess
 					{
 						await Task.Delay(30000, Cancel);
 					}
-					catch (TaskCanceledException Ex)
+					catch (TaskCanceledException)
 					{
 						CrashReporterProcessServicer.WriteEvent("WatcherTask exiting delay because it was cancelled.");
 						break;
