Copying //UE4/Dev-Core to //UE4/Dev-Main (Source: //UE4/Dev-Core @ 3345728)

diff --git a/Engine/Source/Programs/CrashReporter/CrashReportWebSite/Controllers/BuggsController.cs b/Engine/Source/Programs/CrashReporter/CrashReportWebSite/Controllers/BuggsController.cs
--- a/Engine/Source/Programs/CrashReporter/CrashReportWebSite/Controllers/BuggsController.cs
+++ b/Engine/Source/Programs/CrashReporter/CrashReportWebSite/Controllers/BuggsController.cs
@@ -8,7 +8,6 @@ using System.Linq.Expressions;
 using System.Net;
 using System.Web;
 using System.Web.Mvc;
-using System.Web.UI;
 using Microsoft.Practices.ObjectBuilder2;
 using Newtonsoft.Json.Linq;
 using Tools.CrashReporter.CrashReportWebSite.DataModels;
@@ -38,7 +37,7 @@ namespace Tools.CrashReporter.CrashReportWebSite.Controllers
 		/// <returns>The view to display a list of Buggs on the client.</returns>
 		public ActionResult Index( FormCollection BuggsForm )
 		{
-			using (var logTimer = new FAutoScopedLogTimer( this.GetType().ToString(), bCreateNewLog: true ))
+			using (var logTimer = new FAutoScopedLogTimer( this.GetType().ToString()))
 			{
 				var formData = new FormHelper( Request, BuggsForm, "CrashInTimeFrameGroup" );
 			    var results = GetResults( formData );
@@ -53,215 +52,219 @@ namespace Tools.CrashReporter.CrashReportWebSite.Controllers
 		/// <param name="buggsForm">The form of user data passed up from the client.</param>
 		/// <param name="id">The unique id of the Bugg.</param>
 		/// <returns>The view to display a Bugg on the client.</returns>
-		public ActionResult Show( FormCollection buggsForm, int id )
+		public ActionResult Show( FormCollection buggsForm, int id, int page = 0 )
 		{
-			using( var logTimer = new FAutoScopedLogTimer( this.GetType().ToString() + "(BuggId=" + id + ")", bCreateNewLog: true ) )
+			using( var logTimer = new FAutoScopedLogTimer( this.GetType().ToString() + "(BuggId=" + id + ")") )
 			{
-				// Set the display properties based on the radio buttons
-			    var displayModuleNames = buggsForm["DisplayModuleNames"] == "true";
-			    var displayFunctionNames = buggsForm["DisplayFunctionNames"] == "true";
-			    var displayFileNames = buggsForm["DisplayFileNames"] == "true";
-			    var displayFilePathNames = false;
-
-				if( buggsForm["DisplayFilePathNames"] == "true" )
-				{
-					displayFilePathNames = true;
-					displayFileNames = false;
-				}
-
-			    var displayUnformattedCallStack = buggsForm["DisplayUnformattedCallStack"] == "true";
-
-			    var model = GetResult(id);
-                model.Bugg.PrepareBuggForJira(model.Crashes);
-			    // Handle 'CopyToJira' button
-                var buggIdToBeAddedToJira = 0;
-                foreach (var entry in buggsForm.Cast<object>().Where(entry => entry.ToString().Contains("CopyToJira-")))
-                {
-                    int.TryParse(entry.ToString().Substring("CopyToJira-".Length), out buggIdToBeAddedToJira);
-                    break;
-                }
-                if (buggIdToBeAddedToJira != 0)
-                {
-                    model.Bugg.JiraProject = buggsForm["JiraProject"];
-                    model.Bugg.CopyToJira();
-                }
-
-                var jc = JiraConnection.Get();
-				var bValidJira = false;
-
-				// Verify valid JiraID, this may be still a TTP
-				if( !string.IsNullOrEmpty( model.Bugg.TTPID ) )
-				{
-					var jira = 0;
-                    int.TryParse(model.Bugg.TTPID, out jira);
-
-                    if (jira == 0)
-					{
-						bValidJira = true;
-					}
-				}
-
-				if( jc.CanBeUsed() && bValidJira )
-				{
-					// Grab the data form JIRA.
-					var jiraSearchQuery = "key = " + model.Bugg.TTPID;
-					var jiraResults = new Dictionary<string, Dictionary<string, object>>();
-
-					try
-					{
-						jiraResults = jc.SearchJiraTickets(
-						jiraSearchQuery,
-						new string[] 
-						{ 
-							"key",				// string
-							"summary",			// string
-							"components",		// System.Collections.ArrayList, Dictionary<string,object>, name
-							"resolution",		// System.Collections.Generic.Dictionary`2[System.String,System.Object], name
-							"fixVersions",		// System.Collections.ArrayList, Dictionary<string,object>, name
-							"customfield_11200" // string
-						} );
-					}
-					catch (System.Exception)
-					{
-                        model.Bugg.JiraSummary = "JIRA MISMATCH";
-                        model.Bugg.JiraComponentsText = "JIRA MISMATCH";
-                        model.Bugg.JiraResolution = "JIRA MISMATCH";
-                        model.Bugg.JiraFixVersionsText = "JIRA MISMATCH";
-                        model.Bugg.JiraFixCL = "JIRA MISMATCH";
-					}
-
-					// Jira Key, Summary, Components, Resolution, Fix version, Fix changelist
-					if(jiraResults.Any())
-					{
-					    var jira = jiraResults.First();
-						var summary = (string)jira.Value["summary"];
-
-						var componentsText = "";
-                        var components = (System.Collections.ArrayList)jira.Value["components"];
-                        foreach (Dictionary<string, object> component in components)
-						{
-							componentsText += (string)component["name"];
-							componentsText += " ";
-						}
-
-						var resolutionFields = (Dictionary<string, object>)jira.Value["resolution"];
-						var resolution = resolutionFields != null ? (string)resolutionFields["name"] : "";
-
-						var fixVersionsText = "";
-						var fixVersions = (System.Collections.ArrayList)jira.Value["fixVersions"];
-						foreach( Dictionary<string, object> fixVersion in fixVersions )
-						{
-							fixVersionsText += (string)fixVersion["name"];
-							fixVersionsText += " ";
-						}
-
-						var fixCl = jira.Value["customfield_11200"] != null ? (int)(decimal)jira.Value["customfield_11200"] : 0;
-
-                        //Conversion to ado.net entity framework
-
-                        model.Bugg.JiraSummary = summary;
-                        model.Bugg.JiraComponentsText = componentsText;
-                        model.Bugg.JiraResolution = resolution;
-                        model.Bugg.JiraFixVersionsText = fixVersionsText;
-						if( fixCl != 0 )
-						{
-							model.Bugg.FixedChangeList = fixCl.ToString();
-						}
-					}
-				}
-
-				// Apply any user settings
-				if( buggsForm.Count > 0 )
-				{
-					if( !string.IsNullOrEmpty( buggsForm["SetStatus"] ) )
-					{
-						model.Bugg.Status = buggsForm["SetStatus"];
-                        model.Bugg.Crashes.ForEach(data =>data.Status = buggsForm["SetStatus"]);
-					}
-
-					if( !string.IsNullOrEmpty( buggsForm["SetFixedIn"] ) )
-					{
-						model.Bugg.FixedChangeList = buggsForm["SetFixedIn"];
-                        model.Bugg.Crashes.ForEach(data => data.FixedChangeList = buggsForm["SetFixedIn"]);
-					}
-
-					if( !string.IsNullOrEmpty( buggsForm["SetTTP"] ) )
-					{
-						model.Bugg.TTPID = buggsForm["SetTTP"];
-                        model.Bugg.Crashes.ForEach(data => data.Jira = buggsForm["SetTTP"]);
-					}
-				    using (var unitOfWork = new UnitOfWork(new CrashReportEntities()))
-				    {
-				        unitOfWork.BuggRepository.Update(model.Bugg);
-                        foreach (var crash in model.Bugg.Crashes)
-				        {
-				            unitOfWork.CrashRepository.Update(crash);
-				        }
-				        unitOfWork.Save();
-				    }
-				}
-
-				// Set up the view model with the Crash data
-			    int page;
-                int.TryParse(buggsForm["Page"], out page);
-
-				model.Crashes = model.Crashes.Skip(page * pageSize).Take(pageSize).ToList();
-			    
-				var newCrash = model.Crashes.FirstOrDefault();
-				if( newCrash != null )
-				{
-                    var callStack = new CallStackContainer( newCrash );
-
-					// Set callstack properties
-					callStack.bDisplayModuleNames = displayModuleNames;
-					callStack.bDisplayFunctionNames = displayFunctionNames;
-					callStack.bDisplayFileNames = displayFileNames;
-					callStack.bDisplayFilePathNames = displayFilePathNames;
-					callStack.bDisplayUnformattedCallStack = displayUnformattedCallStack;
-
-					model.CallStack = callStack;
-
-					// Shorten very long function names.
-					foreach( var entry in model.CallStack.CallStackEntries )
-					{
-						entry.FunctionName = entry.GetTrimmedFunctionName( 128 );
-					}
-
-					model.SourceContext = newCrash.SourceContext;
-
-					model.LatestCrashSummary = newCrash.Summary;
-				}
-
-                model.LatestCrashSummary = newCrash.Summary;
-                model.Bugg.LatestCrashSummary = newCrash.Summary;
-				model.GenerationTime = logTimer.GetElapsedSeconds().ToString( "F2" );
-
-                //Populate Jira Projects
-                var jiraConnection = JiraConnection.Get();
-                var response = jiraConnection.JiraRequest("/issue/createmeta", JiraConnection.JiraMethod.GET, null, HttpStatusCode.OK);
-
-                using (var responseReader = new StreamReader(response.GetResponseStream()))
-                {
-                    var responseText = responseReader.ReadToEnd();
-
-                    JObject jsonObject = JObject.Parse(responseText);
-
-                    JToken fields = jsonObject["projects"];
-
-                    foreach (var project in fields)
-                    {
-                        model.JiraProjects.Add(new SelectListItem() { Text = project["name"].ToObject<string>(), Value = project["key"].ToObject<string>() });
-                    }
-                }
-
-			    model.PagingInfo = new PagingInfo
+			    using (var unitOfWork = new UnitOfWork(new CrashReportEntities()))
 			    {
-			        CurrentPage = page,
-			        PageSize = pageSize,
-			        TotalResults = model.Bugg.NumberOfCrashes.Value
-			    };
 
-				return View( "Show", model );
+			        // Set the display properties based on the radio buttons
+			        var displayModuleNames = buggsForm["DisplayModuleNames"] == "true";
+			        var displayFunctionNames = buggsForm["DisplayFunctionNames"] == "true";
+			        var displayFileNames = buggsForm["DisplayFileNames"] == "true";
+			        var displayFilePathNames = false;
+
+
+			        // Set up the view model with the Crash data
+                    if(buggsForm["Page"] != null)
+                        int.TryParse(buggsForm["Page"], out page);
+
+			        if (buggsForm["DisplayFilePathNames"] == "true")
+			        {
+			            displayFilePathNames = true;
+			            displayFileNames = false;
+			        }
+
+			        var displayUnformattedCallStack = buggsForm["DisplayUnformattedCallStack"] == "true";
+
+			        var model = GetResult(id, page, pageSize, unitOfWork);
+			        model.SourceContext = model.CrashData.First().SourceContext;
+			        model.Bugg.PrepareBuggForJira(model.CrashData);
+			        // Handle 'CopyToJira' button
+			        var buggIdToBeAddedToJira = 0;
+			        foreach (var entry in buggsForm.Cast<object>().Where(entry => entry.ToString().Contains("CopyToJira-")))
+			        {
+			            int.TryParse(entry.ToString().Substring("CopyToJira-".Length), out buggIdToBeAddedToJira);
+			            break;
+			        }
+			        if (buggIdToBeAddedToJira != 0)
+			        {
+			            model.Bugg.JiraProject = buggsForm["JiraProject"];
+			            model.Bugg.CopyToJira();
+			        }
+
+			        var jc = JiraConnection.Get();
+			        var bValidJira = false;
+
+			        // Verify valid JiraID, this may be still a TTP
+			        if (!string.IsNullOrEmpty(model.Bugg.TTPID))
+			        {
+			            var jira = 0;
+			            int.TryParse(model.Bugg.TTPID, out jira);
+
+			            if (jira == 0)
+			            {
+			                bValidJira = true;
+			            }
+			        }
+
+			        if (jc.CanBeUsed() && bValidJira)
+			        {
+			            // Grab the data form JIRA.
+			            var jiraSearchQuery = "key = " + model.Bugg.TTPID;
+			            var jiraResults = new Dictionary<string, Dictionary<string, object>>();
+
+			            try
+			            {
+			                jiraResults = jc.SearchJiraTickets(
+			                    jiraSearchQuery,
+			                    new string[]
+			                    {
+			                        "key", // string
+			                        "summary", // string
+			                        "components", // System.Collections.ArrayList, Dictionary<string,object>, name
+			                        "resolution", // System.Collections.Generic.Dictionary`2[System.String,System.Object], name
+			                        "fixVersions", // System.Collections.ArrayList, Dictionary<string,object>, name
+			                        "customfield_11200" // string
+			                    });
+			            }
+			            catch (System.Exception)
+			            {
+			                model.Bugg.JiraSummary = "JIRA MISMATCH";
+			                model.Bugg.JiraComponentsText = "JIRA MISMATCH";
+			                model.Bugg.JiraResolution = "JIRA MISMATCH";
+			                model.Bugg.JiraFixVersionsText = "JIRA MISMATCH";
+			                model.Bugg.JiraFixCL = "JIRA MISMATCH";
+			            }
+
+			            // Jira Key, Summary, Components, Resolution, Fix version, Fix changelist
+			            if (jiraResults.Any())
+			            {
+			                var jira = jiraResults.First();
+			                var summary = (string) jira.Value["summary"];
+			                var componentsText = "";
+			                var components = (System.Collections.ArrayList) jira.Value["components"];
+			                foreach (Dictionary<string, object> component in components)
+			                {
+			                    componentsText += (string) component["name"];
+			                    componentsText += " ";
+			                }
+
+			                var resolutionFields = (Dictionary<string, object>) jira.Value["resolution"];
+			                var resolution = resolutionFields != null ? (string) resolutionFields["name"] : "";
+
+			                var fixVersionsText = "";
+			                var fixVersions = (System.Collections.ArrayList) jira.Value["fixVersions"];
+			                foreach (Dictionary<string, object> fixVersion in fixVersions)
+			                {
+			                    fixVersionsText += (string) fixVersion["name"];
+			                    fixVersionsText += " ";
+			                }
+
+			                var fixCl = jira.Value["customfield_11200"] != null
+			                    ? (int) (decimal) jira.Value["customfield_11200"]
+			                    : 0;
+
+			                //Conversion to ado.net entity framework
+
+			                model.Bugg.JiraSummary = summary;
+			                model.Bugg.JiraComponentsText = componentsText;
+			                model.Bugg.JiraResolution = resolution;
+			                model.Bugg.JiraFixVersionsText = fixVersionsText;
+			                if (fixCl != 0)
+			                {
+			                    model.Bugg.FixedChangeList = fixCl.ToString();
+			                }
+			            }
+			        }
+
+			        // Apply any user settings
+			        if (buggsForm.Count > 0)
+			        {
+			            if (!string.IsNullOrEmpty(buggsForm["SetStatus"]))
+			            {
+			                model.Bugg.Status = buggsForm["SetStatus"];
+                            unitOfWork.CrashRepository.SetStatusByBuggId(model.Bugg.Id, buggsForm["SetFixedIn"]);
+			            }
+
+			            if (!string.IsNullOrEmpty(buggsForm["SetFixedIn"]))
+			            {
+			                model.Bugg.FixedChangeList = buggsForm["SetFixedIn"];
+                            unitOfWork.CrashRepository.SetFixedCLByBuggId(model.Bugg.Id, buggsForm["SetFixedIn"]);
+			            }
+
+			            if (!string.IsNullOrEmpty(buggsForm["SetTTP"]))
+			            {
+			                model.Bugg.TTPID = buggsForm["SetTTP"];
+                            unitOfWork.CrashRepository.SetJiraByBuggId(model.Bugg.Id, buggsForm["SetTTP"]);
+			            }
+                        unitOfWork.BuggRepository.Update(model.Bugg);
+			            
+			            unitOfWork.Save();
+			        }
+
+			        var newCrash = model.CrashData.FirstOrDefault();
+			        if (newCrash != null)
+			        {
+			            var callStack = new CallStackContainer(newCrash.CrashType, newCrash.RawCallStack, newCrash.PlatformName);
+
+			            // Set callstack properties
+			            callStack.bDisplayModuleNames = displayModuleNames;
+			            callStack.bDisplayFunctionNames = displayFunctionNames;
+			            callStack.bDisplayFileNames = displayFileNames;
+			            callStack.bDisplayFilePathNames = displayFilePathNames;
+			            callStack.bDisplayUnformattedCallStack = displayUnformattedCallStack;
+
+			            model.CallStack = callStack;
+
+			            // Shorten very long function names.
+			            foreach (var entry in model.CallStack.CallStackEntries)
+			            {
+			                entry.FunctionName = entry.GetTrimmedFunctionName(128);
+			            }
+
+			            model.SourceContext = newCrash.SourceContext;
+
+			            model.LatestCrashSummary = newCrash.Summary;
+			        }
+
+			        model.LatestCrashSummary = newCrash.Summary;
+			        model.Bugg.LatestCrashSummary = newCrash.Summary;
+			        model.GenerationTime = logTimer.GetElapsedSeconds().ToString("F2");
+
+			        //Populate Jira Projects
+			        var jiraConnection = JiraConnection.Get();
+			        var response = jiraConnection.JiraRequest("/issue/createmeta", JiraConnection.JiraMethod.GET, null,
+			            HttpStatusCode.OK);
+
+			        using (var responseReader = new StreamReader(response.GetResponseStream()))
+			        {
+			            var responseText = responseReader.ReadToEnd();
+
+			            JObject jsonObject = JObject.Parse(responseText);
+
+			            JToken fields = jsonObject["projects"];
+
+			            foreach (var project in fields)
+			            {
+			                model.JiraProjects.Add(new SelectListItem()
+			                {
+			                    Text = project["name"].ToObject<string>(),
+			                    Value = project["key"].ToObject<string>()
+			                });
+			            }
+			        }
+
+			        model.PagingInfo = new PagingInfo
+			        {
+			            CurrentPage = page,
+			            PageSize = pageSize,
+			            TotalResults = model.Bugg.NumberOfCrashes.Value
+			        };
+
+			        return View("Show", model);
+			    }
 			}
 		}
         
@@ -282,12 +285,13 @@ namespace Tools.CrashReporter.CrashReportWebSite.Controllers
             var fromDate = FormData.DateFrom;
             var toDate = FormData.DateTo.AddDays(1);
 
-            IQueryable<Bugg> results = null;
+            IQueryable<Crash> results = null;
+            IQueryable<Bugg> buggResults = null;
             var skip = (FormData.Page - 1) * FormData.PageSize;
             var take = FormData.PageSize;
 
-            var sortedResultsList = new List<Bugg>();
-            var groupCounts = new SortedDictionary<string, int>();
+            var sortedResultsList = new System.Collections.Generic.List<Bugg>();
+            var groupCounts = new System.Collections.Generic.SortedDictionary<string, int>();
             int totalCountedRecords = 0;
 
             using (var unitOfWork = new UnitOfWork(new CrashReportEntities()))
@@ -295,7 +299,7 @@ namespace Tools.CrashReporter.CrashReportWebSite.Controllers
                 var userGroupId = unitOfWork.UserGroupRepository.First(data => data.Name == FormData.UserGroup).Id;
 
                 // Get every Bugg.
-                var resultsAll = unitOfWork.BuggRepository.ListAll();
+                var resultsAll = unitOfWork.CrashRepository.ListAll();
 
                 // Look at all Buggs that are still 'open' i.e. the last Crash occurred in our date range.
                 results = FilterByDate(resultsAll, FormData.DateFrom, FormData.DateTo);
@@ -303,23 +307,21 @@ namespace Tools.CrashReporter.CrashReportWebSite.Controllers
                 // Filter results by build version.
                 if (!string.IsNullOrEmpty(FormData.VersionName))
                 {
-                    results = FilterByBuildVersion(results, FormData.VersionName);
+                    results = results.Where(data => data.BuildVersion == FormData.VersionName);
                 }
 
                 // Filter by BranchName
                 if (!string.IsNullOrEmpty(FormData.BranchName))
                 {
                     results =
-                        results.Where(
-                            data => data.Crashes.Any(da => da.Branch.Equals(FormData.BranchName)));
+                        results.Where(da => da.Branch.Equals(FormData.BranchName));
                 }
 
                 // Filter by PlatformName
                 if (!string.IsNullOrEmpty(FormData.PlatformName))
                 {
                     results =
-                        results.Where(
-                            data => data.Crashes.Any(da => da.PlatformName.Equals(FormData.PlatformName)));
+                        results.Where(da => da.PlatformName.Equals(FormData.PlatformName));
                 }
 
                 // Run at the end
@@ -373,22 +375,19 @@ namespace Tools.CrashReporter.CrashReportWebSite.Controllers
 
                 // Filter by user group if present
                 if (!string.IsNullOrEmpty(FormData.UserGroup))
-                    results = FilterByUserGroup(results, FormData.UserGroup);
+                    results = FilterByUserGroup(results, userGroupId);
 
                 if (!string.IsNullOrEmpty(FormData.JiraId))
                 {
-                    results = FilterByJira(results, FormData.JiraId);
+                    results = results.Where(data => data.Jira == FormData.JiraId);
                 }
 
-                // Grab just the results we want to display on this page
-                totalCountedRecords = results.Count();
-
                 var CrashByUserGroupCounts = unitOfWork.CrashRepository.ListAll().Where(data =>
                     data.User.UserGroupId == userGroupId &&
                     data.TimeOfCrash >= fromDate &&
                     data.TimeOfCrash <= toDate &&
-                    data.PatternId != null)
-                    .GroupBy(data => data.PatternId)
+                    data.BuggId != null)
+                    .GroupBy(data => data.BuggId)
                     .Select(data => new {data.Key, Count = data.Count()})
                     .OrderByDescending(data => data.Count)
                     .ToDictionary(data => data.Key, data => data.Count);
@@ -396,8 +395,8 @@ namespace Tools.CrashReporter.CrashReportWebSite.Controllers
                 var CrashInTimeFrameCounts = unitOfWork.CrashRepository.ListAll().Where(data =>
                     data.TimeOfCrash >= fromDate &&
                     data.TimeOfCrash <= toDate
-                    && data.PatternId != null)
-                    .GroupBy(data => data.PatternId)
+                    && data.BuggId != null)
+                    .GroupBy(data => data.BuggId)
                     .Select(data => new {data.Key, Count = data.Count()})
                     .OrderByDescending(data => data.Count)
                     .ToDictionary(data => data.Key, data => data.Count);
@@ -405,24 +404,34 @@ namespace Tools.CrashReporter.CrashReportWebSite.Controllers
                 var affectedUsers = unitOfWork.CrashRepository.ListAll().Where(data =>
                     data.User.UserGroupId == userGroupId &&
                     data.TimeOfCrash >= fromDate &&
-                    data.TimeOfCrash <= toDate &&
-                    data.PatternId != null)
-                    .Select(data => new {PatternId = data.PatternId, ComputerName = data.ComputerName})
+                    data.TimeOfCrash <= toDate && data.BuggId != null)
+                    .Select(data => new { BuggId = data.BuggId, ComputerName = data.ComputerName })
                     .Distinct()
-                    .GroupBy(data => data.PatternId)
+                    .GroupBy(data => data.BuggId)
                     .Select(data => new {data.Key, Count = data.Count()})
                     .OrderByDescending(data => data.Count)
                     .ToDictionary(data => data.Key, data => data.Count);
 
-                results = GetSortedResults(results, FormData.SortTerm, FormData.SortTerm == "descending",
+                var ids = results.Where(data => data.BuggId != null).Select(data => data.BuggId).Distinct().ToList();
+
+                buggResults = unitOfWork.BuggRepository.ListAll().Where(data => ids.Contains(data.Id));
+
+                //buggResults = results.Where(data => data.BuggId != null).Select(data => data.Bugg).Distinct();
+                
+                // Grab just the results we want to display on this page
+                totalCountedRecords = buggResults.Count();
+
+                buggResults = GetSortedResults(buggResults, FormData.SortTerm, FormData.SortTerm == "descending",
                     FormData.DateFrom, FormData.DateTo, FormData.UserGroup);
 
+               
+
                 sortedResultsList =
-                    results.ToList();
+                    buggResults.ToList();
 
                 // Get UserGroup ResultCounts
                 var groups =
-                    results.SelectMany(data => data.UserGroups)
+                    buggResults.SelectMany(data => data.UserGroups)
                         .GroupBy(data => data.Name)
                         .Select(data => new {Key = data.Key, Count = data.Count()})
                         .ToDictionary(x => x.Key, y => y.Count);
@@ -431,23 +440,23 @@ namespace Tools.CrashReporter.CrashReportWebSite.Controllers
 
                 foreach (var bugg in sortedResultsList)
                 {
-                    if (bugg.PatternId.HasValue && CrashByUserGroupCounts.ContainsKey(bugg.PatternId.Value))
+                    if (CrashByUserGroupCounts.ContainsKey(bugg.Id))
                     {
-                        bugg.CrashesInTimeFrameGroup = CrashByUserGroupCounts[bugg.PatternId.Value];
+                        bugg.CrashesInTimeFrameGroup = CrashByUserGroupCounts[bugg.Id];
                     }
                     else
                         bugg.CrashesInTimeFrameGroup = 0;
 
-                    if (bugg.PatternId.HasValue && CrashInTimeFrameCounts.ContainsKey(bugg.PatternId.Value))
+                    if (CrashInTimeFrameCounts.ContainsKey(bugg.Id))
                     {
-                        bugg.CrashesInTimeFrameAll = CrashInTimeFrameCounts[bugg.PatternId.Value];
+                        bugg.CrashesInTimeFrameAll = CrashInTimeFrameCounts[bugg.Id];
                     }
                     else
                         bugg.CrashesInTimeFrameAll = 0;
 
-                    if (bugg.PatternId.HasValue && affectedUsers.ContainsKey(bugg.PatternId.Value))
+                    if (affectedUsers.ContainsKey(bugg.Id))
                     {
-                        bugg.NumberOfUniqueMachines = affectedUsers[bugg.PatternId.Value];
+                        bugg.NumberOfUniqueMachines = affectedUsers[bugg.Id];
                     }
                     else
                         bugg.NumberOfUniqueMachines = 0;
@@ -464,9 +473,9 @@ namespace Tools.CrashReporter.CrashReportWebSite.Controllers
                 }
             }
 
-            List<SelectListItem> branchNames;
-            List<SelectListItem> versionNames;
-            List<SelectListItem> platformNames;
+            System.Collections.Generic.List<SelectListItem> branchNames;
+            System.Collections.Generic.List<SelectListItem> versionNames;
+            System.Collections.Generic.List<SelectListItem> platformNames;
             using (var unitOfWork = new UnitOfWork(new CrashReportEntities()))
             {
                 branchNames = unitOfWork.CrashRepository.GetBranchesAsListItems();
@@ -502,31 +511,37 @@ namespace Tools.CrashReporter.CrashReportWebSite.Controllers
         /// <param name="buggsForm"></param>
         /// <param name="id"></param>
         /// <returns></returns>
-        private BuggViewModel GetResult(int id)
+        private BuggViewModel GetResult(int id, int page, int pageSize, UnitOfWork unitOfWork)
 	    {
             var model = new BuggViewModel();
-            using (var unitOfWork = new UnitOfWork(new CrashReportEntities()))
+            model.Bugg = unitOfWork.BuggRepository.GetById(id);
+
+            model.Bugg.CrashesInTimeFrameAll = 0;
+            model.Bugg.CrashesInTimeFrameGroup = 0;
+            model.Bugg.NumberOfCrashes = unitOfWork.CrashRepository.Count(data => data.BuggId == id);
+            model.Bugg.NumberOfUniqueMachines =
+                unitOfWork.CrashRepository.ListAll()
+                    .Where(data => data.BuggId == id)
+                    .GroupBy(data => data.ComputerName)
+                    .Count();
+            model.CrashData = unitOfWork.CrashRepository.ListAll().Where(data => data.BuggId == id)
+                .OrderByDescending(data => data.TimeOfCrash).Skip(page * pageSize).Take(pageSize).Select(data => new CrashDataModel()
             {
-                // Create a new view and populate with Crash
-                List<Crash> crashes = null;
-                model.Bugg = unitOfWork.BuggRepository.GetById(id);
-
-                crashes = model.Bugg.Crashes.OrderByDescending(data => data.TimeOfCrash).ToList();
-                model.Bugg.CrashesInTimeFrameAll = crashes.Count;
-                model.Bugg.CrashesInTimeFrameGroup = crashes.Count;
-                model.Bugg.NumberOfCrashes = crashes.Count;
-                model.Crashes = crashes;
-                model.CrashData = crashes.Take(50).Select(data => new CrashDataModel()
-                {
-                    Id = data.Id,
-                    ChangelistVersion = data.ChangelistVersion,
-                    GameName = data.GameName,
-                    EngineMode = data.EngineMode,
-                    PlatformName = data.PlatformName,
-                    TimeOfCrash = data.TimeOfCrash,
-                    Description = data.Description,
-                }).ToList();
-            }
+                Id = data.Id,
+                ChangelistVersion = data.ChangelistVersion,
+                GameName = data.GameName,
+                EngineMode = data.EngineMode,
+                PlatformName = data.PlatformName,
+                TimeOfCrash = data.TimeOfCrash,
+                Description = data.Description,
+                RawCallStack = data.RawCallStack,
+                CrashType = data.CrashType,
+                Summary = data.Summary,
+                SourceContext = data.SourceContext,
+                BuildVersion = data.BuildVersion,
+                ComputerName = data.ComputerName,
+                Branch = data.Branch
+            }).ToList();
 
             return model;
 	    }
@@ -549,12 +564,13 @@ namespace Tools.CrashReporter.CrashReportWebSite.Controllers
         /// <param name="dateFrom">The earliest date to filter by.</param>
         /// <param name="dateTo">The latest date to filter by.</param>
         /// <returns>The set of Buggs between the earliest and latest date.</returns>
-        private IQueryable<Bugg> FilterByDate(IQueryable<Bugg> results, DateTime dateFrom, DateTime dateTo)
+        private IQueryable<Crash> FilterByDate(IQueryable<Crash> results, DateTime dateFrom, DateTime dateTo)
         {
             dateTo = dateTo.AddDays(1);
-            var buggsInTimeFrame =
-                results.Where(bugg => (bugg.TimeOfFirstCrash >= dateFrom && bugg.TimeOfFirstCrash <= dateTo) ||
-                    (bugg.TimeOfLastCrash >= dateFrom && bugg.TimeOfLastCrash <= dateTo) || (bugg.TimeOfFirstCrash <= dateFrom && bugg.TimeOfLastCrash >= dateTo));
+            //var buggsInTimeFrame =
+            //    results.Where(bugg => (bugg.TimeOfFirstCrash >= dateFrom && bugg.TimeOfFirstCrash <= dateTo) ||
+            //        (bugg.TimeOfLastCrash >= dateFrom && bugg.TimeOfLastCrash <= dateTo) || (bugg.TimeOfFirstCrash <= dateFrom && bugg.TimeOfLastCrash >= dateTo));
+            var buggsInTimeFrame = results.Where(MyCrash => MyCrash.TimeOfCrash >= dateFrom && MyCrash.TimeOfCrash <= dateTo);
 
             return buggsInTimeFrame;
         }
@@ -580,9 +596,9 @@ namespace Tools.CrashReporter.CrashReportWebSite.Controllers
         /// <param name="setOfBuggs">The unfiltered set of Buggs.</param>
         /// <param name="groupName">The user group name to filter by.</param>
         /// <returns>The set of Buggs by users in the requested user group.</returns>
-        private IQueryable<Bugg> FilterByUserGroup(IQueryable<Bugg> setOfBuggs, string groupName)
+        private IQueryable<Crash> FilterByUserGroup(IQueryable<Crash> setOfBuggs, int groupName)
         {
-            return setOfBuggs.Where(data => data.UserGroups.Any(ug => ug.Name == groupName));
+            return setOfBuggs.Where(data => data.User.UserGroup.Id == groupName);
         }
 
         /// <summary>
diff --git a/Engine/Source/Programs/CrashReporter/CrashReportWebSite/Controllers/CrashesController.cs b/Engine/Source/Programs/CrashReporter/CrashReportWebSite/Controllers/CrashesController.cs
--- a/Engine/Source/Programs/CrashReporter/CrashReportWebSite/Controllers/CrashesController.cs
+++ b/Engine/Source/Programs/CrashReporter/CrashReportWebSite/Controllers/CrashesController.cs
@@ -11,6 +11,7 @@ using System.Text;
 using System.Text.RegularExpressions;
 using System.Web;
 using System.Web.Mvc;
+using Tools.CrashReporter.CrashReportWebSite.Classes;
 using Tools.CrashReporter.CrashReportWebSite.DataModels;
 using Tools.CrashReporter.CrashReportWebSite.DataModels.Repositories;
 using Tools.CrashReporter.CrashReportWebSite.Properties;
@@ -18,6 +19,7 @@ using Tools.CrashReporter.CrashReportWebSite.ViewModels;
 using Tools.DotNETCommon.XmlHandler;
 using Tools.CrashReporter.CrashReportCommon;
 using System.Data.SqlClient;
+using System.Runtime.Caching;
 
 namespace Tools.CrashReporter.CrashReportWebSite.Controllers
 {
@@ -26,9 +28,13 @@ namespace Tools.CrashReporter.CrashReportWebSite.Controllers
 	/// </summary>
 	[HandleError]
 	public class CrashesController : Controller
-	{
+    {
+        private static MemoryCache crashCache = MemoryCache.Default;
+
         //Ugly instantiation of Crash repository will replace with dependency injection BEFORE this gets anywhere near live.
 	    private readonly SlackWriter _slackWriter;
+
+	    private static readonly PerformanceTracker PerfTracker = new PerformanceTracker(100);
         
         /// <summary> Special user name, currently used to mark Crashes from UE4 releases. </summary>
         const string UserNameAnonymous = "Anonymous";
@@ -47,6 +53,41 @@ namespace Tools.CrashReporter.CrashReportWebSite.Controllers
 		    };
 		}
 
+        class CrashCacheItem
+        {
+            public List<Crash> Crashes { get; set; }
+            public Dictionary<string, int> GroupCount { get; set; }
+            public int ResultCount { get; set; }
+        }
+
+        private string getCacheKeyFromFormData(FormHelper formData)
+        {
+            return string.Format("{0}-{1}-{2}-{3}-{4}-{5}-{6}-{7}-{8}-{9}-{10}-{11}-{12}-{13}-{14}-{15}-{16}-{17}-{18}-{19}-{20}-{21}-{22}",
+                formData.BranchName,
+                formData.BuggId,
+                formData.BuiltFromCL,
+                formData.CrashType,
+                formData.DateFrom,
+                formData.DateTo,
+                formData.EngineMode,
+                formData.EngineVersion,
+                formData.EpicIdOrMachineQuery,
+                formData.GameName,
+                formData.JiraId,
+                formData.JiraQuery,
+                formData.MessageQuery,
+                formData.Page,
+                formData.PageSize,
+                formData.PlatformName,
+                formData.PreviousOrder,
+                formData.SearchQuery,
+                formData.SortOrder,
+                formData.SortTerm,
+                formData.UserGroup,
+                formData.UsernameQuery,
+                formData.VersionName);
+        }
+
 		/// <summary>
 		/// Display a summary list of Crashes based on the search criteria.
 		/// </summary>
@@ -54,7 +95,7 @@ namespace Tools.CrashReporter.CrashReportWebSite.Controllers
 		/// <returns>A view to display a list of Crash reports.</returns>
 		public ActionResult Index(FormCollection CrashesForm )
 		{
-            using( var logTimer = new FAutoScopedLogTimer( this.GetType().ToString(), bCreateNewLog: true ) )
+            using( var logTimer = new FAutoScopedLogTimer( this.GetType().ToString()) )
 			{
 				// Handle any edits made in the Set form fields
                 //foreach( var Entry in CrashesForm )
@@ -171,7 +212,7 @@ namespace Tools.CrashReporter.CrashReportWebSite.Controllers
 			        unitOfWork.Save();
 			    }
 
-			    currentCallStack = new CallStackContainer( currentCrash );
+			    currentCallStack = new CallStackContainer( currentCrash.CrashType, currentCrash.RawCallStack, currentCrash.PlatformName );
 			    currentCrash.Module = currentCallStack.GetModuleName();
 
 				//Set call stack properties
@@ -181,7 +222,7 @@ namespace Tools.CrashReporter.CrashReportWebSite.Controllers
 				currentCallStack.bDisplayFilePathNames = true;
 				currentCallStack.bDisplayUnformattedCallStack = false;
 
-				currentCrash.CallStackContainer = new CallStackContainer(currentCrash);
+			    currentCrash.CallStackContainer = currentCallStack;
 
 				var Model = new CrashViewModel { Crash = currentCrash, User = CrashUser, Bugg = currentBugg, UserGroup = currentUserGroup, CallStack = currentCallStack };
 				Model.GenerationTime = logTimer.GetElapsedSeconds().ToString( "F2" );
@@ -196,17 +237,13 @@ namespace Tools.CrashReporter.CrashReportWebSite.Controllers
 		/// <returns>The row id of the newly added Crash.</returns>
 		public ActionResult AddCrash(int id)
 		{
-		   
-		    var buildPatternDuration = "";
-		    var createCrashDuration = "";
-		    var userTime = "";
 		    var newCrashResult = new CrashReporterResult();
 		    CrashDescription newCrash;
 		    newCrashResult.ID = -1;
 		    string payloadString;
             
             using (var logTimer = new FAutoScopedLogTimer(this.GetType().ToString() 
-                + "(CrashId=" + id + ")", bCreateNewLog: true))
+                + "(CrashId=" + id + ")"))
             {
 		        //Read the request payload
 		        try
@@ -258,16 +295,9 @@ namespace Tools.CrashReporter.CrashReportWebSite.Controllers
 		        //Add Crash to database
 		        try
 		        {
-		            using (
-		                var logTimer2 = new FAutoScopedLogTimer(this.GetType().ToString() + "(CrashId=" + id + ")",
-		                    bCreateNewLog: true))
-		            {
-		                var crash = CreateCrash(newCrash, out buildPatternDuration, out userTime);
+                    var crash = CreateCrash(newCrash);
 		                newCrashResult.ID = crash.Id;
 		                newCrashResult.bSuccess = true;
-
-		                createCrashDuration = logTimer2.GetElapsedSeconds().ToString("F2");
-		            }
 		        }
 		        catch (DbEntityValidationException dbentEx)
 		        {
@@ -339,23 +369,9 @@ namespace Tools.CrashReporter.CrashReportWebSite.Controllers
 
 		        var returnResult = XmlHandler.ToXmlString<CrashReporterResult>(newCrashResult);
 
-                _slackWriter.Write("Add Crash completed. Pattern Match Time : " + buildPatternDuration +
-                    " Find User Time :" + userTime +
-                    " Create Crash Time : " + createCrashDuration + 
-                    " Total Time : " + logTimer.GetElapsedSeconds().ToString("F2") + " Return Result = " + newCrashResult.bSuccess.ToString());
-
-		        if (logTimer.GetElapsedSeconds() > 100.0f)
-		        {
-                    var slackWriter = new SlackWriter()
-                    {
-                        WebhookUrl = Settings.Default.SlackWebhookUrl,
-                        Channel = "CrashReporterdetails",
-                        Username = Settings.Default.SlackUsername,
-                        IconEmoji = Settings.Default.SlackEmoji
-                    };
-                    slackWriter.Write("Long Duration Add Crash. Crash Result Message : " + newCrashResult.Message);
-		        }
-
+                PerfTracker.AddStat("Add Crash Complete", logTimer.GetElapsedSeconds());
+                PerfTracker.IncrementCount();
+                
 		        return Content(returnResult, "text/xml");
 		    }
 		}
@@ -370,68 +386,85 @@ namespace Tools.CrashReporter.CrashReportWebSite.Controllers
             List<Crash> results = null;
             IQueryable<Crash> resultsQuery = null;
 
-            var skip = (formData.Page - 1) * formData.PageSize;
-            var take = formData.PageSize;
+            var cacheKey = getCacheKeyFromFormData(formData);
+            var cachedResults = crashCache.Get(cacheKey) as CrashCacheItem;
 
-            Dictionary<string, int> groupCounts;
-            var resultCount = 0;
-            using (var unitOfWork = new UnitOfWork(new CrashReportEntities()))
+            if (cachedResults == null)
             {
-                resultsQuery = ConstructQueryForFiltering(unitOfWork, formData);
-
-                // Filter by data and get as enumerable.
-                resultsQuery = FilterByDate(resultsQuery, formData.DateFrom, formData.DateTo);
+                var skip = (formData.Page - 1) * formData.PageSize;
+                var take = formData.PageSize;
 
-                // Filter by BuggId
-                if (!string.IsNullOrEmpty(formData.BuggId))
+                Dictionary<string, int> groupCounts;
+                var resultCount = 0;
+                using (var unitOfWork = new UnitOfWork(new CrashReportEntities()))
                 {
-                    var buggId = 0;
-                    var bValid = int.TryParse(formData.BuggId, out buggId);
+                    resultsQuery = ConstructQueryForFiltering(unitOfWork, formData);
+
+                    // Filter by data and get as enumerable.
+                    resultsQuery = FilterByDate(resultsQuery, formData.DateFrom, formData.DateTo);
 
-                    if (bValid)
+                    // Filter by BuggId
+                    if (!string.IsNullOrEmpty(formData.BuggId))
                     {
-                        var newBugg = unitOfWork.BuggRepository.GetById(buggId);
+                        var buggId = 0;
+                        var bValid = int.TryParse(formData.BuggId, out buggId);
 
-                        if (newBugg != null)
+                        if (bValid)
                         {
-                            resultsQuery = resultsQuery.Where(data => data.PatternId == newBugg.PatternId);
+                            var newBugg = unitOfWork.BuggRepository.GetById(buggId);
+
+                            if (newBugg != null)
+                            {
+                                resultsQuery = resultsQuery.Where(data => data.PatternId == newBugg.PatternId);
+                            }
                         }
                     }
-                }
 
-                var countsQuery =
-                    resultsQuery.GroupBy(data => data.User.UserGroup)
-                        .Select(data => new {Key = data.Key.Name, Count = data.Count()});
-                groupCounts = countsQuery.OrderBy(data => data.Key)
-                    .ToDictionary(data => data.Key, data => data.Count);
+                    var countsQuery =
+                        resultsQuery.GroupBy(data => data.User.UserGroup)
+                            .Select(data => new { Key = data.Key.Name, Count = data.Count() });
+                    groupCounts = countsQuery.OrderBy(data => data.Key)
+                        .ToDictionary(data => data.Key, data => data.Count);
 
-                // Filter by user group if present
-                var userGroupId = !string.IsNullOrEmpty(formData.UserGroup)
-                    ? unitOfWork.UserGroupRepository.First(data => data.Name == formData.UserGroup).Id
-                    : 1;
-                resultsQuery = resultsQuery.Where(data => data.User.UserGroupId == userGroupId);
+                    // Filter by user group if present
+                    var userGroupId = !string.IsNullOrEmpty(formData.UserGroup)
+                        ? unitOfWork.UserGroupRepository.First(data => data.Name == formData.UserGroup).Id
+                        : 1;
+                    resultsQuery = resultsQuery.Where(data => data.User.UserGroupId == userGroupId);
 
-                var orderedQuery = GetSortedQuery(resultsQuery, formData.SortTerm ?? "TimeOfCrash",
-                    formData.SortOrder == "Descending");
+                    var orderedQuery = GetSortedQuery(resultsQuery, formData.SortTerm ?? "TimeOfCrash",
+                        formData.SortOrder == "Descending");
 
-                // Grab just the results we want to display on this page
-                results = orderedQuery.Skip(skip).Take(take).ToList();
+                    // Grab just the results we want to display on this page
+                    results = orderedQuery.Skip(skip).Take(take).ToList();
 
-                // Get the Count for pagination
-                resultCount = orderedQuery.Count();
+                    // Get the Count for pagination
+                    resultCount = orderedQuery.Count();
 
-                // Process call stack for display
-                foreach (var CrashInstance in results)
-                {
-                    // Put call stacks into an list so we can access them line by line in the view
-                    CrashInstance.CallStackContainer = new CallStackContainer(CrashInstance);
+                    // Process call stack for display
+                    foreach (var CrashInstance in results)
+                    {
+                        // Put call stacks into an list so we can access them line by line in the view
+                        CrashInstance.CallStackContainer = new CallStackContainer(CrashInstance.CrashType, CrashInstance.RawCallStack, CrashInstance.PlatformName);
+                    }
+
+                    cachedResults = new CrashCacheItem()
+                    {
+                        Crashes = results,
+                        GroupCount = groupCounts,
+                        ResultCount = resultCount
+                    };
+
+                    var itemCachePolicy = new CacheItemPolicy() { AbsoluteExpiration = DateTimeOffset.Now.AddMinutes(Settings.Default.CrashesCacheInMinutes) };
+
+                    crashCache.Add(cacheKey, cachedResults, itemCachePolicy);
                 }
             }
 
             return new CrashesViewModel
             {
-                Results = results,
-                PagingInfo = new PagingInfo { CurrentPage = formData.Page, PageSize = formData.PageSize, TotalResults = resultCount },
+                Results = cachedResults.Crashes,
+                PagingInfo = new PagingInfo { CurrentPage = formData.Page, PageSize = formData.PageSize, TotalResults = cachedResults.ResultCount },
                 SortOrder = formData.SortOrder,
                 SortTerm = formData.SortTerm,
                 UserGroup = formData.UserGroup,
@@ -443,13 +476,13 @@ namespace Tools.CrashReporter.CrashReportWebSite.Controllers
                 BuiltFromCL = formData.BuiltFromCL,
                 BuggId = formData.BuggId,
                 JiraQuery = formData.JiraQuery,
-                DateFrom = (long)(formData.DateFrom - CrashesViewModel.Epoch).TotalMilliseconds,
-                DateTo = (long)(formData.DateTo - CrashesViewModel.Epoch).TotalMilliseconds,
+                DateFrom = (long)(formData.DateFrom.ToUniversalTime() - CrashesViewModel.Epoch).TotalMilliseconds,
+                DateTo = (long)(formData.DateTo.ToUniversalTime() - CrashesViewModel.Epoch).TotalMilliseconds,
                 BranchName = formData.BranchName,
                 VersionName = formData.VersionName,
                 PlatformName = formData.PlatformName,
                 GameName = formData.GameName,
-                GroupCounts = groupCounts
+                GroupCounts = cachedResults.GroupCount
             };
         }
 
@@ -729,56 +762,21 @@ namespace Tools.CrashReporter.CrashReportWebSite.Controllers
 	    /// </summary>
 	    /// <param name="description"></param>
 	    /// <returns></returns>
-        private Crash CreateCrash(CrashDescription description, out string buildPatternDuration, out string findUserTime)
+	    private Crash CreateCrash(CrashDescription description)
 	    {
-            var newCrash = new Crash
+	        var newCrash = new Crash
 	        {
 	            Branch = description.BranchName,
 	            BaseDir = description.BaseDir,
 	            BuildVersion = description.EngineVersion,
-                EngineVersion = description.BuildVersion,
+	            EngineVersion = description.BuildVersion,
 	            ChangelistVersion = description.BuiltFromCL.ToString(),
 	            CommandLine = description.CommandLine,
 	            EngineMode = description.EngineMode,
 	            ComputerName = description.MachineGuid,
-                IsVanilla = description.EngineModeEx.ToLower() == "vanilla"
+	            IsVanilla = description.EngineModeEx.ToLower() == "vanilla"
 	        };
 
-	        using (var unitOfWork = new UnitOfWork(new CrashReportEntities()))
-	        {
-	            using (
-	                var logTimer = new FAutoScopedLogTimer(this.GetType().ToString(),
-	                    bCreateNewLog: true))
-	            {
-	                //if there's a valid username assign the associated UserNameId else use "anonymous".
-	                var userName = (!string.IsNullOrEmpty(description.UserName)) ? description.UserName : UserNameAnonymous;
-	                var user =
-	                    unitOfWork.UserRepository.First(
-	                        data => data.UserName.Equals(userName, StringComparison.InvariantCultureIgnoreCase));
-
-                    var user2 =
-                       unitOfWork.UserRepository.First(
-                           data => data.UserName.Equals(userName, StringComparison.InvariantCultureIgnoreCase));
-
-                    var user3 =
-                       unitOfWork.UserRepository.First(
-                           data => data.UserName.Equals(userName, StringComparison.InvariantCultureIgnoreCase));
-
-	                if (user != null)
-	                {
-	                    newCrash.UserId = user.Id;
-	                }
-	                else
-	                {
-	                    newCrash.User = new User() {UserName = description.UserName, UserGroupId = 5};
-	                }
-
-	                user = user2;
-
-	                findUserTime = logTimer.GetElapsedSeconds().ToString("F2");
-	            }
-	        }
-
 	        //If there's a valid EpicAccountId assign that.
 	        if (!string.IsNullOrEmpty(description.EpicAccountId))
 	        {
@@ -798,8 +796,8 @@ namespace Tools.CrashReporter.CrashReportWebSite.Controllers
 	        newCrash.LanguageExt = description.Language; //Converted by the Crash process.
 	        newCrash.PlatformName = description.Platform;
 	        newCrash.HasLogFile = description.bHasLog;
-            newCrash.HasVideoFile = description.bHasVideo;
-            newCrash.HasMiniDumpFile = description.bHasMiniDump;
+	        newCrash.HasVideoFile = description.bHasVideo;
+	        newCrash.HasMiniDumpFile = description.bHasMiniDump;
 
 	        if (description.ErrorMessage != null)
 	        {
@@ -816,7 +814,7 @@ namespace Tools.CrashReporter.CrashReportWebSite.Controllers
 	            newCrash.SourceContext = string.Join("\n", description.SourceContext);
 	        }
 
-	        newCrash.TimeOfCrash = DateTime.Now;//description.TimeOfCrash;
+	        newCrash.TimeOfCrash = description.TimeofCrash;
 
 	        newCrash.Jira = "";
 	        newCrash.FixedChangeList = "";
@@ -872,34 +870,49 @@ namespace Tools.CrashReporter.CrashReportWebSite.Controllers
 	        // As we're adding it, the status is always new
 	        newCrash.Status = "New";
 	        newCrash.UserActivityHint = description.UserActivityHint;
+	        newCrash.UserName = (!string.IsNullOrEmpty(description.UserName)) ? description.UserName : UserNameAnonymous;
+            
+            //Match this crash to an existing user or create a new user if one does not exist.
+	       using (var unitOfWork = new UnitOfWork(new CrashReportEntities()))
+	        {
+	            //if there's a valid username assign the associated UserNameId else use "anonymous".
+	            var userName = (!string.IsNullOrEmpty(description.UserName)) ? description.UserName : UserNameAnonymous;
+                    
+                if (unitOfWork.UserRepository.Any(data => data.UserName.Equals(userName)))
+                {
+                    newCrash.UserId = unitOfWork.UserRepository.ListAll()
+                        .Where(data => data.UserName.Equals(userName))
+                        .Select(data => data.Id)
+                        .First();
+                }
+                else
+                {
+                    newCrash.User = new User() { UserName = description.UserName, UserGroupId = 5 };
+                }
+	        }
+
 	        try
 	        {
-	            using (
-	                var logTimer = new FAutoScopedLogTimer(this.GetType().ToString(),
-	                    bCreateNewLog: true))
-	            {
-                    BuildPattern(newCrash);
-                    buildPatternDuration = logTimer.GetElapsedSeconds().ToString("F2");
-	            }
+	            //Build the callstack pattern for this crash 
+	            BuildPattern(newCrash);
 	        }
 	        catch (Exception ex)
 	        {
-                FLogger.Global.WriteException("Error in Create Crash Build Pattern Method");
-                _slackWriter.Write("Error in Create Crash Build Pattern Method");
-                _slackWriter.Write(ex.Message);
+	            FLogger.Global.WriteException("Error in Create Crash Build Pattern Method");
+	            _slackWriter.Write("Error in Create Crash Build Pattern Method");
+	            _slackWriter.Write(ex.Message);
 	            if (ex.InnerException != null)
 	            {
-                    _slackWriter.Write(ex.InnerException.Message);
+	                _slackWriter.Write(ex.InnerException.Message);
 	            }
 	            throw;
 	        }
 
-
-            if(newCrash.CommandLine == null)
+	        if(newCrash.CommandLine == null)
                 newCrash.CommandLine = "";
 
             using (var unitOfWork = new UnitOfWork(new CrashReportEntities()))
-            {
+	        {
 	            var callStackRepository = unitOfWork.CallstackRepository;
 
 	            try
@@ -939,21 +952,21 @@ namespace Tools.CrashReporter.CrashReportWebSite.Controllers
 	                var lineRegex = new Regex(@"(\[Line:).*?(])"); //Match the line no.
 
 	                /**
-                 * Regex to match ints of two characters or longer
-                 * First term ((?<=\s)|(-)) : Positive look behind, match if preceeded by whitespace or if first character is '-'
-                 * Second term (\d{3,}) match three or more decimal chracters in a row.
-                 * Third term (?=(\s|$)) positive look ahead match if followed by whitespace or end of line/file.
-                 */
+                        * * Regex to match ints of two characters or longer
+                        * * First term ((?<=\s)|(-)) : Positive look behind, match if preceeded by whitespace or if first character is '-'
+                        * Second term (\d{3,}) match three or more decimal chracters in a row.
+                        * Third term (?=(\s|$)) positive look ahead match if followed by whitespace or end of line/file.
+                        */
 	                var intRegex = new Regex(@"-?\d{3,}");
 
 	                /**
-                 * Regular expression for masking out floats
-                 */
+                        * Regular expression for masking out floats
+                        */
 	                var floatRegex = new Regex(@"-?\d+\.\d+");
 
 	                /**
-                 * Regular expression for masking out hexadecimal numbers
-                 */
+                        * Regular expression for masking out hexadecimal numbers
+                        */
 	                var hexRegex = new Regex(@"0x[\da-fA-F]+");
 
 	                //mask out terms matches by our regex's
@@ -964,22 +977,28 @@ namespace Tools.CrashReporter.CrashReportWebSite.Controllers
 	                trimmedError = intRegex.Replace(trimmedError, "");
 	                if (trimmedError.Length > 450)
 	                    trimmedError = trimmedError.Substring(0, 450);
-	                    //error message is used as an index - trim to max index length
-                    
-                    //Check to see if the masked error message is unique
+	                //error message is used as an index - trim to max index length
+
+	                //Check to see if the masked error message is unique
 	                ErrorMessage errorMessage = null;
 	                if (
 	                    unitOfWork.ErrorMessageRepository.Any(
-	                        data => data.ErrorMessageText.Equals(trimmedError, StringComparison.InvariantCultureIgnoreCase)))
+	                        data =>
+	                            data.ErrorMessageText.Equals(trimmedError, StringComparison.InvariantCultureIgnoreCase)))
 	                {
 	                    errorMessage =
 	                        unitOfWork.ErrorMessageRepository.First(
-                                data => data.ErrorMessageText.Equals(trimmedError, StringComparison.InvariantCultureIgnoreCase));
+	                            data =>
+	                                data.ErrorMessageText.Equals(trimmedError,
+	                                    StringComparison.InvariantCultureIgnoreCase));
 	                }
 	                else
 	                {
 	                    //if it's a new message then add it to the database.
-                        errorMessage = new ErrorMessage() { ErrorMessageText = trimmedError.ToLower(CultureInfo.InvariantCulture) };
+	                    errorMessage = new ErrorMessage()
+	                    {
+	                        ErrorMessageText = trimmedError.ToLower(CultureInfo.InvariantCulture)
+	                    };
 	                    unitOfWork.ErrorMessageRepository.Save(errorMessage);
 	                    unitOfWork.Save();
 	                }
@@ -989,7 +1008,7 @@ namespace Tools.CrashReporter.CrashReportWebSite.Controllers
 	                    (data.PatternId == newCrash.PatternId) &&
 	                    (data.ErrorMessageId == errorMessage.Id || data.ErrorMessageId == null)))
 	                {
-                        //if a bugg exists for this pattern update the bugg data
+	                    //if a bugg exists for this pattern update the bugg data
 	                    var bugg = unitOfWork.BuggRepository.First(data => data.PatternId == newCrash.PatternId);
 	                    bugg.PatternId = newCrash.PatternId;
 
@@ -1079,11 +1098,12 @@ namespace Tools.CrashReporter.CrashReportWebSite.Controllers
 	                }
 
 	                FLogger.Global.WriteException("Create Crash Exception : " +
-	                                              messageBuilder.ToString());
+	                                                messageBuilder.ToString());
 	                _slackWriter.Write("Create Crash Exception : " + messageBuilder.ToString());
 	                throw;
 	            }
-	        }
+            }
+
 	        return newCrash;
 	    }
 
@@ -1094,58 +1114,64 @@ namespace Tools.CrashReporter.CrashReportWebSite.Controllers
 	    /// <param name="newCrash"></param>
         private void BuildPattern(Crash newCrash)
 	    {
-	        var callstack = new CallStackContainer(newCrash);
-	        newCrash.Module = callstack.GetModuleName();
-	        if (newCrash.PatternId == null)
-	        {
-	            var patternList = new List<string>();
+	        
+            var callstack = new CallStackContainer(newCrash.CrashType, newCrash.RawCallStack, newCrash.PlatformName);
+	            newCrash.Module = callstack.GetModuleName();
+	            newCrash.Module = newCrash.Module.Length >= 127 ? newCrash.Module.Substring(0, 127) : newCrash.Module;
 
-	            try
+	            if (newCrash.PatternId == null)
 	            {
-	                using (var unitOfWork = new UnitOfWork(new CrashReportEntities()))
-	                {
-	                    foreach (var entry in callstack.CallStackEntries.Take(CallStackContainer.MaxLinesToParse))
-	                    {
-	                        FunctionCall currentFunctionCall;
+	                var patternList = new List<string>();
 
-	                        var csEntry = entry;
-	                        var functionCall = unitOfWork.FunctionRepository.First(f => f.Call == csEntry.FunctionName);
-	                        if (functionCall != null)
-	                        {
-	                            currentFunctionCall = functionCall;
-	                        }
-	                        else
-	                        {
-	                            var call = csEntry.FunctionName;
-	                            if (call.Length > 899)
-	                                call = call.Substring(0, 899);
-	                            currentFunctionCall = new FunctionCall {Call = call};
-	                            unitOfWork.FunctionRepository.Save(currentFunctionCall);
-	                            unitOfWork.Save();
+	                try
+                    {
+                        using (var unitOfWork = new UnitOfWork(new CrashReportEntities()))
+                        {
+                            foreach (var entry in callstack.CallStackEntries.Take(CallStackContainer.MaxLinesToParse))
+                            {
+                                var functionCallId = 0;
+
+	                            //if this function is already in the database get id
+	                            if (unitOfWork.FunctionRepository.Any(f => f.Call == entry.FunctionName))
+	                            {
+	                                functionCallId =
+	                                    unitOfWork.FunctionRepository.FirstId(f => f.Call == entry.FunctionName);
+	                            }
+	                            else
+	                            {
+	                                //else add this function to the db.
+	                                var call = entry.FunctionName;
+	                                if (call.Length > 899)
+	                                    call = call.Substring(0, 899);
+	                                var currentFunctionCall = new FunctionCall {Call = call};
+	                                unitOfWork.FunctionRepository.Save(currentFunctionCall);
+	                                unitOfWork.Save();
+	                                functionCallId = currentFunctionCall.Id;
+	                            }
+
+	                            patternList.Add(functionCallId.ToString());
 	                        }
-
-	                        patternList.Add(currentFunctionCall.Id.ToString());
 	                    }
-	                }
 
-	                newCrash.Pattern = string.Join("+", patternList);
-	            }
-	            catch (Exception ex)
-	            {
-                    var messageBuilder = new StringBuilder();
-                    FLogger.Global.WriteException("Build Pattern exception: " + ex.Message.ToString(CultureInfo.InvariantCulture));
-                    messageBuilder.AppendLine("Exception was:");
-                    messageBuilder.AppendLine(ex.ToString());
-                    while (ex.InnerException != null)
-                    {
-                        ex = ex.InnerException;
-                        messageBuilder.AppendLine(ex.ToString());
-                    }
+	                    newCrash.Pattern = string.Join("+", patternList);
+	                }
+	                catch (Exception ex)
+	                {
+	                    var messageBuilder = new StringBuilder();
+	                    FLogger.Global.WriteException("Build Pattern exception: " +
+	                                                  ex.Message.ToString(CultureInfo.InvariantCulture));
+	                    messageBuilder.AppendLine("Exception was:");
+	                    messageBuilder.AppendLine(ex.ToString());
+	                    while (ex.InnerException != null)
+	                    {
+	                        ex = ex.InnerException;
+	                        messageBuilder.AppendLine(ex.ToString());
+	                    }
 
-                    _slackWriter.Write("Build Pattern Exception : " + ex.Message.ToString(CultureInfo.InvariantCulture));
-                    throw;
+	                    _slackWriter.Write("Build Pattern Exception : " + ex.Message.ToString(CultureInfo.InvariantCulture));
+	                    throw;
+	                }
 	            }
-	        }
 	    }
 
         /// <summary></summary>
@@ -1161,8 +1187,7 @@ namespace Tools.CrashReporter.CrashReportWebSite.Controllers
 	    public void TestAddCrash()
 	    {
             using (
-                var logTimer = new FAutoScopedLogTimer(this.GetType().ToString() + "()",
-                    bCreateNewLog: true))
+                var logTimer = new FAutoScopedLogTimer(this.GetType().ToString() + "()"))
             {
                 string CrashContext;
                 using (
@@ -1176,13 +1201,9 @@ namespace Tools.CrashReporter.CrashReportWebSite.Controllers
 
                 var Crash = XmlHandler.FromXmlString<CrashDescription>(CrashContext);
 
-                var buildPatternDuration = "";
-                var userTime = "";
-                var crash = CreateCrash(Crash, out buildPatternDuration, out userTime);
+                var crash = CreateCrash(Crash);
 
-                _slackWriter.Write("Test Add Crash completed. Pattern Match time : " + buildPatternDuration +
-                    " Find User Time : " + userTime + 
-                    " Total Time : " + logTimer.GetElapsedSeconds().ToString("F2"));
+                _slackWriter.Write("Test Add Crash completed. Total Time : " + logTimer.GetElapsedSeconds().ToString("F2"));
 
             }
 	    }
