Copying //UE4/Dev-Core to //UE4/Dev-Main (Source: //UE4/Dev-Core @ 3345728)

diff --git a/Engine/Source/Programs/CrashReporter/CrashReportWebSite/Controllers/ReportsController.cs b/Engine/Source/Programs/CrashReporter/CrashReportWebSite/Controllers/ReportsController.cs
--- a/Engine/Source/Programs/CrashReporter/CrashReportWebSite/Controllers/ReportsController.cs
+++ b/Engine/Source/Programs/CrashReporter/CrashReportWebSite/Controllers/ReportsController.cs
@@ -58,7 +58,7 @@ namespace Tools.CrashReporter.CrashReportWebSite.Controllers
 		{
 			using( FAutoScopedLogTimer LogTimer = new FAutoScopedLogTimer( this.GetType().ToString(), bCreateNewLog: true ) )
 			{
-				FormHelper FormData = new FormHelper( Request, ReportsForm, "JustReport" );
+				var FormData = new FormHelper( Request, ReportsForm, "JustReport" );
 
 				// Handle 'CopyToJira' button
 				int BuggIDToBeAddedToJira = 0;
@@ -255,296 +255,304 @@ namespace Tools.CrashReporter.CrashReportWebSite.Controllers
                     _unitOfWork.UserGroupRepository.First(data => data.Id == anonymousGroupId).Users.Select(data => data.Id).ToList();
 	            var anonymousId = anonymousIDs.First();
 
-	            var CrashQuery = _unitOfWork.CrashRepository.ListAll().Where(data => data.TimeOfCrash > startDate && data.TimeOfCrash <= endDate)
+	            var crashQuery = _unitOfWork.CrashRepository.ListAll().Where(data => data.TimeOfCrash > startDate && data.TimeOfCrash <= endDate)
 	                // Only Crashes and asserts
-	                .Where(Crash => Crash.CrashType == 1 || Crash.CrashType == 2)
+                    .Where(crash => crash.CrashType == 1 || crash.CrashType == 2)
 	                // Only anonymous user
-	                .Where(Crash => Crash.UserId == anonymousId);
+                    .Where(crash => crash.UserId == anonymousId)
+                    .Where(crash => crash.BuggId.HasValue);
 
 	            // Filter by BranchName
 	            if (!string.IsNullOrEmpty(branchName))
 	            {
-	                CrashQuery = CrashQuery.Where(data => data.Branch == branchName);
+	                crashQuery = crashQuery.Where(data => data.Branch == branchName);
 	            }
 
-                var crashes = CrashQuery.Select(data => new
-	            {
+	            var buggQuery =
+	                _unitOfWork.BuggRepository.ListAll()
+                        .Where(bugg => crashQuery.Select(data => data.BuggId).Distinct().Contains(bugg.Id));
+
+                var crashes = crashQuery.Select(data => new
+                {
                     ID = data.Id,
+                    BuggId = data.BuggId,
                     TimeOfCrash = data.TimeOfCrash,
-	                UserID = data.UserId,
+                    UserID = data.UserId,
                     BuildVersion = data.BuildVersion,
                     JIRA = data.Jira,
                     Platform = data.PlatformName,
                     FixCL = data.FixedChangeList,
                     BuiltFromCL = data.ChangelistVersion,
-                    Pattern = data.Pattern,
                     MachineID = data.ComputerName,
                     Branch = data.Branch,
                     Description = data.Description,
                     RawCallStack = data.RawCallStack,
-	            }).ToList();
-
-                var numCrashes = CrashQuery.Count();
-
-	            // Total # of ALL (Anonymous) Crashes in timeframe
-	            var totalAnonymousCrashes = numCrashes;
-
-	            // Total # of UNIQUE (Anonymous) Crashes in timeframe
-                var uniquePatterns = new HashSet<string>();
-                var uniqueMachines = new HashSet<string>();
-	            var patternToCount = new Dictionary<string, int>();
-                
-	            foreach (var crash in crashes)
-	            {
-                    uniquePatterns.Add(crash.Pattern);
-                    uniqueMachines.Add(crash.MachineID);
-
-                    var bAdd = !patternToCount.ContainsKey(crash.Pattern);
-	                if (bAdd)
-	                {
-                        patternToCount.Add(crash.Pattern, 1);
-	                }
-	                else
-	                {
-                        patternToCount[crash.Pattern]++;
-	                }
-	            }
-	            var PatternToCountOrdered = patternToCount.OrderByDescending(X => X.Value).ToList();
-	            const int NumTopRecords = 200;
-	            var PatternAndCount = PatternToCountOrdered.Take(NumTopRecords).ToDictionary(x => x.Key, y => y.Value);
-
-	            int TotalUniqueAnonymousCrashes = uniquePatterns.Count;
-
-	            // Total # of AFFECTED USERS (Anonymous) in timeframe
-	            int TotalAffectedUsers = uniqueMachines.Count;
-
-	            var RealBuggs = _unitOfWork.BuggRepository.ListAll().Where(Bugg => PatternAndCount.Keys.Contains(Bugg.Pattern));
-
-	            // Build search string.
-	            var foundJiras = new HashSet<string>();
-	            var jiraIDtoBugg = new Dictionary<string, List<Bugg>>();
-
-	            var buggs = new List<Bugg>(NumTopRecords);
-	            foreach (var Top in PatternAndCount)
-	            {
-	                var newBugg = RealBuggs.FirstOrDefault(X => X.Pattern == Top.Key);
-	                if (newBugg != null)
-	                {
-                        var CrashForBugg = crashes.Where(Crash => Crash.Pattern == Top.Key).ToList();
-
-	                    // Convert anonymous to full type;
-                        var fullCrashesForBugg = new List<Crash>(CrashForBugg.Count);
-                        foreach (var anon in CrashForBugg)
-	                    {
-                            fullCrashesForBugg.Add(new Crash()
-	                        {
-	                            Id = anon.ID,
-	                            TimeOfCrash = anon.TimeOfCrash,
-	                            BuildVersion = anon.BuildVersion,
-	                            Jira = anon.JIRA,
-	                            PlatformName = anon.Platform,
-	                            FixedChangeList = anon.FixCL,
-	                            ChangelistVersion = anon.BuiltFromCL,
-	                            Pattern = anon.Pattern,
-	                            ComputerName = anon.MachineID,
-	                            Branch = anon.Branch,
-	                            Description = anon.Description,
-	                            RawCallStack = anon.RawCallStack,
-	                        });
-	                    }
-
-	                    newBugg.PrepareBuggForJira(fullCrashesForBugg);
-
-	                    // Verify valid JiraID, this may be still a TTP 
-	                    if (!string.IsNullOrEmpty(newBugg.TTPID))
-	                    {
-	                        int ttpid = 0;
-	                        int.TryParse(newBugg.TTPID, out ttpid);
-
-	                        if (ttpid == 0)
-	                        {
-	                            AddBuggJiraMapping(newBugg, ref foundJiras, ref jiraIDtoBugg);
-	                        }
-	                    }
-
-	                    buggs.Add(newBugg);
-	                }
-	                else
-	                {
-	                    FLogger.Global.WriteEvent("Bugg for pattern " + Top.Key + " not found");
-	                }
-	            }
-
-	            if (BuggIDToBeAddedToJira > 0)
-	            {
-                    var bugg = buggs.FirstOrDefault(x => x.Id == BuggIDToBeAddedToJira);
-                    if (bugg != null)
-                    {
-                        bugg.CopyToJira();
-                        AddBuggJiraMapping(bugg, ref foundJiras, ref jiraIDtoBugg);
-                    }
-	            }
-
-	            if (JC.CanBeUsed())
-	            {
-	                var buggsCopy = new List<Bugg>(buggs);
-
-	                var invalidJiras = new HashSet<string>();
-
-	                // Simple verification of JIRA
-	                foreach (var Value in foundJiras)
-	                {
-	                    if (Value.Length < 3 || !Value.Contains('-'))
-	                    {
-	                        invalidJiras.Add(Value);
-	                    }
-	                }
-
-	                foreach (var invalidJira in invalidJiras)
-	                {
-	                    foundJiras.Remove(invalidJira);
-	                }
-
-	                // Grab the data form JIRA.
-	                string JiraSearchQuery = string.Join(" OR ", foundJiras);
+                }).ToList();
+
+                //var numCrashes = crashQuery.Count();
+
+                //// Total # of ALL (Anonymous) Crashes in timeframe
+                //var totalAnonymousCrashes = numCrashes;
+
+                //// Total # of UNIQUE (Anonymous) Crashes in timeframe
+                //var uniqueBuggs = new HashSet<int>();
+                //var uniqueMachines = new HashSet<string>();
+                //var buggToCount = new Dictionary<int, int>();
+
+                //foreach (var crash in crashes)
+                //{
+                //    if (!crash.BuggId.HasValue)
+                //        continue;
+
+                //    uniqueBuggs.Add(crash.BuggId.Value);
+                //    uniqueMachines.Add(crash.MachineID);
+
+                //    var bAdd = !buggToCount.ContainsKey(crash.Pattern);
+                //    if (bAdd)
+                //    {
+                //        buggToCount.Add(crash.BuggId.Value, 1);
+                //    }
+                //    else
+                //    {
+                //        buggToCount[crash.BuggId.Value]++;
+                //    }
+                //}
+                //var PatternToCountOrdered = buggToCount.OrderByDescending(X => X.Value).ToList();
+                //const int NumTopRecords = 50;
+                //var PatternAndCount = PatternToCountOrdered.Take(NumTopRecords).ToDictionary(x => x.Key, y => y.Value);
+
+                //int TotalUniqueAnonymousCrashes = uniqueBuggs.Count;
+
+                //// Total # of AFFECTED USERS (Anonymous) in timeframe
+                //int TotalAffectedUsers = uniqueMachines.Count;
+
+                //var RealBuggs = _unitOfWork.BuggRepository.ListAll().Where(Bugg => crashes.Select(data => data.BuggId).Contains(Bugg.Id));
+
+                //var realBuggs = RealBuggs.OrderByDescending(data => data.Crashes.Count);
+
+                //// Build search string.
+                //var foundJiras = new HashSet<string>();
+                //var jiraIDtoBugg = new Dictionary<string, List<Bugg>>();
+
+                //var buggs = new List<Bugg>(NumTopRecords);
+                //foreach (var Top in PatternAndCount)
+                //{
+                //    var newBugg = RealBuggs.FirstOrDefault(X => X.Id == Top.Key);
+                //    if (newBugg != null)
+                //    {
+                //        var CrashForBugg = crashes.Where(Crash => Crash.Pattern == Top.Key).ToList();
+
+                //        // Convert anonymous to full type;
+                //        var fullCrashesForBugg = new List<CrashDataModel>(CrashForBugg.Count);
+                //        foreach (var anon in CrashForBugg)
+                //        {
+                //            fullCrashesForBugg.Add(new CrashDataModel()
+                //            {
+                //                Id = anon.ID,
+                //                TimeOfCrash = anon.TimeOfCrash,
+                //                BuildVersion = anon.BuildVersion,
+                //                PlatformName = anon.Platform,
+                //                FixedChangeList = anon.FixCL,
+                //                ChangelistVersion = anon.BuiltFromCL,
+                //                ComputerName = anon.MachineID,
+                //                Branch = anon.Branch,
+                //                Description = anon.Description,
+                //                RawCallStack = anon.RawCallStack,
+                //            });
+                //        }
+
+                //        newBugg.PrepareBuggForJira(fullCrashesForBugg);
+
+                //        // Verify valid JiraID, this may be still a TTP 
+                //        if (!string.IsNullOrEmpty(newBugg.TTPID))
+                //        {
+                //            int ttpid = 0;
+                //            int.TryParse(newBugg.TTPID, out ttpid);
+
+                //            if (ttpid == 0)
+                //            {
+                //                AddBuggJiraMapping(newBugg, ref foundJiras, ref jiraIDtoBugg);
+                //            }
+                //        }
+
+                //        buggs.Add(newBugg);
+                //    }
+                //    else
+                //    {
+                //        FLogger.Global.WriteEvent("Bugg for pattern " + Top.Key + " not found");
+                //    }
+                //}
+
+                //if (BuggIDToBeAddedToJira > 0)
+                //{
+                //    var bugg = buggs.FirstOrDefault(x => x.Id == BuggIDToBeAddedToJira);
+                //    if (bugg != null)
+                //    {
+                //        bugg.CopyToJira();
+                //        AddBuggJiraMapping(bugg, ref foundJiras, ref jiraIDtoBugg);
+                //    }
+                //}
+
+                //if (JC.CanBeUsed())
+                //{
+                //    var buggsCopy = new List<Bugg>(buggs);
+
+                //    var invalidJiras = new HashSet<string>();
+
+                //    // Simple verification of JIRA
+                //    foreach (var Value in foundJiras)
+                //    {
+                //        if (Value.Length < 3 || !Value.Contains('-'))
+                //        {
+                //            invalidJiras.Add(Value);
+                //        }
+                //    }
+
+                //    foreach (var invalidJira in invalidJiras)
+                //    {
+                //        foundJiras.Remove(invalidJira);
+                //    }
+
+                //    // Grab the data form JIRA.
+                //    string JiraSearchQuery = string.Join(" OR ", foundJiras);
                     
-                    bool bInvalid = false;
-                    var jiraResults = new Dictionary<string, Dictionary<string, object>>();
-	                try
-	                {
-	                    if (!string.IsNullOrWhiteSpace(JiraSearchQuery))
-	                    {
-	                        jiraResults = JC.SearchJiraTickets(
-	                            JiraSearchQuery,
-	                            new string[]
-	                            {
-	                                "key", // string
-	                                "summary", // string
-	                                "components", // System.Collections.ArrayList, Dictionary<string,object>, name
-	                                "resolution",
-	                                //System.Collections.Generic.Dictionary`2[System.String,System.Object], name
-	                                "fixVersions", // System.Collections.ArrayList, Dictionary<string,object>, name
-	                                "customfield_11200" // string
-	                            });
-	                    }
-	                }
-	                catch (System.Exception)
-	                {
-	                    bInvalid = true;
-	                }
-
-	                // Invalid records have been found, find the broken using the slow path.
-	                if (bInvalid)
-	                {
-	                    foreach (var Query in foundJiras)
-	                    {
-	                        try
-	                        {
-	                            var TempResult = JC.SearchJiraTickets(
-	                                Query,
-	                                new string[]
-	                                {
-	                                    "key", // string
-	                                    "summary", // string
-	                                    "components", // System.Collections.ArrayList, Dictionary<string,object>, name
-	                                    "resolution",
-	                                    // System.Collections.Generic.Dictionary`2[System.String,System.Object], name
-	                                    "fixVersions", // System.Collections.ArrayList, Dictionary<string,object>, name
-	                                    "customfield_11200" // string
-	                                });
-
-	                            foreach (var Temp in TempResult)
-	                            {
-	                                jiraResults.Add(Temp.Key, Temp.Value);
-	                            }
-	                        }
-	                        catch (System.Exception)
-	                        {
-
-	                        }
-	                    }
-	                }
-
-	                // Jira Key, Summary, Components, Resolution, Fix version, Fix changelist
-	                foreach (var Jira in jiraResults)
-	                {
-	                    string JiraID = Jira.Key;
-
-	                    string Summary = (string) Jira.Value["summary"];
-
-	                    string ComponentsText = "";
-	                    System.Collections.ArrayList Components =
-	                        (System.Collections.ArrayList) Jira.Value["components"];
-	                    foreach (Dictionary<string, object> Component in Components)
-	                    {
-	                        ComponentsText += (string) Component["name"];
-	                        ComponentsText += " ";
-	                    }
-
-	                    Dictionary<string, object> ResolutionFields =
-	                        (Dictionary<string, object>) Jira.Value["resolution"];
-	                    string Resolution = ResolutionFields != null ? (string) ResolutionFields["name"] : "";
-
-	                    string FixVersionsText = "";
-	                    System.Collections.ArrayList FixVersions =
-	                        (System.Collections.ArrayList) Jira.Value["fixVersions"];
-	                    foreach (Dictionary<string, object> FixVersion in FixVersions)
-	                    {
-	                        FixVersionsText += (string) FixVersion["name"];
-	                        FixVersionsText += " ";
-	                    }
-
-	                    int FixCL = Jira.Value["customfield_11200"] != null
-	                        ? (int) (decimal) Jira.Value["customfield_11200"]
-	                        : 0;
-
-	                    List<Bugg> BuggsForJira;
-	                    jiraIDtoBugg.TryGetValue(JiraID, out BuggsForJira);
-
-	                    //var BuggsForJira = JiraIDtoBugg[JiraID];
-
-	                    if (BuggsForJira != null)
-	                    {
-	                        foreach (Bugg Bugg in BuggsForJira)
-	                        {
-	                            Bugg.JiraSummary = Summary;
-	                            Bugg.JiraComponentsText = ComponentsText;
-	                            Bugg.JiraResolution = Resolution;
-	                            Bugg.JiraFixVersionsText = FixVersionsText;
-	                            if (FixCL != 0)
-	                            {
-	                                Bugg.JiraFixCL = FixCL.ToString();
-	                            }
-
-	                            buggsCopy.Remove(Bugg);
-	                        }
-	                    }
-	                }
-
-	                // If there are buggs, we need to update the summary to indicate an error.
-	                // Usually caused when bugg's project has changed.
-	                foreach (var Bugg in buggsCopy.Where(b => !string.IsNullOrEmpty(b.TTPID)))
-	                {
-	                    Bugg.JiraSummary = "JIRA MISMATCH";
-	                    Bugg.JiraComponentsText = "JIRA MISMATCH";
-	                    Bugg.JiraResolution = "JIRA MISMATCH";
-	                    Bugg.JiraFixVersionsText = "JIRA MISMATCH";
-	                    Bugg.JiraFixCL = "JIRA MISMATCH";
-	                }
-	            }
-
-	            buggs = buggs.OrderByDescending(b => b.CrashesInTimeFrameGroup).ToList();
+                //    bool bInvalid = false;
+                //    var jiraResults = new Dictionary<string, Dictionary<string, object>>();
+                //    try
+                //    {
+                //        if (!string.IsNullOrWhiteSpace(JiraSearchQuery))
+                //        {
+                //            jiraResults = JC.SearchJiraTickets(
+                //                JiraSearchQuery,
+                //                new string[]
+                //                {
+                //                    "key", // string
+                //                    "summary", // string
+                //                    "components", // System.Collections.ArrayList, Dictionary<string,object>, name
+                //                    "resolution",
+                //                    //System.Collections.Generic.Dictionary`2[System.String,System.Object], name
+                //                    "fixVersions", // System.Collections.ArrayList, Dictionary<string,object>, name
+                //                    "customfield_11200" // string
+                //                });
+                //        }
+                //    }
+                //    catch (System.Exception)
+                //    {
+                //        bInvalid = true;
+                //    }
+
+                //    // Invalid records have been found, find the broken using the slow path.
+                //    if (bInvalid)
+                //    {
+                //        foreach (var Query in foundJiras)
+                //        {
+                //            try
+                //            {
+                //                var TempResult = JC.SearchJiraTickets(
+                //                    Query,
+                //                    new string[]
+                //                    {
+                //                        "key", // string
+                //                        "summary", // string
+                //                        "components", // System.Collections.ArrayList, Dictionary<string,object>, name
+                //                        "resolution",
+                //                        // System.Collections.Generic.Dictionary`2[System.String,System.Object], name
+                //                        "fixVersions", // System.Collections.ArrayList, Dictionary<string,object>, name
+                //                        "customfield_11200" // string
+                //                    });
+
+                //                foreach (var Temp in TempResult)
+                //                {
+                //                    jiraResults.Add(Temp.Key, Temp.Value);
+                //                }
+                //            }
+                //            catch (System.Exception)
+                //            {
+
+                //            }
+                //        }
+                //    }
+
+                //    // Jira Key, Summary, Components, Resolution, Fix version, Fix changelist
+                //    foreach (var Jira in jiraResults)
+                //    {
+                //        string JiraID = Jira.Key;
+
+                //        string Summary = (string) Jira.Value["summary"];
+
+                //        string ComponentsText = "";
+                //        System.Collections.ArrayList Components =
+                //            (System.Collections.ArrayList) Jira.Value["components"];
+                //        foreach (Dictionary<string, object> Component in Components)
+                //        {
+                //            ComponentsText += (string) Component["name"];
+                //            ComponentsText += " ";
+                //        }
+
+                //        Dictionary<string, object> ResolutionFields =
+                //            (Dictionary<string, object>) Jira.Value["resolution"];
+                //        string Resolution = ResolutionFields != null ? (string) ResolutionFields["name"] : "";
+
+                //        string FixVersionsText = "";
+                //        System.Collections.ArrayList FixVersions =
+                //            (System.Collections.ArrayList) Jira.Value["fixVersions"];
+                //        foreach (Dictionary<string, object> FixVersion in FixVersions)
+                //        {
+                //            FixVersionsText += (string) FixVersion["name"];
+                //            FixVersionsText += " ";
+                //        }
+
+                //        int FixCL = Jira.Value["customfield_11200"] != null
+                //            ? (int) (decimal) Jira.Value["customfield_11200"]
+                //            : 0;
+
+                //        List<Bugg> BuggsForJira;
+                //        jiraIDtoBugg.TryGetValue(JiraID, out BuggsForJira);
+
+                //        //var BuggsForJira = JiraIDtoBugg[JiraID];
+
+                //        if (BuggsForJira != null)
+                //        {
+                //            foreach (Bugg Bugg in BuggsForJira)
+                //            {
+                //                Bugg.JiraSummary = Summary;
+                //                Bugg.JiraComponentsText = ComponentsText;
+                //                Bugg.JiraResolution = Resolution;
+                //                Bugg.JiraFixVersionsText = FixVersionsText;
+                //                if (FixCL != 0)
+                //                {
+                //                    Bugg.JiraFixCL = FixCL.ToString();
+                //                }
+
+                //                buggsCopy.Remove(Bugg);
+                //            }
+                //        }
+                //    }
+
+                //    // If there are buggs, we need to update the summary to indicate an error.
+                //    // Usually caused when bugg's project has changed.
+                //    foreach (var Bugg in buggsCopy.Where(b => !string.IsNullOrEmpty(b.TTPID)))
+                //    {
+                //        Bugg.JiraSummary = "JIRA MISMATCH";
+                //        Bugg.JiraComponentsText = "JIRA MISMATCH";
+                //        Bugg.JiraResolution = "JIRA MISMATCH";
+                //        Bugg.JiraFixVersionsText = "JIRA MISMATCH";
+                //        Bugg.JiraFixCL = "JIRA MISMATCH";
+                //    }
+                //}
+
+                //buggs = buggs.OrderByDescending(b => b.CrashesInTimeFrameGroup).ToList();
 
 	            return new ReportsViewModel
 	            {
-	                Buggs = buggs,
+	               // Buggs = buggs,
                     BranchName = branchName,
 	                BranchNames = _unitOfWork.CrashRepository.GetBranchesAsListItems(),
 	                DateFrom = (long) (startDate - CrashesViewModel.Epoch).TotalMilliseconds,
 	                DateTo = (long) (endDate - CrashesViewModel.Epoch).TotalMilliseconds,
-	                TotalAffectedUsers = TotalAffectedUsers,
-	                TotalAnonymousCrashes = totalAnonymousCrashes,
-	                TotalUniqueAnonymousCrashes = TotalUniqueAnonymousCrashes
+	                //TotalAffectedUsers = TotalAffectedUsers,
+	                //TotalAnonymousCrashes = totalAnonymousCrashes,
+	                //TotalUniqueAnonymousCrashes = TotalUniqueAnonymousCrashes
 	            };
 	        }
 	    }
@@ -567,9 +575,9 @@ namespace Tools.CrashReporter.CrashReportWebSite.Controllers
             var CrashesQuery = _unitOfWork.CrashRepository.ListAll()
                 .Where(data => data.TimeOfCrash > startDate && data.TimeOfCrash <= endDate)
                 // Only Crashes and asserts
-                .Where(Crash => Crash.CrashType == 1 || Crash.CrashType == 2);
+                .Where(crash => crash.CrashType == 1 || crash.CrashType == 2)
                 // Only anonymous user
-                    //.Where(Crash => Crash.UserNameId.Value == anonymousId);
+                .Where(crash => crash.UserId == anonymousId);
 
             // Filter by BranchName
             if (!string.IsNullOrEmpty(branchName))
@@ -586,9 +594,24 @@ namespace Tools.CrashReporter.CrashReportWebSite.Controllers
                 if (bugg == null)
                     continue;
 
-                //var Crashes = CrashesQuery.Where(data => data.BuggId == bugg.Id).ToList();
+                var crashData = CrashesQuery.Where(data => data.BuggId == bugg.Id).Select(data => new CrashDataModel()
+                {
+                    Id = data.Id,
+                    ChangelistVersion = data.ChangelistVersion,
+                    GameName = data.GameName,
+                    EngineMode = data.EngineMode,
+                    PlatformName = data.PlatformName,
+                    TimeOfCrash = data.TimeOfCrash,
+                    Description = data.Description,
+                    RawCallStack = data.RawCallStack,
+                    CrashType = data.CrashType,
+                    Summary = data.Summary,
+                    BuildVersion = data.BuildVersion,
+                    ComputerName = data.ComputerName,
+                    Branch = data.Branch
+                }).ToList();
 
-                //bugg.PrepareBuggForJira(Crashes);
+                bugg.PrepareBuggForJira(crashData);
 
                 bugg.CrashesInTimeFrameGroup = CrashesQuery.Count(data => data.BuggId == bugg.Id);
                 bugg.AffectedVersions = new SortedSet<string>(CrashesQuery.Where(data => data.BuggId == bugg.Id).Select(data => data.BuildVersion));
@@ -610,95 +633,95 @@ namespace Tools.CrashReporter.CrashReportWebSite.Controllers
             }
 
             //Get jira connection
-            //var jiraConnection = JiraConnection.Get();
-
-            //if (BuggIDToBeAddedToJira > 0)
-            //{
-            //    var bugg = buggsList.FirstOrDefault(x => x.Id == BuggIDToBeAddedToJira);
-            //    if (bugg != null)
-            //    {
-            //        bugg.CopyToJira();
-            //    }
-            //}
-
-            //var jiraResults = new Dictionary<string, Dictionary<string, object>>();
-
-            //foreach (var Query in buggsList.Where(data =>
-            //    !string.IsNullOrWhiteSpace(data.TTPID)).Select(data => data.TTPID))
-            //{
-            //    try
-            //    {
-            //        var TempResult = jiraConnection.SearchJiraTickets(
-            //            "key = " + Query,
-            //            new string[]
-            //            {
-            //                "key", // string
-            //                "summary", // string
-            //                "components", // System.Collections.ArrayList, Dictionary<string,object>, name
-            //                "resolution",
-            //                // System.Collections.Generic.Dictionary`2[System.String,System.Object], name
-            //                "fixVersions", // System.Collections.ArrayList, Dictionary<string,object>, name
-            //                "customfield_11200" // string
-            //            });
-
-            //        foreach (var Temp in TempResult)
-            //        {
-            //            jiraResults.Add(Temp.Key, Temp.Value);
-            //        }
-            //    }
-            //    catch (System.Exception)
-            //    {
-
-            //    }
-            //}
-
-            //// Jira Key, Summary, Components, Resolution, Fix version, Fix changelist
-            //foreach (var Jira in jiraResults)
-            //{
-            //    var jiraId = Jira.Key;
-
-            //    var summary = (string)Jira.Value["summary"];
-
-            //    var ComponentsText = "";
-            //    System.Collections.ArrayList Components =
-            //        (System.Collections.ArrayList)Jira.Value["components"];
-            //    foreach (Dictionary<string, object> Component in Components)
-            //    {
-            //        ComponentsText += (string)Component["name"];
-            //        ComponentsText += " ";
-            //    }
-
-            //    Dictionary<string, object> ResolutionFields =
-            //        (Dictionary<string, object>)Jira.Value["resolution"];
-            //    string Resolution = ResolutionFields != null ? (string)ResolutionFields["name"] : "";
-
-            //    string FixVersionsText = "";
-            //    System.Collections.ArrayList FixVersions =
-            //        (System.Collections.ArrayList)Jira.Value["fixVersions"];
-            //    foreach (Dictionary<string, object> FixVersion in FixVersions)
-            //    {
-            //        FixVersionsText += (string)FixVersion["name"];
-            //        FixVersionsText += " ";
-            //    }
-
-            //    int FixCL = Jira.Value["customfield_11200"] != null
-            //        ? (int)(decimal)Jira.Value["customfield_11200"]
-            //        : 0;
-
-
-            //    if (buggsList.Any(data => data.TTPID == jiraId))
-            //    {
-            //        var bugg = buggsList.First(data => data.TTPID == jiraId);
-            //        bugg.JiraSummary = summary;
-            //        bugg.JiraComponentsText = ComponentsText;
-            //        bugg.JiraResolution = Resolution;
-            //        bugg.JiraFixVersionsText = FixVersionsText;
-            //        if (FixCL != 0)
-            //        {
-            //            bugg.JiraFixCL = FixCL.ToString();
-            //        }
-            //    }
-            //}
+            var jiraConnection = JiraConnection.Get();
+
+            if (BuggIDToBeAddedToJira > 0)
+            {
+                var bugg = buggsList.FirstOrDefault(x => x.Id == BuggIDToBeAddedToJira);
+                if (bugg != null)
+                {
+                    bugg.CopyToJira();
+                }
+            }
+
+            var jiraResults = new Dictionary<string, Dictionary<string, object>>();
+
+            foreach (var Query in buggsList.Where(data =>
+                !string.IsNullOrWhiteSpace(data.TTPID)).Select(data => data.TTPID))
+            {
+                try
+                {
+                    var TempResult = jiraConnection.SearchJiraTickets(
+                        "key = " + Query,
+                        new string[]
+                        {
+                            "key", // string
+                            "summary", // string
+                            "components", // System.Collections.ArrayList, Dictionary<string,object>, name
+                            "resolution",
+                            // System.Collections.Generic.Dictionary`2[System.String,System.Object], name
+                            "fixVersions", // System.Collections.ArrayList, Dictionary<string,object>, name
+                            "customfield_11200" // string
+                        });
+
+                    foreach (var Temp in TempResult)
+                    {
+                        jiraResults.Add(Temp.Key, Temp.Value);
+                    }
+                }
+                catch (System.Exception)
+                {
+
+                }
+            }
+
+            // Jira Key, Summary, Components, Resolution, Fix version, Fix changelist
+            foreach (var Jira in jiraResults)
+            {
+                var jiraId = Jira.Key;
+
+                var summary = (string)Jira.Value["summary"];
+
+                var ComponentsText = "";
+                System.Collections.ArrayList Components =
+                    (System.Collections.ArrayList)Jira.Value["components"];
+                foreach (Dictionary<string, object> Component in Components)
+                {
+                    ComponentsText += (string)Component["name"];
+                    ComponentsText += " ";
+                }
+
+                Dictionary<string, object> ResolutionFields =
+                    (Dictionary<string, object>)Jira.Value["resolution"];
+                string Resolution = ResolutionFields != null ? (string)ResolutionFields["name"] : "";
+
+                string FixVersionsText = "";
+                System.Collections.ArrayList FixVersions =
+                    (System.Collections.ArrayList)Jira.Value["fixVersions"];
+                foreach (Dictionary<string, object> FixVersion in FixVersions)
+                {
+                    FixVersionsText += (string)FixVersion["name"];
+                    FixVersionsText += " ";
+                }
+
+                int FixCL = Jira.Value["customfield_11200"] != null
+                    ? (int)(decimal)Jira.Value["customfield_11200"]
+                    : 0;
+
+
+                if (buggsList.Any(data => data.TTPID == jiraId))
+                {
+                    var bugg = buggsList.First(data => data.TTPID == jiraId);
+                    bugg.JiraSummary = summary;
+                    bugg.JiraComponentsText = ComponentsText;
+                    bugg.JiraResolution = Resolution;
+                    bugg.JiraFixVersionsText = FixVersionsText;
+                    if (FixCL != 0)
+                    {
+                        bugg.JiraFixCL = FixCL.ToString();
+                    }
+                }
+            }
 
             return new ReportsViewModel
             {
diff --git a/Engine/Source/Programs/CrashReporter/CrashReportWebSite/CrashReportWebSite.csproj b/Engine/Source/Programs/CrashReporter/CrashReportWebSite/CrashReportWebSite.csproj
--- a/Engine/Source/Programs/CrashReporter/CrashReportWebSite/CrashReportWebSite.csproj
+++ b/Engine/Source/Programs/CrashReporter/CrashReportWebSite/CrashReportWebSite.csproj
@@ -46,6 +46,10 @@
     <Prefer32Bit>false</Prefer32Bit>
   </PropertyGroup>
   <ItemGroup>
+    <Reference Include="AWSSDK, Version=2.3.50.1, Culture=neutral, PublicKeyToken=9f476d3089b52be3, processorArchitecture=MSIL">
+      <SpecificVersion>False</SpecificVersion>
+      <HintPath>..\..\..\..\Binaries\ThirdParty\AWSSDK\AWSSDK.dll</HintPath>
+    </Reference>
     <Reference Include="Common.Logging, Version=3.3.1.0, Culture=neutral, PublicKeyToken=af08829b84f0328e, processorArchitecture=MSIL">
       <HintPath>packages\Common.Logging.3.3.1\lib\net40\Common.Logging.dll</HintPath>
       <Private>True</Private>
@@ -124,6 +128,7 @@
     <Reference Include="System.ComponentModel.DataAnnotations">
       <RequiredTargetFramework>3.5</RequiredTargetFramework>
     </Reference>
+    <Reference Include="System.Runtime.Caching" />
     <Reference Include="System.Runtime.Serialization" />
     <Reference Include="System.Security" />
     <Reference Include="System.Web.ApplicationServices" />
@@ -180,6 +185,7 @@
     </Compile>
     <Compile Include="App_Start\UnityConfig.cs" />
     <Compile Include="App_Start\UnityMvcActivator.cs" />
+    <Compile Include="Classes\PerformanceTracker.cs" />
     <Compile Include="Controllers\CSVController.cs" />
     <Compile Include="Controllers\DashboardController.cs" />
     <Compile Include="Controllers\ErrorController.cs" />
@@ -345,6 +351,7 @@
       <DependentUpon>CrashReports.edmx</DependentUpon>
     </Content>
     <None Include="Properties\PublishProfiles\CRB.pubxml" />
+    <None Include="Properties\PublishProfiles\AWS.pubxml" />
     <None Include="Scripts\jquery-3.1.1.intellisense.js" />
     <Content Include="Scripts\jquery-3.1.1.js" />
     <Content Include="Scripts\jquery-3.1.1.min.js" />
@@ -381,6 +388,7 @@
     </Content>
     <Content Include="Web.Release.config">
       <DependentUpon>Web.config</DependentUpon>
+      <SubType>Designer</SubType>
     </Content>
     <Content Include="Content\Site.css" />
     <None Include="Views\Shared\Error.cshtml" />
@@ -403,6 +411,7 @@
     <Content Include="Views\Shared\_CRWLayout.cshtml" />
     <None Include="Web.Development.config">
       <DependentUpon>Web.config</DependentUpon>
+      <SubType>Designer</SubType>
     </None>
   </ItemGroup>
   <ItemGroup>
@@ -415,9 +424,7 @@
       <Name>CrashReportCommon</Name>
     </ProjectReference>
   </ItemGroup>
-  <ItemGroup>
-    <Folder Include="Classes\" />
-  </ItemGroup>
+  <ItemGroup />
   <PropertyGroup>
     <VisualStudioVersion Condition="'$(VisualStudioVersion)' == ''">10.0</VisualStudioVersion>
     <VSToolsPath Condition="'$(VSToolsPath)' == ''">$(MSBuildExtensionsPath32)\Microsoft\VisualStudio\v$(VisualStudioVersion)</VSToolsPath>
diff --git a/Engine/Source/Programs/CrashReporter/CrashReportWebSite/DataModels/BuggDerived.cs b/Engine/Source/Programs/CrashReporter/CrashReportWebSite/DataModels/BuggDerived.cs
--- a/Engine/Source/Programs/CrashReporter/CrashReportWebSite/DataModels/BuggDerived.cs
+++ b/Engine/Source/Programs/CrashReporter/CrashReportWebSite/DataModels/BuggDerived.cs
@@ -10,6 +10,7 @@ using Newtonsoft.Json;
 using Newtonsoft.Json.Linq;
 using Tools.CrashReporter.CrashReportCommon;
 using Tools.CrashReporter.CrashReportWebSite.Properties;
+using Tools.CrashReporter.CrashReportWebSite.ViewModels;
 
 namespace Tools.CrashReporter.CrashReportWebSite.DataModels
 {
@@ -102,7 +103,7 @@ namespace Tools.CrashReporter.CrashReportWebSite.DataModels
         /// Prepares Bugg for JIRA
         /// </summary>
         /// <param name="CrashesForBugg"></param>
-        public void PrepareBuggForJira(List<Crash> CrashesForBugg)
+        public void PrepareBuggForJira(List<CrashDataModel> CrashesForBugg)
         {
             var jiraConnection = JiraConnection.Get();
 
@@ -115,39 +116,39 @@ namespace Tools.CrashReporter.CrashReportWebSite.DataModels
             var machineIds = new HashSet<string>();
             var firstClAffected = int.MaxValue;
 
-            foreach (var Crash in CrashesForBugg)
+            foreach (var crashDataModel in CrashesForBugg)
             {
                 // Only add machine if the number has 32 characters
-                if (Crash.ComputerName != null && Crash.ComputerName.Length == 32)
+                if (crashDataModel.ComputerName != null && crashDataModel.ComputerName.Length == 32)
                 {
-                    machineIds.Add(Crash.ComputerName);
-                    if (Crash.Description.Length > 4)
+                    machineIds.Add(crashDataModel.ComputerName);
+                    if (crashDataModel.Description.Length > 4)
                     {
-                        hashSetDescriptions.Add(Crash.Description);
+                        hashSetDescriptions.Add(crashDataModel.Description);
                     }
                 }
 
-                if (!string.IsNullOrEmpty(Crash.BuildVersion))
+                if (!string.IsNullOrEmpty(crashDataModel.BuildVersion))
                 {
-                    this.AffectedVersions.Add(Crash.BuildVersion);
+                    this.AffectedVersions.Add(crashDataModel.BuildVersion);
                 }
                 // Depot || Stream
-                if (!string.IsNullOrEmpty(Crash.Branch))
+                if (!string.IsNullOrEmpty(crashDataModel.Branch))
                 {
-                    this.BranchesFoundIn.Add(Crash.Branch);
+                    this.BranchesFoundIn.Add(crashDataModel.Branch);
                 }
 
                 var CrashBuiltFromCl = 0;
-                int.TryParse(Crash.ChangelistVersion, out CrashBuiltFromCl);
+                int.TryParse(crashDataModel.ChangelistVersion, out CrashBuiltFromCl);
                 if (CrashBuiltFromCl > 0)
                 {
                     firstClAffected = Math.Min(firstClAffected, CrashBuiltFromCl);
                 }
 
-                if (!string.IsNullOrEmpty(Crash.PlatformName))
+                if (!string.IsNullOrEmpty(crashDataModel.PlatformName))
                 {
                     // Platform = "Platform [Desc]";
-                    var platSubs = Crash.PlatformName.Split(" ".ToCharArray(), StringSplitOptions.RemoveEmptyEntries);
+                    var platSubs = crashDataModel.PlatformName.Split(" ".ToCharArray(), StringSplitOptions.RemoveEmptyEntries);
                     if (platSubs.Length >= 1)
                     {
                         this.AffectedPlatforms.Add(platSubs[0]);
@@ -179,7 +180,6 @@ namespace Tools.CrashReporter.CrashReportWebSite.DataModels
             }
 
             var bv = this.BuildVersion;
-            this.NumberOfUniqueMachines = machineIds.Distinct().Count();// # Affected Users
             var latestClAffected = CrashesForBugg.				// CL of the latest build
                 Where(Crash => Crash.BuildVersion == bv).
                 Max(Crash => Crash.ChangelistVersion);
@@ -192,7 +192,7 @@ namespace Tools.CrashReporter.CrashReportWebSite.DataModels
             this.LatestOSAffected = latestOsAffected;			// Latest Environment Affected
 
             // ToJiraSummary
-            var functionCalls = new CallStackContainer(CrashesForBugg.First()).GetFunctionCallsForJira();
+            var functionCalls = new CallStackContainer(CrashesForBugg.First().CrashType, CrashesForBugg.First().RawCallStack, CrashesForBugg.First().PlatformName).GetFunctionCallsForJira();
             if (functionCalls.Count > 0)
             {
                 this.ToJiraSummary = functionCalls[0];
diff --git a/Engine/Source/Programs/CrashReporter/CrashReportWebSite/DataModels/Crash.cs b/Engine/Source/Programs/CrashReporter/CrashReportWebSite/DataModels/Crash.cs
--- a/Engine/Source/Programs/CrashReporter/CrashReportWebSite/DataModels/Crash.cs
+++ b/Engine/Source/Programs/CrashReporter/CrashReportWebSite/DataModels/Crash.cs
@@ -48,6 +48,7 @@ namespace Tools.CrashReporter.CrashReportWebSite.DataModels
         public Nullable<int> BuggId { get; set; }
         public Nullable<bool> ProcessFailed { get; set; }
         public Nullable<bool> IsVanilla { get; set; }
+        public string UserName { get; set; }
     
         public virtual Bugg Bugg { get; set; }
         public virtual CallStackPattern CallStackPattern { get; set; }
diff --git a/Engine/Source/Programs/CrashReporter/CrashReportWebSite/DataModels/CrashDerived.cs b/Engine/Source/Programs/CrashReporter/CrashReportWebSite/DataModels/CrashDerived.cs
--- a/Engine/Source/Programs/CrashReporter/CrashReportWebSite/DataModels/CrashDerived.cs
+++ b/Engine/Source/Programs/CrashReporter/CrashReportWebSite/DataModels/CrashDerived.cs
@@ -1,5 +1,7 @@
-﻿// Copyright 1998-2017 Epic Games, Inc. All Rights Reserved.
+// Copyright 1998-2017 Epic Games, Inc. All Rights Reserved.
 
+using Amazon.Runtime;
+using Amazon.S3;
 using System;
 using System.Collections.Generic;
 using System.Diagnostics;
@@ -7,6 +9,7 @@ using System.Linq;
 using System.Web;
 using Tools.CrashReporter.CrashReportCommon;
 using Tools.CrashReporter.CrashReportWebSite.Models;
+using Tools.CrashReporter.CrashReportWebSite.Properties;
 
 namespace Tools.CrashReporter.CrashReportWebSite.DataModels
 {
@@ -24,6 +27,25 @@ namespace Tools.CrashReporter.CrashReportWebSite.DataModels
 
         private string SitePath = Properties.Settings.Default.SitePath;
 
+        public static AmazonClient AmazonClient { get; private set; }
+
+        static Crash()
+        {
+            //string AWSError;
+            //AWSCredentials AWSCredentialsForOutput = new StoredProfileAWSCredentials(Settings.Default.AWSS3ProfileName, Settings.Default.AWSCredentialsFilepath);
+            //AmazonS3Config S3ConfigForOutput = new AmazonS3Config
+            //{
+            //    ServiceURL = Settings.Default.AWSS3URL
+            //};
+
+            //AmazonClient = new AmazonClient(AWSCredentialsForOutput, null, S3ConfigForOutput, out AWSError);
+
+            //if (!AmazonClient.IsS3Valid)
+            //{
+            //    System.Diagnostics.Debug.WriteLine("Failed to initailize S3");
+            //}
+        }
+
         /// <summary>If available, will read CrashContext.runtime-xml.</summary>
         public void ReadCrashContextIfAvailable()
         {
@@ -33,7 +55,8 @@ namespace Tools.CrashReporter.CrashReportWebSite.DataModels
                 bool bHasCrashContext = HasCrashContextFile();
                 if (bHasCrashContext)
                 {
-                    _CrashContext = FGenericCrashContext.FromFile(SitePath + GetCrashContextUrl());
+                    //_CrashContext = FGenericCrashContext.FromFile(SitePath + GetCrashContextUrl());
+                    _CrashContext = FGenericCrashContext.FromUrl(GetCrashContextUrl());
                     bool bTest = _CrashContext != null && !string.IsNullOrEmpty(_CrashContext.PrimaryCrashProperties.FullCrashDumpLocation);
                     if (bTest)
                     {
@@ -72,8 +95,8 @@ namespace Tools.CrashReporter.CrashReportWebSite.DataModels
         /// <summary>Return true, if there is a Crash context file associated with the Crash</summary>
         public bool HasCrashContextFile()
         {
-            var Path = SitePath + GetCrashContextUrl();
-            return System.IO.File.Exists(Path);
+            var Path = GetCrashContextUrl();
+            return !String.IsNullOrEmpty(Path);
         }
 
         /// <summary>Whether to use the full minidump.</summary>
@@ -85,8 +108,8 @@ namespace Tools.CrashReporter.CrashReportWebSite.DataModels
         /// <summary>Return true, if there is a metadata file associated with the Crash</summary>
         public bool HasMetaDataFile()
         {
-            var Path = SitePath + GetMetaDataUrl();
-            return System.IO.File.Exists(Path);
+            var Path = GetMetaDataUrl();
+            return !String.IsNullOrEmpty(Path);
         }
 
         /// <summary>
@@ -96,7 +119,7 @@ namespace Tools.CrashReporter.CrashReportWebSite.DataModels
         /// <returns>A class containing the parsed callstack.</returns>
         public CallStackContainer GetCallStack()
         {
-            return new CallStackContainer(this);
+            return new CallStackContainer(this.CrashType, this.RawCallStack, this.PlatformName);
         }
 
         /// <summary>
@@ -104,7 +127,14 @@ namespace Tools.CrashReporter.CrashReportWebSite.DataModels
         /// </summary>
         public string GetLogUrl()
         {
-            return Properties.Settings.Default.CrashReporterFiles + Id + "_Launch.log";
+            // Check compressed first. Then uncompressed
+            string fileName = GetFilePath("_Launch.log" + Settings.Default.AWSS3CompressedSuffix);
+            if (fileName == null)
+            {
+                GetFilePath("_Launch.log");
+            }
+
+            return fileName;
         }
 
         /// <summary>
@@ -112,9 +142,14 @@ namespace Tools.CrashReporter.CrashReportWebSite.DataModels
         /// </summary>
         public string GetMiniDumpUrl()
         {
-            var WebPath = Properties.Settings.Default.CrashReporterFiles + Id + "_MiniDump.dmp";
-            var Path = UseFullMinidumpPath() ? _CrashContext.PrimaryCrashProperties.FullCrashDumpLocation : WebPath;
-            return Path;
+            // Check compressed first. Then uncompressed
+            string fileName = GetFilePath("_MiniDump.dmp" + Settings.Default.AWSS3CompressedSuffix);
+            if (fileName == null)
+            {
+                GetFilePath("_MiniDump.dmp");
+            }
+
+            return fileName;
         }
 
         /// <summary>
@@ -140,7 +175,7 @@ namespace Tools.CrashReporter.CrashReportWebSite.DataModels
         /// </summary>
         public string GetDiagnosticsUrl()
         {
-            return Properties.Settings.Default.CrashReporterFiles + Id + "_Diagnostics.txt";
+            return GetFilePath("_Diagnostics.txt");
         }
 
         /// <summary>
@@ -148,7 +183,7 @@ namespace Tools.CrashReporter.CrashReportWebSite.DataModels
         /// </summary>
         public string GetCrashContextUrl()
         {
-            return Properties.Settings.Default.CrashReporterFiles + Id + "_CrashContext.runtime-xml";
+            return GetFilePath("_CrashContext.runtime-xml");
         }
 
         /// <summary>
@@ -156,7 +191,7 @@ namespace Tools.CrashReporter.CrashReportWebSite.DataModels
         /// </summary>
         public string GetMetaDataUrl()
         {
-            return Properties.Settings.Default.CrashReporterFiles + Id + "_WERMeta.xml";
+            return GetFilePath("_WERMeta.xml");
         }
 
         /// <summary>
@@ -165,7 +200,7 @@ namespace Tools.CrashReporter.CrashReportWebSite.DataModels
         /// <returns>The Url of the Crash report video file.</returns>
         public string GetVideoUrl()
         {
-            return Properties.Settings.Default.CrashReporterVideos + Id + "_CrashVideo.avi";
+            return GetFilePath("_CrashVideo.avi");
         }
 
         /// <summary>
@@ -226,29 +261,50 @@ namespace Tools.CrashReporter.CrashReportWebSite.DataModels
         /// <summary>Return true, if there is a minidump file associated with the Crash</summary>
         public bool GetHasMiniDumpFile()
         {
-            var Path = SitePath + GetMiniDumpUrl();
-            return UseFullMinidumpPath() || System.IO.File.Exists(Path);
+            var Path = GetMiniDumpUrl();
+            return !String.IsNullOrEmpty(Path);
         }
 
         /// <summary>Return true, if there is a video file associated with the Crash</summary>
         public bool GetHasVideoFile()
         {
-            var Path = SitePath + GetVideoUrl();
-            return System.IO.File.Exists(Path);
+            var Path = GetVideoUrl();
+            return !String.IsNullOrEmpty(Path);
         }
 
         /// <summary>Return true, if there is a log file associated with the Crash</summary>
         public bool GetHasLogFile()
         {
-            var Path = SitePath + GetLogUrl();
-            return System.IO.File.Exists(Path);
+            var Path = GetLogUrl();
+
+            return !String.IsNullOrEmpty(Path);
         }
 
         /// <summary>Return true, if there is a diagnostics file associated with the Crash</summary>
         public bool GetHasDiagnosticsFile()
         {
-            var Path = SitePath + GetDiagnosticsUrl();
-            return System.IO.File.Exists(Path);
+            var Path = GetDiagnosticsUrl();
+            return !String.IsNullOrEmpty(Path);
+        }
+        
+        private string GetFilePath(String FileName)
+        {
+            string Path = null;
+
+            if (Settings.Default.DownloadFromS3)
+            {
+                int ReportIDSegment = (Id / 10000) * 10000;
+
+                string reportPath = string.Format("{0}/{1}/{2}/{2}{3}", Settings.Default.AWSS3KeyPrefix, ReportIDSegment, Id, FileName);
+
+                Path = AmazonClient.GetS3SignedUrl("crashreporter", reportPath, 60);
+            }
+            else
+            {
+                Path = Properties.Settings.Default.CrashReporterFiles + Id + FileName;
+            }
+
+            return Path;
         }
     }
 }
\ No newline at end of file
diff --git a/Engine/Source/Programs/CrashReporter/CrashReportWebSite/DataModels/Repositories/BuggRespository.cs b/Engine/Source/Programs/CrashReporter/CrashReportWebSite/DataModels/Repositories/BuggRespository.cs
--- a/Engine/Source/Programs/CrashReporter/CrashReportWebSite/DataModels/Repositories/BuggRespository.cs
+++ b/Engine/Source/Programs/CrashReporter/CrashReportWebSite/DataModels/Repositories/BuggRespository.cs
@@ -34,6 +34,16 @@ namespace Tools.CrashReporter.CrashReportWebSite.DataModels.Repositories
             return _entityContext.Buggs;
         }
 
+        /// <summary>
+        /// Count the number of objects that satisfy the filter
+        /// </summary>
+        /// <param name="filter"></param>
+        /// <returns></returns>
+        public int Count(Expression<Func<Bugg, bool>> filter)
+        {
+            return _entityContext.Buggs.Count(filter);
+        }
+
         /// <summary>
         /// Get a filtered list of Buggs from data storage
         /// Calling this method returns the data directly. It will execute the data retrieval - in this case an sql transaction.
@@ -84,7 +94,7 @@ namespace Tools.CrashReporter.CrashReportWebSite.DataModels.Repositories
         /// <returns>Bugg data model</returns>
         public Bugg GetById(int id)
         {
-            return _entityContext.Buggs.Include(data => data.Crashes).FirstOrDefault(data => data.Id == id);
+            return _entityContext.Buggs.FirstOrDefault(data => data.Id == id);
         }
 
         /// <summary>
diff --git a/Engine/Source/Programs/CrashReporter/CrashReportWebSite/DataModels/Repositories/CallStackRepository.cs b/Engine/Source/Programs/CrashReporter/CrashReportWebSite/DataModels/Repositories/CallStackRepository.cs
--- a/Engine/Source/Programs/CrashReporter/CrashReportWebSite/DataModels/Repositories/CallStackRepository.cs
+++ b/Engine/Source/Programs/CrashReporter/CrashReportWebSite/DataModels/Repositories/CallStackRepository.cs
@@ -34,6 +34,16 @@ namespace Tools.CrashReporter.CrashReportWebSite.DataModels.Repositories
             return _entityContext.CallStackPatterns;
         }
 
+        /// <summary>
+        /// Count the number of objects that satisfy the filter
+        /// </summary>
+        /// <param name="filter"></param>
+        /// <returns></returns>
+        public int Count(Expression<Func<CallStackPattern, bool>> filter)
+        {
+            return _entityContext.CallStackPatterns.Count(filter);
+        }
+
         /// <summary>
         /// Get a filtered list of CallStackPatterns from data storage
         /// Calling this method returns the data directly. It will execute the data retrieval - in this case an sql transaction.
diff --git a/Engine/Source/Programs/CrashReporter/CrashReportWebSite/DataModels/Repositories/CrashRepository.cs b/Engine/Source/Programs/CrashReporter/CrashReportWebSite/DataModels/Repositories/CrashRepository.cs
--- a/Engine/Source/Programs/CrashReporter/CrashReportWebSite/DataModels/Repositories/CrashRepository.cs
+++ b/Engine/Source/Programs/CrashReporter/CrashReportWebSite/DataModels/Repositories/CrashRepository.cs
@@ -3,7 +3,10 @@ using System.Collections.Generic;
 using System.Data.Entity;
 using System.Linq;
 using System.Linq.Expressions;
+using System.Runtime.Caching;
 using System.Web.Mvc;
+using Tools.CrashReporter.CrashReportWebSite.Properties;
+using System.Runtime.Caching;
 
 namespace Tools.CrashReporter.CrashReportWebSite.DataModels.Repositories
 {
@@ -13,6 +16,8 @@ namespace Tools.CrashReporter.CrashReportWebSite.DataModels.Repositories
     public class CrashRepository : ICrashRepository
     {
         private CrashReportEntities _entityContext;
+        
+        private static MemoryCache cache = MemoryCache.Default;
 
         /// <summary>
         /// Constructor
@@ -35,6 +40,16 @@ namespace Tools.CrashReporter.CrashReportWebSite.DataModels.Repositories
             return _entityContext.Crashes;
         }
 
+        /// <summary>
+        /// Count the number of objects that satisfy the filter
+        /// </summary>
+        /// <param name="filter"></param>
+        /// <returns></returns>
+        public int Count(Expression<Func<Crash, bool>> filter)
+        {
+            return _entityContext.Crashes.Count(filter);
+        }
+
         /// <summary>
         /// Get a filtered list of Crashes from data storage
         /// Calling this method returns the data directly. It will execute the data retrieval - in this case an sql transaction.
@@ -158,28 +173,36 @@ namespace Tools.CrashReporter.CrashReportWebSite.DataModels.Repositories
         /// <returns></returns>
         public List<SelectListItem> GetBranchesAsListItems()
         {
-            DateTime Now = DateTime.UtcNow;
-            var date = DateTime.Now.AddDays(-14);
-            if (Now - LastBranchesDate > TimeSpan.FromHours(1))
-            {
-                var BranchList = _entityContext.Crashes
-                    .Where(n => n.TimeOfCrash > date)
-                    .Where(c => c.CrashType != 3) // Ignore ensures
-                    // Depot - //depot/UE4* || Stream //UE4, //Something etc.
-                    .Where(n => n.Branch.StartsWith("UE4") || n.Branch.StartsWith("//"))
-                    .Select(n => n.Branch)
-                    .Distinct()
-                    .ToList();
-
-                BranchesAsSelectList = BranchList
-                    .Select(listitem => new SelectListItem {Selected = false, Text = listitem, Value = listitem})
-                    .ToList();
+            var branchesAsSelectList = cache["branches"] as List<SelectListItem>;
 
-                BranchesAsSelectList.Insert(0, new SelectListItem {Selected = true, Text = "", Value = ""});
-
-                LastBranchesDate = Now;
+            if (branchesAsSelectList == null)
+            {
+                using (FAutoScopedLogTimer LogTimer = new FAutoScopedLogTimer("CrashRepository.GetBranches"))
+                {
+                    var date = DateTime.Now.AddDays(-14);
+
+                    var BranchList = _entityContext.Crashes
+                        .Where(n => n.TimeOfCrash > date)
+                        .Where(c => c.CrashType != 3) // Ignore ensures
+                        // Depot - //depot/UE4* || Stream //UE4, //Something etc.
+                        .Where(n => n.Branch.StartsWith("UE4") || n.Branch.StartsWith("//"))
+                        .Select(n => n.Branch)
+                        .Distinct()
+                        .ToList();
+
+                    branchesAsSelectList =
+                        BranchList
+                        .Select(listItem => new SelectListItem { Selected = false, Text = listItem, Value = listItem })
+                        .ToList();
+
+                    branchesAsSelectList.Insert(0, new SelectListItem { Selected = true, Text = "", Value = "" });
+                }
+                
+                var searchFilterCachePolicy = new CacheItemPolicy() { AbsoluteExpiration = DateTimeOffset.Now.AddMinutes(Settings.Default.SearchFilterCacheInMinutes) };
+                cache.Set("branches", branchesAsSelectList, searchFilterCachePolicy);
             }
-            return BranchesAsSelectList;
+
+            return branchesAsSelectList;
         }
 
         /// <summary>
@@ -215,18 +238,17 @@ namespace Tools.CrashReporter.CrashReportWebSite.DataModels.Repositories
             return engineModesAsListItems;
         }
 
-        private static DateTime LastVersionDate = DateTime.UtcNow.AddDays(-1);
-        private static List<SelectListItem> VersionsAsSelectList = null;
         private static HashSet<string> DistinctBuildVersions = null;
         /// <summary>
         /// Retrieves a list of distinct UE4 Versions from the CrashRepository
         /// </summary>
         public List<SelectListItem> GetVersionsAsListItems()
         {
-            DateTime Now = DateTime.UtcNow;
-            var date = DateTime.Now.AddDays(-14);
-            if (Now - LastVersionDate > TimeSpan.FromHours(1))
+            List<SelectListItem> versionsAsSelectList = cache["versions"] as List<SelectListItem>;
+
+            if (versionsAsSelectList == null)
             {
+                var date = DateTime.Now.AddDays(-14);
                 var BuildVersions = _entityContext.Crashes
                 .Where(c => c.TimeOfCrash > date)
                 .Where(c => c.CrashType != 3) // Ignore ensures
@@ -246,26 +268,27 @@ namespace Tools.CrashReporter.CrashReportWebSite.DataModels.Repositories
                     }
                 }
 
-                VersionsAsSelectList = DistinctBuildVersions
+                versionsAsSelectList = DistinctBuildVersions
                     .Select(listitem => new SelectListItem { Selected = false, Text = listitem, Value = listitem })
                     .ToList();
-                VersionsAsSelectList.Insert(0, new SelectListItem { Selected = true, Text = "", Value = "" });
+                versionsAsSelectList.Insert(0, new SelectListItem { Selected = true, Text = "", Value = "" });
 
-                LastVersionDate = Now;
+                var searchFilterCachePolicy = new CacheItemPolicy() { AbsoluteExpiration = DateTimeOffset.Now.AddMinutes(Settings.Default.SearchFilterCacheInMinutes) };
+                cache.Set("versions", versionsAsSelectList, searchFilterCachePolicy);
             }
 
-            return VersionsAsSelectList;
+            return versionsAsSelectList;
         }
 
         private static DateTime LastEngineVersionDate = DateTime.UtcNow.AddDays(-1);
-        private static List<SelectListItem> EngineVersionsAsSelectList = null;
         private static HashSet<string> DistinctEngineVersions = null;
         public List<SelectListItem> GetEngineVersionsAsListItems()
         {
-            DateTime Now = DateTime.UtcNow;
-            var date = DateTime.Now.AddDays(-14);
-            if (Now - LastEngineVersionDate > TimeSpan.FromHours(1))
+            List<SelectListItem> engineVersionsAsSelectList = cache["engineVersions"] as List<SelectListItem>;
+
+            if (engineVersionsAsSelectList == null)
             {
+                var date = DateTime.Now.AddDays(-14);
                 var engineVersions = _entityContext.Crashes
                 .Where(c => c.TimeOfCrash > date)
                 .Where(c => c.CrashType != 3) // Ignore ensures
@@ -273,7 +296,7 @@ namespace Tools.CrashReporter.CrashReportWebSite.DataModels.Repositories
                 .Distinct()
                 .ToList();
 
-                DistinctEngineVersions = new HashSet<string>();
+                var DistinctEngineVersions = new HashSet<string>();
 
                 foreach (var engineVersion in engineVersions)
                 {
@@ -288,15 +311,16 @@ namespace Tools.CrashReporter.CrashReportWebSite.DataModels.Repositories
                     }
                 }
 
-                EngineVersionsAsSelectList = DistinctEngineVersions
+                engineVersionsAsSelectList = DistinctEngineVersions
                     .Select(listitem => new SelectListItem { Selected = false, Text = listitem, Value = listitem })
                     .ToList();
-                EngineVersionsAsSelectList.Insert(0, new SelectListItem { Selected = true, Text = "", Value = "" });
+                engineVersionsAsSelectList.Insert(0, new SelectListItem { Selected = true, Text = "", Value = "" });
 
-                LastEngineVersionDate = Now;
+                var searchFilterCachePolicy = new CacheItemPolicy() { AbsoluteExpiration = DateTimeOffset.Now.AddMinutes(Settings.Default.SearchFilterCacheInMinutes) };
+                cache.Set("engineVersions", engineVersionsAsSelectList, searchFilterCachePolicy);
             }
 
-            return EngineVersionsAsSelectList;
+            return engineVersionsAsSelectList;
         } 
 
 
@@ -308,5 +332,20 @@ namespace Tools.CrashReporter.CrashReportWebSite.DataModels.Repositories
             GetVersionsAsListItems();
             return DistinctBuildVersions;
         }
+
+        public void SetStatusByBuggId(int buggId, string status)
+        {
+            _entityContext.Crashes.SqlQuery("UPDATE CRASHES SET Status = '" + status + "' * WHERE Crashes.BuggId = '" + buggId.ToString("N") + "'");
+        }
+
+        public void SetFixedCLByBuggId(int buggId, string fixedCl)
+        {
+            _entityContext.Crashes.SqlQuery("UPDATE CRASHES SET FixedChangeList = '" + fixedCl + "' * WHERE Crashes.BuggId = '" + buggId.ToString("N") + "'");
+        }
+
+        public void SetJiraByBuggId(int buggId, string jira)
+        {
+            _entityContext.Crashes.SqlQuery("UPDATE CRASHES SET Jira = '" + jira + "' * WHERE Crashes.BuggId = '" + buggId.ToString("N") + "'");
+        }
     }
-}
\ No newline at end of file
+}
diff --git a/Engine/Source/Programs/CrashReporter/CrashReportWebSite/DataModels/Repositories/DataRepository.cs b/Engine/Source/Programs/CrashReporter/CrashReportWebSite/DataModels/Repositories/DataRepository.cs
--- a/Engine/Source/Programs/CrashReporter/CrashReportWebSite/DataModels/Repositories/DataRepository.cs
+++ b/Engine/Source/Programs/CrashReporter/CrashReportWebSite/DataModels/Repositories/DataRepository.cs
@@ -30,7 +30,16 @@ namespace Tools.CrashReporter.CrashReportWebSite.DataModels.Repositories
         {
             return _entityContext.Set<T>();
         }
-
+        
+        /// <summary>
+        /// Count the number of objects that satisfy the filter
+        /// </summary>
+        /// <param name="filter"></param>
+        /// <returns></returns>
+        public int Count(Expression<Func<T, bool>> filter)
+        {
+            return _entityContext.Set<T>().Count(filter);
+        }
         /// <summary>
         /// Get a filtered list of type T from data storage
         /// Calling this method returns the data directly. It will execute the data retrieval - in this case an sql transaction.
diff --git a/Engine/Source/Programs/CrashReporter/CrashReportWebSite/DataModels/Repositories/ErrorMessageRepository.cs b/Engine/Source/Programs/CrashReporter/CrashReportWebSite/DataModels/Repositories/ErrorMessageRepository.cs
--- a/Engine/Source/Programs/CrashReporter/CrashReportWebSite/DataModels/Repositories/ErrorMessageRepository.cs
+++ b/Engine/Source/Programs/CrashReporter/CrashReportWebSite/DataModels/Repositories/ErrorMessageRepository.cs
@@ -31,6 +31,16 @@ namespace Tools.CrashReporter.CrashReportWebSite.DataModels.Repositories
             return _entityContext.ErrorMessages;
         }
 
+        /// <summary>
+        /// Count the number of objects that satisfy the filter
+        /// </summary>
+        /// <param name="filter"></param>
+        /// <returns></returns>
+        public int Count(Expression<Func<ErrorMessage, bool>> filter)
+        {
+            return _entityContext.ErrorMessages.Count(filter);
+        }
+
         /// <summary>
         /// Get a filtered list of ErrorMessages from data storage
         /// Calling this method returns the data directly. It will execute the data retrieval - in this case an sql transaction.
diff --git a/Engine/Source/Programs/CrashReporter/CrashReportWebSite/DataModels/Repositories/FunctionRepository.cs b/Engine/Source/Programs/CrashReporter/CrashReportWebSite/DataModels/Repositories/FunctionRepository.cs
--- a/Engine/Source/Programs/CrashReporter/CrashReportWebSite/DataModels/Repositories/FunctionRepository.cs
+++ b/Engine/Source/Programs/CrashReporter/CrashReportWebSite/DataModels/Repositories/FunctionRepository.cs
@@ -34,6 +34,16 @@ namespace Tools.CrashReporter.CrashReportWebSite.DataModels.Repositories
             return _entityContext.FunctionCalls;
         }
 
+        /// <summary>
+        /// Count the number of objects that satisfy the filter
+        /// </summary>
+        /// <param name="filter"></param>
+        /// <returns></returns>
+        public int Count(Expression<Func<FunctionCall, bool>> filter)
+        {
+            return _entityContext.FunctionCalls.Count(filter);
+        }
+
         /// <summary>
         /// Get a filtered list of Buggs from data storage
         /// Calling this method returns the data directly. It will execute the data retrieval - in this case an sql transaction.
@@ -107,6 +117,16 @@ namespace Tools.CrashReporter.CrashReportWebSite.DataModels.Repositories
             return _entityContext.FunctionCalls.FirstOrDefault(filter);
         }
 
+        /// <summary>
+        /// Get the first function calls matching a specific filter.
+        /// </summary>
+        /// <param name="filter"></param>
+        /// <returns></returns>
+        public int FirstId(Expression<Func<FunctionCall, bool>> filter)
+        {
+            return _entityContext.FunctionCalls.Where(filter).Select(data => data.Id).First();
+        }
+
         /// <summary>
         /// Add a new Bugg to the data store
         /// </summary>
diff --git a/Engine/Source/Programs/CrashReporter/CrashReportWebSite/DataModels/Repositories/ICrashRepository.cs b/Engine/Source/Programs/CrashReporter/CrashReportWebSite/DataModels/Repositories/ICrashRepository.cs
--- a/Engine/Source/Programs/CrashReporter/CrashReportWebSite/DataModels/Repositories/ICrashRepository.cs
+++ b/Engine/Source/Programs/CrashReporter/CrashReportWebSite/DataModels/Repositories/ICrashRepository.cs
@@ -15,5 +15,8 @@ namespace Tools.CrashReporter.CrashReportWebSite.DataModels.Repositories
         List<SelectListItem> GetEngineModesAsListItems();
         List<SelectListItem> GetVersionsAsListItems();
         List<SelectListItem> GetEngineVersionsAsListItems();
+        void SetStatusByBuggId(int buggId, string status);
+        void SetFixedCLByBuggId(int buggId, string fixedCl);
+        void SetJiraByBuggId(int buggId, string jira);
     }
 }
diff --git a/Engine/Source/Programs/CrashReporter/CrashReportWebSite/DataModels/Repositories/IDataRepository.cs b/Engine/Source/Programs/CrashReporter/CrashReportWebSite/DataModels/Repositories/IDataRepository.cs
--- a/Engine/Source/Programs/CrashReporter/CrashReportWebSite/DataModels/Repositories/IDataRepository.cs
+++ b/Engine/Source/Programs/CrashReporter/CrashReportWebSite/DataModels/Repositories/IDataRepository.cs
@@ -18,6 +18,13 @@ namespace Tools.CrashReporter.CrashReportWebSite.DataModels.Repositories
         /// </summary>
         IQueryable<T> ListAll();
 
+        /// <summary>
+        /// Return the number of objects that satisfy the query filter
+        /// </summary>
+        /// <param name="filter"></param>
+        /// <returns></returns>
+        int Count(Expression<Func<T, bool>> filter);
+
         /// <summary>
         /// Return a list of data objects of type T matching our search expression
         /// </summary>
diff --git a/Engine/Source/Programs/CrashReporter/CrashReportWebSite/DataModels/Repositories/IFunctionRepository.cs b/Engine/Source/Programs/CrashReporter/CrashReportWebSite/DataModels/Repositories/IFunctionRepository.cs
--- a/Engine/Source/Programs/CrashReporter/CrashReportWebSite/DataModels/Repositories/IFunctionRepository.cs
+++ b/Engine/Source/Programs/CrashReporter/CrashReportWebSite/DataModels/Repositories/IFunctionRepository.cs
@@ -1,8 +1,12 @@
 ﻿// Copyright 1998-2017 Epic Games, Inc. All Rights Reserved.
 
+using System;
+using System.Linq.Expressions;
+
 namespace Tools.CrashReporter.CrashReportWebSite.DataModels.Repositories
 {
     public interface IFunctionRepository : IDataRepository<FunctionCall>
     {
+        int FirstId(Expression<Func<FunctionCall, bool>> filter);
     }
 }
diff --git a/Engine/Source/Programs/CrashReporter/CrashReportWebSite/DataModels/Repositories/UserGroupRepository.cs b/Engine/Source/Programs/CrashReporter/CrashReportWebSite/DataModels/Repositories/UserGroupRepository.cs
--- a/Engine/Source/Programs/CrashReporter/CrashReportWebSite/DataModels/Repositories/UserGroupRepository.cs
+++ b/Engine/Source/Programs/CrashReporter/CrashReportWebSite/DataModels/Repositories/UserGroupRepository.cs
@@ -31,6 +31,16 @@ namespace Tools.CrashReporter.CrashReportWebSite.DataModels.Repositories
             return _entityContext.UserGroups;
         }
 
+        /// <summary>
+        /// Count the number of objects that satisfy the filter
+        /// </summary>
+        /// <param name="filter"></param>
+        /// <returns></returns>
+        public int Count(Expression<Func<UserGroup, bool>> filter)
+        {
+            return _entityContext.UserGroups.Count(filter);
+        }
+
         /// <summary>
         /// Get a filtered list of UserGroups from data storage
         /// Calling this method returns the data directly. It will execute the data retrieval - in this case an sql transaction.
diff --git a/Engine/Source/Programs/CrashReporter/CrashReportWebSite/DataModels/Repositories/UserRepository.cs b/Engine/Source/Programs/CrashReporter/CrashReportWebSite/DataModels/Repositories/UserRepository.cs
--- a/Engine/Source/Programs/CrashReporter/CrashReportWebSite/DataModels/Repositories/UserRepository.cs
+++ b/Engine/Source/Programs/CrashReporter/CrashReportWebSite/DataModels/Repositories/UserRepository.cs
@@ -31,6 +31,16 @@ namespace Tools.CrashReporter.CrashReportWebSite.DataModels.Repositories
             return _entityContext.Users;
         }
 
+        /// <summary>
+        /// Count the number of objects that satisfy the filter
+        /// </summary>
+        /// <param name="filter"></param>
+        /// <returns></returns>
+        public int Count(Expression<Func<User, bool>> filter)
+        {
+            return _entityContext.Users.Count(filter);
+        }
+
         /// <summary>
         /// Get a filtered list of Users from data storage
         /// Calling this method returns the data directly. It will execute the data retrieval - in this case an sql transaction.
diff --git a/Engine/Source/Programs/CrashReporter/CrashReportWebSite/Models/CachedDataService.cs b/Engine/Source/Programs/CrashReporter/CrashReportWebSite/Models/CachedDataService.cs
--- a/Engine/Source/Programs/CrashReporter/CrashReportWebSite/Models/CachedDataService.cs
+++ b/Engine/Source/Programs/CrashReporter/CrashReportWebSite/Models/CachedDataService.cs
@@ -6,6 +6,7 @@ using System.Linq;
 using System.Web.Caching;
 using Tools.CrashReporter.CrashReportWebSite.DataModels;
 using Tools.CrashReporter.CrashReportWebSite.DataModels.Repositories;
+using Tools.CrashReporter.CrashReportWebSite.ViewModels;
 
 namespace Tools.CrashReporter.CrashReportWebSite.Models
 {
@@ -70,14 +71,14 @@ namespace Tools.CrashReporter.CrashReportWebSite.Models
 		/// </summary>
 		/// <param name="currentCrash">The Crash to retrieve the parsed callstack for.</param>
 		/// <returns>A parsed callstack.</returns>
-		public DataModels.CallStackContainer GetCallStackFast( Crash currentCrash )
+        public CallStackContainer GetCallStackFast(CrashDataModel currentCrash)
 		{
 			using( var logTimer = new FAutoScopedLogTimer( this.GetType().ToString() + "(CrashId=" + currentCrash.Id + ")" ) )
 			{
 				var key = CacheKeyPrefix + CallstackKeyPrefix + currentCrash.Id;
 				var callStack = (CallStackContainer)CacheInstance[key];
 			    if (callStack != null) return callStack;
-			    callStack = new CallStackContainer( currentCrash );
+			    callStack = new CallStackContainer( currentCrash.CrashType, currentCrash.RawCallStack, currentCrash.PlatformName );
 			    callStack.bDisplayFunctionNames = true;
 			    CacheInstance.Insert( key, callStack );
 
diff --git a/Engine/Source/Programs/CrashReporter/CrashReportWebSite/Models/CrashDataFormater.cs b/Engine/Source/Programs/CrashReporter/CrashReportWebSite/Models/CrashDataFormater.cs
--- a/Engine/Source/Programs/CrashReporter/CrashReportWebSite/Models/CrashDataFormater.cs
+++ b/Engine/Source/Programs/CrashReporter/CrashReportWebSite/Models/CrashDataFormater.cs
@@ -6,6 +6,7 @@ using System.Linq;
 using System.Text;
 using System.Text.RegularExpressions;
 using System.IO;
+using Tools.CrashReporter.CrashReportWebSite.ViewModels;
 
 namespace Tools.CrashReporter.CrashReportWebSite.DataModels
 {
@@ -152,11 +153,11 @@ namespace Tools.CrashReporter.CrashReportWebSite.DataModels
 		/// Construct a processed callstack.
 		/// </summary>
 		/// <param name="currentCrash"></param>
-		public CallStackContainer( Crash currentCrash )
+        public CallStackContainer(int crashType, string rawCallStack, string platformName)
 		{
-			using( FAutoScopedLogTimer LogTimer = new FAutoScopedLogTimer( this.GetType().ToString() + "(CrashId=" + currentCrash.Id + ")" ) )
+			using( FAutoScopedLogTimer LogTimer = new FAutoScopedLogTimer( this.GetType().ToString() + "(CrashId)" ) )
 			{
-				ParseCallStack( currentCrash ); 
+				ParseCallStack( crashType, rawCallStack, platformName ); 
 			}
 		}
 
@@ -237,7 +238,7 @@ namespace Tools.CrashReporter.CrashReportWebSite.DataModels
 		/// Parse a raw callstack into a pattern
 		/// </summary>
 		/// <param name="currentCrash">The Crash with a raw callstack to parse.</param>
-		private void ParseCallStack( Crash currentCrash )
+		private void ParseCallStack(int crashType, string rawCallStack, string platformName)
 		{
 			// Everything is disabled by default
 			bDisplayUnformattedCallStack = false;
@@ -248,7 +249,7 @@ namespace Tools.CrashReporter.CrashReportWebSite.DataModels
 
 			bool bSkipping = false;
 			string LineToSkipUpto = "";
-			switch( currentCrash.CrashType )
+            switch (crashType)
 			{
 			case 2:
 				bSkipping = false;
@@ -260,7 +261,7 @@ namespace Tools.CrashReporter.CrashReportWebSite.DataModels
 				break;
 			}
 
-			if (string.IsNullOrEmpty(currentCrash.RawCallStack))
+            if (string.IsNullOrEmpty(rawCallStack))
 			{
 				return;
 			}
@@ -268,7 +269,7 @@ namespace Tools.CrashReporter.CrashReportWebSite.DataModels
 			CallStackEntries.Clear();
 
 			// Store off a pre split array of call stack lines
-			string[] RawCallStackLines = currentCrash.RawCallStack.Split(new[] { '\n' }, StringSplitOptions.RemoveEmptyEntries);
+            string[] RawCallStackLines = rawCallStack.Split(new[] { '\n' }, StringSplitOptions.RemoveEmptyEntries);
 
 			int Middle = RawCallStackLines.Length / 2;
 
@@ -333,9 +334,9 @@ namespace Tools.CrashReporter.CrashReportWebSite.DataModels
 				int MacModuleStart = CurrentLine.IndexOf( "[in " );
 				int MacModuleEnd = MacModuleStart > 0 ? CurrentLine.IndexOf( "]", MacModuleStart ) : 0;
 
-				bool bLinux = currentCrash.PlatformName.Contains( "Linux" );
-				bool bMac = currentCrash.PlatformName.Contains( "Mac" );
-				bool bWindows = currentCrash.PlatformName.Contains( "Windows" );
+                bool bLinux = platformName.Contains("Linux");
+                bool bMac = platformName.Contains("Mac");
+                bool bWindows = platformName.Contains("Windows");
 
 
 				// Parse out the juicy info from the line of the callstack
