Restore deleted copy of XmlHandler.cs into CrashReportCommon. Fixes build failure.

diff --git a/Engine/Source/Programs/CrashReporter/CrashReportCommon/CrashContext.cs b/Engine/Source/Programs/CrashReporter/CrashReportCommon/CrashContext.cs
--- a/Engine/Source/Programs/CrashReporter/CrashReportCommon/CrashContext.cs
+++ b/Engine/Source/Programs/CrashReporter/CrashReportCommon/CrashContext.cs
@@ -5,7 +5,6 @@ using System.IO;
 using System.Xml.Serialization;
 using System.Collections;
 using System.Collections.Generic;
-using Tools.DotNETCommon.XmlHandler;
 using System.Xml;
 using System.Linq;
 
diff --git a/Engine/Source/Programs/CrashReporter/CrashReportCommon/CrashReportCommon.csproj b/Engine/Source/Programs/CrashReporter/CrashReportCommon/CrashReportCommon.csproj
--- a/Engine/Source/Programs/CrashReporter/CrashReportCommon/CrashReportCommon.csproj
+++ b/Engine/Source/Programs/CrashReporter/CrashReportCommon/CrashReportCommon.csproj
@@ -79,6 +79,7 @@
     <Compile Include="SlackWriter.cs" />
     <Compile Include="Stats.cs" />
     <Compile Include="WerReport.cs" />
+    <Compile Include="XmlHandler.cs" />
   </ItemGroup>
   <ItemGroup>
     <ProjectReference Include="..\..\DotNETCommon\DotNETUtilities\DotNETUtilities.csproj">
diff --git a/Engine/Source/Programs/CrashReporter/CrashReportCommon/XmlHandler.cs b/Engine/Source/Programs/CrashReporter/CrashReportCommon/XmlHandler.cs
new file mode 100644
--- /dev/null
+++ b/Engine/Source/Programs/CrashReporter/CrashReportCommon/XmlHandler.cs
@@ -0,0 +1,225 @@
+﻿using System;
+using System.Collections.Generic;
+using System.Diagnostics;
+using System.IO;
+using System.Linq;
+using System.Text;
+using System.Threading.Tasks;
+using System.Xml;
+using System.Xml.Serialization;
+
+namespace Tools.CrashReporter.CrashReportCommon
+{
+	/// <summary>
+	/// A helper class to handle reading and writing of XML files.
+	/// </summary>
+	public class XmlHandler
+	{
+		/// <summary>
+		/// Get a good set of default xml writer settings.
+		/// </summary>
+		/// <returns>Returns settings that indent with tabs and omit the Xml declaration.</returns>
+		static private XmlWriterSettings GetDefaultWriterSettings()
+		{
+			XmlWriterSettings WriterSettings = new XmlWriterSettings();
+			WriterSettings.CloseOutput = true;
+			WriterSettings.Indent = true;
+			WriterSettings.IndentChars = "\t";
+			WriterSettings.OmitXmlDeclaration = true;
+
+			return WriterSettings;
+		}
+
+		/// <summary>
+		/// Get a good set of default xml reader settings.
+		/// </summary>
+		/// <returns>Returns settings that ignore Xml comments.</returns>
+		static private XmlReaderSettings GetDefaultReaderSettings()
+		{
+			XmlReaderSettings ReaderSettings = new XmlReaderSettings();
+			ReaderSettings.CloseInput = true;
+			ReaderSettings.IgnoreComments = true;
+
+			return ReaderSettings;
+		}
+
+		/// <summary>
+		/// A callback for any unknown attributes. Details are spewed to the debug channel.
+		/// </summary>
+		/// <param name="Sender">Unused.</param>
+		/// <param name="Event">Details about the unknown attribute encountered in the Xml file.</param>
+		static private void XmlSerializer_UnknownAttribute( object Sender, XmlAttributeEventArgs Event )
+		{
+			//Debug.WriteLine( "Unknown XML attribute found '" + Event.Attr.Name + "' at line " + Event.LineNumber + " position " + Event.LinePosition );
+		}
+
+		/// <summary>
+		/// A callback for any unknown elements. Details are spewed to the debug channel.
+		/// </summary>
+		/// <param name="Sender">Unused.</param>
+		/// <param name="Event">Details about the unknown element found in the Xml file.</param>
+		static private void XmlSerializer_UnknownElement( object Sender, XmlElementEventArgs Event )
+		{
+			//Debug.WriteLine( "Unknown XML element found '" + Event.Element.Name + "' at line " + Event.LineNumber + " position " + Event.LinePosition );
+		}
+
+		/// <summary>
+		/// A callback for any unknown nodes. Details are spewed to the debug channel.
+		/// </summary>
+		/// <param name="Sender">Unused.</param>
+		/// <param name="Event">Details about the unknown node found in the Xml file.</param> 
+		static private void XmlSerializer_UnknownNode( object Sender, XmlNodeEventArgs Event )
+		{
+			//Debug.WriteLine( "Unknown XML node found '" + Event.Name + "' at line " + Event.LineNumber + " position " + Event.LinePosition );
+		}
+
+		/// <summary>
+		/// Returns a string of the contents of a class as Xml.
+		/// </summary>
+		/// <typeparam name="T">The type of the data to convert.</typeparam>
+		/// <param name="Data">The data to serialise.</param>
+		/// <returns>A string with an Xml representation of the data.</returns>
+		/// <remarks>The default writer settings are tabbed indents and no Xml declaration.</remarks>
+		static public string ToXmlString<T>( T Data )
+		{
+			XmlSerializer Serialiser = new XmlSerializer( Data.GetType() );
+
+			XmlSerializerNamespaces EmptyNameSpace = new XmlSerializerNamespaces();
+			EmptyNameSpace.Add( "", "http://www.w3.org/2001/XMLSchema" );
+
+			using( StringWriter Writer = new StringWriter() )
+			{
+				XmlWriter XmlStream = XmlWriter.Create( Writer, GetDefaultWriterSettings() );
+				Serialiser.Serialize( XmlStream, Data, EmptyNameSpace );
+				return Writer.ToString();
+			}
+		}
+
+		/// <summary>
+		/// Converts an Xml string to an instantiated class of type T.
+		/// </summary>
+		/// <typeparam name="T">The type of the data we are requesting.</typeparam>
+		/// <param name="XmlString">A string with an Xml representation of the data.</param>
+		/// <returns>An instance of the type populated with the contents of the Xml string.</returns>
+		static public T FromXmlString<T>( string XmlString )
+		{
+			XmlSerializer Serialiser = new XmlSerializer( typeof( T ) );
+			using( StringReader Reader = new StringReader( XmlString ) )
+			{
+				return ( T )Serialiser.Deserialize( Reader );
+			}
+		}
+
+		/// <summary>
+		/// Create and instance of the class, and parse an Xml file into it.
+		/// </summary>
+		/// <param name="FileName">Path of Xml file to read.</param>
+		/// <param name="ReaderSettings">Optional settings to control the reading of the Xml file.</param>
+		/// <typeparam name="T">The type of the class we are attempting to deserialise.</typeparam>
+		/// <returns>An instance of the class with data parsed from the Xml file.</returns>
+		/// <remarks>This code will never return null, but any field not in the Xml file will be the default value. Unknown attributes,
+		/// elements, and nodes will be ignored, but output details to the debug channel. The default reader settings ignore comments.</remarks>
+		static public T ReadXml<T>( string FileName, XmlReaderSettings ReaderSettings = null ) where T : new()
+		{
+			// Create a new default instance of the type
+			T Instance = new T();
+			XmlReader XmlStream = null;
+			try
+			{
+				// Use the default reader settings if none are passed in
+				if( ReaderSettings == null )
+				{
+					ReaderSettings = GetDefaultReaderSettings();
+				}
+
+				// Get the xml data stream to read from
+				XmlStream = XmlReader.Create( FileName, ReaderSettings );
+
+				// Creates an instance of the XmlSerializer class so we can read the settings object
+				XmlSerializer ObjectReader = new XmlSerializer( typeof( T ) );
+
+				// Add our callbacks for a malformed xml file
+				ObjectReader.UnknownNode += new XmlNodeEventHandler( XmlSerializer_UnknownNode );
+				ObjectReader.UnknownElement += new XmlElementEventHandler( XmlSerializer_UnknownElement );
+				ObjectReader.UnknownAttribute += new XmlAttributeEventHandler( XmlSerializer_UnknownAttribute );
+
+				// Create an object from the xml data
+				Instance = ( T )ObjectReader.Deserialize( XmlStream );
+			}
+			catch( Exception Ex )
+			{
+				Debug.WriteLine( Ex.Message );
+			}
+			finally
+			{
+				if( XmlStream != null )
+				{
+					// Done with the file so close it
+					XmlStream.Close();
+				}
+			}
+
+			return Instance;
+		}
+
+		/// <summary>
+		/// Write an Xml file representing the contents of the class.
+		/// </summary>
+		/// <param name="Data">An instance of the class to write.</param>
+		/// <param name="FileName">The path of the Xml file to write.</param>
+		/// <param name="DefaultNameSpace">Optional default namespace to use.</param>
+		/// <param name="WriterSettings">Optional settings to control the writing of the Xml file.</param>
+		/// <typeparam name="T">The type of the class we wish to write to a file.</typeparam>
+		/// <returns>True if the Xml file was successfully written.</returns>
+		/// <remarks>The default writer settings are tabbed indents and no Xml declaration.</remarks>
+		static public bool WriteXml<T>( T Data, string FileName, string DefaultNameSpace = "", XmlWriterSettings WriterSettings = null )
+		{
+			bool bSuccess = true;
+			XmlWriter XmlStream = null;
+			try
+			{
+				// Ensure the file is writable
+				FileInfo Info = new FileInfo( FileName );
+				if( Info.Exists )
+				{
+					Info.IsReadOnly = false;
+				}
+
+				// Set the default namespace
+				XmlSerializerNamespaces EmptyNameSpace = new XmlSerializerNamespaces();
+				EmptyNameSpace.Add( "", DefaultNameSpace );
+
+				// Set the defaults for writing xml files if none is passed in
+				if( WriterSettings == null )
+				{
+					WriterSettings = GetDefaultWriterSettings();
+				}
+
+				XmlStream = XmlWriter.Create( FileName, WriterSettings );
+				XmlSerializer ObjectWriter = new XmlSerializer( typeof( T ) );
+
+				// Add our callbacks for a malformed xml file
+				ObjectWriter.UnknownNode += new XmlNodeEventHandler( XmlSerializer_UnknownNode );
+				ObjectWriter.UnknownElement += new XmlElementEventHandler( XmlSerializer_UnknownElement );
+				ObjectWriter.UnknownAttribute += new XmlAttributeEventHandler( XmlSerializer_UnknownAttribute );
+
+				ObjectWriter.Serialize( XmlStream, Data, EmptyNameSpace );
+			}
+			catch( Exception Ex )
+			{
+				Debug.WriteLine( Ex.Message );
+				bSuccess = false;
+			}
+			finally
+			{
+				if( XmlStream != null )
+				{
+					// Done with the file so close it
+					XmlStream.Close();
+				}
+			}
+
+			return bSuccess;
+		}
+	}
+}
