Fix some more CrashReporter compile errors due to things being removed from DotNETUtilities.

diff --git a/Engine/Source/Programs/CrashReporter/CrashReportCommon/CrashReportCommon.csproj b/Engine/Source/Programs/CrashReporter/CrashReportCommon/CrashReportCommon.csproj
--- a/Engine/Source/Programs/CrashReporter/CrashReportCommon/CrashReportCommon.csproj
+++ b/Engine/Source/Programs/CrashReporter/CrashReportCommon/CrashReportCommon.csproj
@@ -70,23 +70,19 @@
     <Compile Include="CrashContext.cs" />
     <Compile Include="CrashReporterConstants.cs" />
     <Compile Include="JiraConnection.cs" />
+    <Compile Include="LaunchProcess.cs" />
     <Compile Include="LogWriter.cs" />
     <Compile Include="Properties\AssemblyInfo.cs" />
     <Compile Include="..\..\DotNETCommon\MetaData.cs">
       <Link>Properties\MetaData.cs</Link>
     </Compile>
     <Compile Include="ReportData.cs" />
+    <Compile Include="SimpleWebRequest.cs" />
     <Compile Include="SlackWriter.cs" />
     <Compile Include="Stats.cs" />
     <Compile Include="WerReport.cs" />
     <Compile Include="XmlHandler.cs" />
   </ItemGroup>
-  <ItemGroup>
-    <ProjectReference Include="..\..\DotNETCommon\DotNETUtilities\DotNETUtilities.csproj">
-      <Project>{5d7d66e8-8c76-4af9-b3ec-2ef03421d730}</Project>
-      <Name>DotNETUtilities</Name>
-    </ProjectReference>
-  </ItemGroup>
   <Import Project="$(MSBuildToolsPath)\Microsoft.CSharp.targets" />
   <!-- To modify your build process, add your task inside one of the targets below and uncomment it. 
        Other similar extension points exist, see Microsoft.Common.targets.
diff --git a/Engine/Source/Programs/CrashReporter/CrashReportCommon/LaunchProcess.cs b/Engine/Source/Programs/CrashReporter/CrashReportCommon/LaunchProcess.cs
new file mode 100644
--- /dev/null
+++ b/Engine/Source/Programs/CrashReporter/CrashReportCommon/LaunchProcess.cs
@@ -0,0 +1,232 @@
+// Copyright 1998-2017 Epic Games, Inc. All Rights Reserved.
+
+using System;
+using System.Collections.Generic;
+using System.Diagnostics;
+using System.IO;
+using System.Linq;
+using System.Text;
+
+namespace Tools.CrashReporter.CrashReportCommon
+{
+	/// <summary>
+	/// Enum to make the meaning of WaitForExit return code clear
+	/// </summary>
+	public enum EWaitResult
+	{
+		/// <summary>The task completed in a timely manner</summary>
+		Ok,
+
+		/// <summary>The task was abandoned, since it was taking longer than the specified time-out duration</summary>
+		TimedOut
+	}
+
+	/// <summary>
+	/// A class to handle spawning and monitoring of a child process
+	/// </summary>
+	public class LaunchProcess : IDisposable
+	{
+		/// <summary>A callback signature for handling logging.</summary>
+		public delegate void CaptureMessageDelegate( string Message );
+		/// <summary>The current callback for logging, or null for quiet operation.</summary>
+		CaptureMessageDelegate CaptureMessage = null;
+
+		/// <summary>The process that was launched.</summary>
+		private Process LaunchedProcess = null;
+		/// <summary>Set to true when the launched process finishes.</summary>
+		private bool bIsFinished = false;
+
+		/// <summary>
+		/// Implementing Dispose.
+		/// </summary>
+		public void Dispose()
+		{
+			Dispose( true );
+			GC.SuppressFinalize( this );
+		}
+
+		/// <summary>
+		/// Disposes the resources.
+		/// </summary>
+		/// <param name="Disposing"></param>
+		protected virtual void Dispose( bool Disposing )
+		{
+			LaunchedProcess.Dispose();
+		}
+
+		/// <summary>
+		/// The process exit event.
+		/// </summary>
+		/// <param name="Sender">Unused.</param>
+		/// <param name="Args">Unused.</param>
+		private void ProcessExit( object Sender, EventArgs Args )
+		{
+			// Flush and close any pending messages
+			if( CaptureMessage != null )
+			{
+				LaunchedProcess.CancelOutputRead();
+				LaunchedProcess.CancelErrorRead();
+
+				// Protect against multiple calls
+				CaptureMessage = null;
+			}
+
+			LaunchedProcess.EnableRaisingEvents = false;
+			bIsFinished = true;
+		}
+
+		/// <summary>
+		/// Safely invoke the logging callback.
+		/// </summary>
+		/// <param name="Message">The line of text to pass back to the calling process via the delegate.</param>
+		private void PrintLog( string Message )
+		{
+			if( CaptureMessage != null )
+			{
+				CaptureMessage( Message );
+			}
+		}
+
+		/// <summary>
+		/// The event called for StdOut an StdErr redirections.
+		/// </summary>
+		/// <param name="Sender">Unused.</param>
+		/// <param name="Args">The container for the line of text to pass through the system.</param>
+		private void CaptureMessageCallback( object Sender, DataReceivedEventArgs Args )
+		{
+			PrintLog( Args.Data );
+		}
+
+		/// <summary>
+		/// Check to see if the currently running child process has finished.
+		/// </summary>
+		/// <returns>true if the process was successfully spawned and correctly exited. It also returns true if the process failed to spawn.</returns>
+		public bool IsFinished()
+		{
+			if( LaunchedProcess != null )
+			{
+				if( bIsFinished )
+				{
+					return true;
+				}
+
+				// Check for timed out
+				return false;
+			}
+
+			return true;
+		}
+
+		/// <summary>
+		/// Wait for the launched process to exit, and return its exit code.
+		/// </summary>
+		/// <param name="TimeoutMilliseconds">Number of milliseconds to wait for the process to exit. Default is forever.</param>
+		/// <returns>The exit code of the launched process.</returns>
+		/// <remarks>false is returned if the process failed to finish before the requested timeout.</remarks>
+		public EWaitResult WaitForExit( int TimeoutMilliseconds = Int32.MaxValue )
+		{
+			if( LaunchedProcess != null )
+			{
+				LaunchedProcess.WaitForExit(TimeoutMilliseconds);
+				if( !bIsFinished )
+				{
+					// Calling Kill() here seems to be a race condition on the process terminating after WaitForExit (TTP#315685). Catch anything it throws, and make sure the Kill() finishes.
+					try
+					{
+						LaunchedProcess.Kill();
+						LaunchedProcess.WaitForExit();
+					}
+					catch (Exception)
+					{
+					}
+				}
+			}
+
+			return bIsFinished ? EWaitResult.Ok : EWaitResult.TimedOut;
+		}
+
+		/// <summary>
+		/// Construct a class wrapper that spawns a new child process with StdOut and StdErr optionally redirected and captured.
+		/// </summary>
+		/// <param name="Executable">The executable to launch.</param>
+		/// <param name="WorkingDirectory">The working directory of the process. If this is null, the current directory is used.</param>
+		/// <param name="InCaptureMessage">The log callback function. This can be null for no logging.</param>
+		/// <param name="Parameters">A string array of parameters passed on the command line, and delimited with spaces.</param>
+		/// <remarks>Any errors are passed back through the capture delegate. The existence of the executable and the working directory are checked before spawning is attempted.
+		/// An exit code of -1 is returned if there was an exception when spawning the process.</remarks>
+		public LaunchProcess( string Executable, string WorkingDirectory, CaptureMessageDelegate InCaptureMessage, params string[] Parameters )
+		{
+			CaptureMessage = InCaptureMessage;
+
+			// Simple check to ensure the executable exists
+			FileInfo Info = new FileInfo( Executable );
+			if( !Info.Exists )
+			{
+				PrintLog( "ERROR: Executable does not exist: " + Executable );
+				bIsFinished = true;
+				return;
+			}
+
+			// Set the default working directory if necessary
+			if( WorkingDirectory == null )
+			{
+				WorkingDirectory = Environment.CurrentDirectory;
+			}
+
+			// Simple check to ensure the working directory exists
+			DirectoryInfo DirInfo = new DirectoryInfo( WorkingDirectory );
+			if( !DirInfo.Exists )
+			{
+				PrintLog( "ERROR: Working directory does not exist: " + WorkingDirectory );
+				bIsFinished = true;
+				return;
+			}
+
+			// Create a new process to launch
+			LaunchedProcess = new Process();
+
+			// Prepare a ProcessStart structure 
+			LaunchedProcess.StartInfo.FileName = Info.FullName;
+			LaunchedProcess.StartInfo.Arguments = String.Join( " ", Parameters );
+			LaunchedProcess.StartInfo.WorkingDirectory = DirInfo.FullName;
+			LaunchedProcess.StartInfo.CreateNoWindow = true;
+
+			// Need this for the Exited event as well as the output capturing
+			LaunchedProcess.EnableRaisingEvents = true;
+			LaunchedProcess.Exited += new EventHandler(ProcessExit);
+
+			// Redirect the output.
+			if (CaptureMessage != null)
+			{
+				LaunchedProcess.StartInfo.UseShellExecute = false;
+				LaunchedProcess.StartInfo.RedirectStandardOutput = true;
+				LaunchedProcess.StartInfo.RedirectStandardError = true;
+
+				LaunchedProcess.OutputDataReceived += new DataReceivedEventHandler(CaptureMessageCallback);
+				LaunchedProcess.ErrorDataReceived += new DataReceivedEventHandler(CaptureMessageCallback);
+			}
+
+			// Spawn the process - try to start the process, handling thrown exceptions as a failure.
+			try
+			{
+				PrintLog( "Launching: " + LaunchedProcess.StartInfo.FileName + " " + LaunchedProcess.StartInfo.Arguments + " (CWD: " + LaunchedProcess.StartInfo.WorkingDirectory + ")" );
+
+				LaunchedProcess.Start();
+
+				// Start the output redirection if we have a logging callback
+				if( CaptureMessage != null )
+				{
+					LaunchedProcess.BeginOutputReadLine();
+					LaunchedProcess.BeginErrorReadLine();
+				}
+			}
+			catch( Exception Ex )
+			{
+				// Clean up should there be any exception
+				LaunchedProcess = null;
+				bIsFinished = true;
+				PrintLog( "ERROR: Failed to launch with exception: " + Ex.Message );
+			}
+		}
+	}
+}
diff --git a/Engine/Source/Programs/CrashReporter/CrashReportCommon/SimpleWebRequest.cs b/Engine/Source/Programs/CrashReporter/CrashReportCommon/SimpleWebRequest.cs
new file mode 100644
--- /dev/null
+++ b/Engine/Source/Programs/CrashReporter/CrashReportCommon/SimpleWebRequest.cs
@@ -0,0 +1,128 @@
+﻿// Copyright 1998-2017 Epic Games, Inc. All Rights Reserved.
+
+using System;
+using System.Diagnostics;
+using System.IO;
+using System.Net;
+using System.Text;
+
+namespace Tools.DotNETCommon.SimpleWebRequest
+{
+	/// <summary>
+	/// A class to make simple Xml web requests to return a string.
+	/// </summary>
+	public class SimpleWebRequest
+	{
+		/// <summary>
+		/// Default constructor.
+		/// </summary>
+		public SimpleWebRequest()
+		{
+		}
+
+		/// <summary>
+		/// Create an HttpWebRequest with an optional Xml payload.
+		/// </summary>
+		/// <param name="WebRequestString">The URL and parameters of the web request.</param>
+		/// <param name="Payload">The Xml data to attached to request, or null.</param>
+		/// <param name="Timeout">Timeout in millisec</param>
+		/// <returns>A valid HttpWebRequest, or null for any failure.</returns>
+		/// <remarks>This uses the POST method with default credentials and a content type of 'text/xml'.</remarks>
+		static public HttpWebRequest CreateWebRequest(string WebRequestString, string Payload, int Timeout = 100000)
+		{
+			HttpWebRequest Request = null;
+			try
+			{
+				Request = ( HttpWebRequest )WebRequest.Create( WebRequestString );
+				Request.Method = "POST";
+				Request.UseDefaultCredentials = true;
+				Request.SendChunked = false;
+				Request.Timeout = Timeout;
+
+				if( Payload != null )
+				{
+					byte[] Buffer = Encoding.UTF8.GetBytes( Payload );
+					Request.ContentType = "text/xml";
+					Request.ContentLength = Buffer.Length;
+					using (Stream RequestStream = Request.GetRequestStream())
+					{
+						RequestStream.Write(Buffer, 0, Buffer.Length);
+					}
+				}
+				else
+				{
+					Request.ContentLength = 0;
+				}
+			}
+			catch( Exception Ex )
+			{
+				Debug.WriteLine( "Exception in CreateRequest: " + Ex.Message );
+				Request = null;
+			}
+
+			return Request;
+		}
+
+		/// <summary>
+		/// Get a string response from an HttpWebRequest.
+		/// </summary>
+		/// <param name="Request">A valid HttpWebRequest.</param>
+		/// <returns>The string returned from the web request, or an empty string if there was an error.</returns>
+		/// <remarks>The response status code needs to be 'OK' for results to be processed. Also, this is designed for simple queries, 
+		/// so there is an arbitrary 1024 byte limit to the response size.</remarks>
+		static public string GetWebResponse( HttpWebRequest Request )
+		{
+			string ResponseString = "";
+
+			try
+			{
+				using (HttpWebResponse WebResponse = (HttpWebResponse)Request.GetResponse())
+				{
+					// Simple size check to prevent exploits
+					if (WebResponse.StatusCode == HttpStatusCode.OK && WebResponse.ContentLength < 4096)
+					{
+						// Process the response
+						byte[] RawResponse;
+						using (Stream ResponseStream = WebResponse.GetResponseStream())
+						{
+							RawResponse = new byte[WebResponse.ContentLength];
+							ResponseStream.Read(RawResponse, 0, (int)WebResponse.ContentLength);
+						}
+						ResponseString = Encoding.UTF8.GetString(RawResponse);
+					}
+				}
+			}
+			catch( Exception Ex )
+			{
+				Debug.WriteLine( "Exception in GetWebResponse: " + Ex.Message );
+				ResponseString = "";
+			}
+
+			return ResponseString;
+		}
+
+		/// <summary>
+		/// Send a web request to a web service and return a response string.
+		/// </summary>
+		/// <param name="WebRequestString">The URL of the request with parameters.</param>
+		/// <param name="Payload">The Xml payload of the request (can be null).</param>
+		/// <param name="Timeout">Timeout in millisec</param>
+		/// <returns>A string returned from the web request, or an empty string if there was any error.</returns>
+		/// <remarks>This calls CreateWebRequest and GetWebResponse internally.</remarks>
+		static public string GetWebServiceResponse( string WebRequestString, string Payload, int Timeout = 100000)
+		{
+			string ResponseString = "";
+
+			if( WebRequestString != null )
+			{
+				HttpWebRequest Request = CreateWebRequest(WebRequestString, Payload, Timeout);
+				if( Request != null )
+				{
+					ResponseString = GetWebResponse( Request );
+				}
+			}
+
+			return ResponseString;
+		}
+	}
+}
diff --git a/Engine/Source/Programs/CrashReporter/CrashReportProcess/CrashReportProcess.csproj b/Engine/Source/Programs/CrashReporter/CrashReportProcess/CrashReportProcess.csproj
--- a/Engine/Source/Programs/CrashReporter/CrashReportProcess/CrashReportProcess.csproj
+++ b/Engine/Source/Programs/CrashReporter/CrashReportProcess/CrashReportProcess.csproj
@@ -148,10 +148,6 @@
     </None>
   </ItemGroup>
   <ItemGroup>
-    <ProjectReference Include="..\..\DotNETCommon\DotNETUtilities\DotNETUtilities.csproj">
-      <Project>{5d7d66e8-8c76-4af9-b3ec-2ef03421d730}</Project>
-      <Name>DotNETUtilities</Name>
-    </ProjectReference>
     <ProjectReference Include="..\CrashReportCommon\CrashReportCommon.csproj">
       <Project>{b83de8c9-9429-4765-86bd-8eeea89b138f}</Project>
       <Name>CrashReportCommon</Name>
diff --git a/Engine/Source/Programs/CrashReporter/CrashReportProcess/ReportProcessor.cs b/Engine/Source/Programs/CrashReporter/CrashReportProcess/ReportProcessor.cs
--- a/Engine/Source/Programs/CrashReporter/CrashReportProcess/ReportProcessor.cs
+++ b/Engine/Source/Programs/CrashReporter/CrashReportProcess/ReportProcessor.cs
@@ -10,9 +10,6 @@ using System.Threading;
 using System.Threading.Tasks;
 using Amazon.S3.Model;
 using Tools.CrashReporter.CrashReportCommon;
-using Tools.DotNETCommon.LaunchProcess;
-using Tools.DotNETCommon.SimpleWebRequest;
-using Tools.DotNETCommon.XmlHandler;
 
 namespace Tools.CrashReporter.CrashReportProcess
 {
diff --git a/Engine/Source/Programs/CrashReporter/CrashReportProcess/ReportQueue.cs b/Engine/Source/Programs/CrashReporter/CrashReportProcess/ReportQueue.cs
--- a/Engine/Source/Programs/CrashReporter/CrashReportProcess/ReportQueue.cs
+++ b/Engine/Source/Programs/CrashReporter/CrashReportProcess/ReportQueue.cs
@@ -6,7 +6,6 @@ using System.Globalization;
 using System.IO;
 using System.Linq;
 using Tools.CrashReporter.CrashReportCommon;
-using Tools.DotNETCommon.XmlHandler;
 
 namespace Tools.CrashReporter.CrashReportProcess
 {
diff --git a/Engine/Source/Programs/CrashReporter/CrashReportProcess/Symbolicator.cs b/Engine/Source/Programs/CrashReporter/CrashReportProcess/Symbolicator.cs
--- a/Engine/Source/Programs/CrashReporter/CrashReportProcess/Symbolicator.cs
+++ b/Engine/Source/Programs/CrashReporter/CrashReportProcess/Symbolicator.cs
@@ -7,7 +7,6 @@ using System.IO;
 using System.Linq;
 using System.Threading.Tasks;
 using Tools.CrashReporter.CrashReportCommon;
-using Tools.DotNETCommon.LaunchProcess;
 
 namespace Tools.CrashReporter.CrashReportProcess
 {
diff --git a/Engine/Source/Programs/CrashReporter/CrashReportReceiver/CrashReportReceiver.csproj b/Engine/Source/Programs/CrashReporter/CrashReportReceiver/CrashReportReceiver.csproj
--- a/Engine/Source/Programs/CrashReporter/CrashReportReceiver/CrashReportReceiver.csproj
+++ b/Engine/Source/Programs/CrashReporter/CrashReportReceiver/CrashReportReceiver.csproj
@@ -110,10 +110,6 @@
     </None>
   </ItemGroup>
   <ItemGroup>
-    <ProjectReference Include="..\..\DotNETCommon\DotNETUtilities\DotNETUtilities.csproj">
-      <Project>{5d7d66e8-8c76-4af9-b3ec-2ef03421d730}</Project>
-      <Name>DotNETUtilities</Name>
-    </ProjectReference>
     <ProjectReference Include="..\CrashReportCommon\CrashReportCommon.csproj">
       <Project>{b83de8c9-9429-4765-86bd-8eeea89b138f}</Project>
       <Name>CrashReportCommon</Name>
diff --git a/Engine/Source/Programs/CrashReporter/CrashReportReceiver/WebHandler.cs b/Engine/Source/Programs/CrashReporter/CrashReportReceiver/WebHandler.cs
--- a/Engine/Source/Programs/CrashReporter/CrashReportReceiver/WebHandler.cs
+++ b/Engine/Source/Programs/CrashReporter/CrashReportReceiver/WebHandler.cs
@@ -11,7 +11,6 @@ using System.Web;
 using System.Text;
 
 using Tools.CrashReporter.CrashReportCommon;
-using Tools.DotNETCommon.XmlHandler;
 using System.ComponentModel;
 using System.Threading;
 
